Final "On Your Own!" pages.

diff --git a/Engine/Documentation/Source/Programming/Tutorials/PlayerInput/3/PlayerInput_3.INT.udn b/Engine/Documentation/Source/Programming/Tutorials/PlayerInput/3/PlayerInput_3.INT.udn
--- a/Engine/Documentation/Source/Programming/Tutorials/PlayerInput/3/PlayerInput_3.INT.udn
+++ b/Engine/Documentation/Source/Programming/Tutorials/PlayerInput/3/PlayerInput_3.INT.udn
@@ -14,6 +14,7 @@ Description:Binding input to C++ code
 		[Programming Tutorials Home](Programming\Tutorials\)
 	[/PARAM]
 	[PARAM:next]
+		[Next Step](Programming\Tutorials\PlayerInput\4 "%Programming\Tutorials\PlayerInput\4:title%")
 	[/PARAM]
 [/OBJECT]
 [/VAR]
diff --git a/Engine/Documentation/Source/Programming/Tutorials/PlayerInput/4/PlayerInput_4.INT.udn b/Engine/Documentation/Source/Programming/Tutorials/PlayerInput/4/PlayerInput_4.INT.udn
new file mode 100644
--- /dev/null
+++ b/Engine/Documentation/Source/Programming/Tutorials/PlayerInput/4/PlayerInput_4.INT.udn
@@ -0,0 +1,34 @@
+Availability: Public
+Title:4. On Your Own!
+Crumbs:%ROOT%, Programming, Programming/Tutorials, Programming/Tutorials/PlayerInput
+Description:Binding input to C++ code
+
+[VAR:Steps]
+[OBJECT:Navigation]
+	[PARAM:previous]
+		[Previous Step](Programming\Tutorials\PlayerInput\3 "%Programming\Tutorials\PlayerInput\3:title%")
+	[/PARAM]
+	[PARAM:current]
+	[/PARAM]
+	[PARAM:home]
+		[Programming Tutorials Home](Programming\Tutorials\)
+	[/PARAM]
+	[PARAM:next]
+	[/PARAM]
+[/OBJECT]
+[/VAR]
+
+%Steps%
+
+Using what you have learned, try to do the following:
+
+* Implement directional controls that increase speed after being held for a certain period of time.
+* Create a special input sequence that instantly expands the object to full scale when the user presses an **Action Mapping** immediately after starting to press an **Axis Mapping**.
+
+As for the specifics covered in this tutorial:
+
+* For more information about **Input**, try the [%Gameplay/Input:title%](Gameplay/Input) page.
+* For further tutorials, see the [%Programming/Tutorials:title%](Programming/Tutorials) page.
+
+
+%Steps%
diff --git a/Engine/Documentation/Source/Programming/Tutorials/PlayerInput/PlayerInput.INT.udn b/Engine/Documentation/Source/Programming/Tutorials/PlayerInput/PlayerInput.INT.udn
--- a/Engine/Documentation/Source/Programming/Tutorials/PlayerInput/PlayerInput.INT.udn
+++ b/Engine/Documentation/Source/Programming/Tutorials/PlayerInput/PlayerInput.INT.udn
@@ -64,6 +64,7 @@ Over the course of this tutorial, you will learn to react to player input by ext
 	* [](Programming\Tutorials\PlayerInput\1 "%Programming\Tutorials\PlayerInput\1:description%")
 	* [](Programming\Tutorials\PlayerInput\2 "%Programming\Tutorials\PlayerInput\2:description%")
 	* [](Programming\Tutorials\PlayerInput\3 "%Programming\Tutorials\PlayerInput\3:description%")
+	* [](Programming\Tutorials\PlayerInput\4 "%Programming\Tutorials\PlayerInput\4:description%")
 	[/PARAM]
 [/OBJECT]
 
diff --git a/Engine/Documentation/Source/Programming/Tutorials/VariablesTimersEvents/3/VariablesTimersEvents_3.INT.udn b/Engine/Documentation/Source/Programming/Tutorials/VariablesTimersEvents/3/VariablesTimersEvents_3.INT.udn
--- a/Engine/Documentation/Source/Programming/Tutorials/VariablesTimersEvents/3/VariablesTimersEvents_3.INT.udn
+++ b/Engine/Documentation/Source/Programming/Tutorials/VariablesTimersEvents/3/VariablesTimersEvents_3.INT.udn
@@ -14,6 +14,7 @@ Description:Using Blueprints to supplement or replace C++ functionality.
 		[Programming Tutorials Home](Programming\Tutorials\)
 	[/PARAM]
 	[PARAM:next]
+		[Next Step](Programming\Tutorials\VariablesTimersEvents\4 "%Programming\Tutorials\VariablesTimersEvents\4:title%")
 	[/PARAM]
 [/OBJECT]
 [/VAR]
diff --git a/Engine/Documentation/Source/Programming/Tutorials/VariablesTimersEvents/4/VariablesTimersEvents_4.INT.udn b/Engine/Documentation/Source/Programming/Tutorials/VariablesTimersEvents/4/VariablesTimersEvents_4.INT.udn
new file mode 100644
--- /dev/null
+++ b/Engine/Documentation/Source/Programming/Tutorials/VariablesTimersEvents/4/VariablesTimersEvents_4.INT.udn
@@ -0,0 +1,36 @@
+Availability: Public
+Title:4. On Your Own!
+Crumbs:%ROOT%, Programming, Programming/Tutorials, Programming/Tutorials/VariablesTimersEvents
+Description:Binding input to C++ code
+
+[VAR:Steps]
+[OBJECT:Navigation]
+	[PARAM:previous]
+		[Previous Step](Programming\Tutorials\VariablesTimersEvents\3 "%Programming\Tutorials\VariablesTimersEvents\3:title%")
+	[/PARAM]
+	[PARAM:current]
+	[/PARAM]
+	[PARAM:home]
+		[Programming Tutorials Home](Programming\Tutorials\)
+	[/PARAM]
+	[PARAM:next]
+	[/PARAM]
+[/OBJECT]
+[/VAR]
+
+%Steps%
+
+Using what you have learned, try to do the following:
+
+* Create an **Actor** that moves or rotates to a target transform when an **Event** is run. This could be used as a moving platform or door in a game. Make the **Event** start a **Timer** that will trigger a second **Event** which moves the **Actor** back to its original location. Use exposed variables (via **UPROPERTY**) instead of hard-coded values wherever appropriate.
+* Make a lit torch that burns out (perhaps by deactivating a fiery **Particle System Component**) by using a **Timer** handle and a few custom **Events**. For example, an _AddFuel_ **Event** could extend the burning time of the torch. A _DouseWithWater_ **Event** could shut it off immediately and prevent _AddFuel_ from working in the future. Both of these features can be written without using a **Tick**, simply by modifying a running **Timer** through its handle.
+
+As for the specifics covered in this tutorial:
+
+* For more information about **Timers**, try the [%Programming/UnrealArchitecture/Timers:title%](Programming/UnrealArchitecture/Timers) page.
+* For a complete reference using the **UPROPERTY** tag with variables in your classes or structs, look at the [%Programming/UnrealArchitecture/Reference/Properties:title%](Programming/UnrealArchitecture/Reference/Properties) page.
+* To learn more about **UFUNCTIONS** and **Event** creation, check the [%Programming/UnrealArchitecture/Reference/Functions:title%](Programming/UnrealArchitecture/Reference/Functions) page.
+* For further tutorials, see the [%Programming/Tutorials:title%](Programming/Tutorials) page.
+
+
+%Steps%
diff --git a/Engine/Documentation/Source/Programming/Tutorials/VariablesTimersEvents/VariablesTimersEvents.INT.udn b/Engine/Documentation/Source/Programming/Tutorials/VariablesTimersEvents/VariablesTimersEvents.INT.udn
--- a/Engine/Documentation/Source/Programming/Tutorials/VariablesTimersEvents/VariablesTimersEvents.INT.udn
+++ b/Engine/Documentation/Source/Programming/Tutorials/VariablesTimersEvents/VariablesTimersEvents.INT.udn
@@ -64,6 +64,7 @@ This tutorial will show you how to expose variables and functions to the editor,
 	* [](Programming\Tutorials\VariablesTimersEvents\1 "%Programming\Tutorials\VariablesTimersEvents\1:description%")
 	* [](Programming\Tutorials\VariablesTimersEvents\2 "%Programming\Tutorials\VariablesTimersEvents\2:description%")
 	* [](Programming\Tutorials\VariablesTimersEvents\3 "%Programming\Tutorials\VariablesTimersEvents\3:description%")
+	* [](Programming\Tutorials\VariablesTimersEvents\4 "%Programming\Tutorials\VariablesTimersEvents\4:description%")
 	[/PARAM]
 [/OBJECT]
 
