#UE4doc #docs  Made minor edits to the doc

diff --git a/Engine/Documentation/Source/Resources/Showcases/MatineeFightScene/MatineeFightScene.INT.udn b/Engine/Documentation/Source/Resources/Showcases/MatineeFightScene/MatineeFightScene.INT.udn
--- a/Engine/Documentation/Source/Resources/Showcases/MatineeFightScene/MatineeFightScene.INT.udn
+++ b/Engine/Documentation/Source/Resources/Showcases/MatineeFightScene/MatineeFightScene.INT.udn
@@ -279,7 +279,7 @@ When a keyframe is added to the **Movement** track, you can specify the desired
 
 ![](CharacterVisibilityTrack.png)
 
-The **Visibility** tracks assigned to Troopers 1, 2, and 3 are simply there to toggle when they character should be visible in the scene. The time the character is visible is denoted by the green bar shown in the image above. You can also look at the keyframes themselves as keyframes pointed downward indicate the character is hidden while keyframes pointed upward indicate the character is set to be shown. 
+The **Visibility** tracks assigned to Troopers 1, 2, and 3 are simply there to toggle when the character should be visible in the scene. The time the character is visible is denoted by the green bar shown in the image above. You can also look at the keyframes themselves as keyframes pointed downward indicate the character is hidden while keyframes pointed upward indicate the character is set to be shown. 
 
 [REGION:tip]
 You can **Right-Click** on a keyframe and "flip" the toggle to the opposite of its current setting by selecting the **Flip Toggle** option from the context menu. 
