#ue4 docs: Changed "Artificial Intelligence" to "AI & Behavior Trees"

diff --git a/Engine/Documentation/Source/Gameplay/AI/AI.INT.udn b/Engine/Documentation/Source/Gameplay/AI/AI.INT.udn
--- a/Engine/Documentation/Source/Gameplay/AI/AI.INT.udn
+++ b/Engine/Documentation/Source/Gameplay/AI/AI.INT.udn
@@ -1,5 +1,5 @@
 Availability:Public
-Title: Artificial Intelligence
+Title: AI & Behavior Trees
 Crumbs:%ROOT%, Engine, Gameplay
 Description:Information over Artificial Intelligence including the use of Behavior Trees and Navmesh.
 
