#ue4 docs: Removed "Using" from "Skeletal Mesh Actors"

diff --git a/Engine/Documentation/Source/Engine/Actors/SkeletalMeshActors/SkeletalMeshActors.INT.udn b/Engine/Documentation/Source/Engine/Actors/SkeletalMeshActors/SkeletalMeshActors.INT.udn
--- a/Engine/Documentation/Source/Engine/Actors/SkeletalMeshActors/SkeletalMeshActors.INT.udn
+++ b/Engine/Documentation/Source/Engine/Actors/SkeletalMeshActors/SkeletalMeshActors.INT.udn
@@ -1,5 +1,5 @@
 Availability: Public
-Title:Using Skeletal Mesh Actors
+Title: Skeletal Mesh Actors
 Crumbs: %ROOT%, Engine, Engine/Actors
 Description:Creating and using Skeletal Mesh Actors
 Related: Engine/Actors
