#eu4 docs: updated Skeletal Mesh Actors link

diff --git a/Engine/Documentation/Source/Engine/Actors/Actors.INT.udn b/Engine/Documentation/Source/Engine/Actors/Actors.INT.udn
--- a/Engine/Documentation/Source/Engine/Actors/Actors.INT.udn
+++ b/Engine/Documentation/Source/Engine/Actors/Actors.INT.udn
@@ -108,7 +108,7 @@ properties will be modified to cause them to look or behave appropriately.
 	[/PARAM]
 	[PARAM:links]	
 	* [Static Mesh Actors](Engine/Actors/StaticMeshActor "%Engine/Actors/StaticMeshActor:description%")
-	* [Skeletal Mesh Actors](Engine/Actors/SkeletalMeshActors "%Engine/Actors/SkeletalMeshActors:description%")
+	* [](Engine/Actors/SkeletalMeshActors "%Engine/Actors/SkeletalMeshActors:description%")
 	* [Geometry Brush Actors](Engine/Actors/Brushes "%Engine/Actors/Brushes:description%")
 	* [](Engine/Actors/Volumes "%Engine/Actors/Volumes:description%")
 	* [](Engine/Actors/Triggers "%Engine/Actors/Triggers:description%")
