#loc UE4DocKOR

diff --git a/Engine/Documentation/Source/Engine/AI/BehaviorTrees/QuickStart/6/6.INT.udn b/Engine/Documentation/Source/Engine/AI/BehaviorTrees/QuickStart/6/6.INT.udn
--- a/Engine/Documentation/Source/Engine/AI/BehaviorTrees/QuickStart/6/6.INT.udn
+++ b/Engine/Documentation/Source/Engine/AI/BehaviorTrees/QuickStart/6/6.INT.udn
@@ -1,7 +1,7 @@
+﻿Availability:Public
 Title: 6. Setup the AI Controller Blueprint
 Crumbs:%ROOT%, Engine, Gameplay, Engine/AI/BehaviorTrees, Engine/AI/BehaviorTrees/QuickStart
 Description:Step 6. Setup the AI Controller Blueprint
-Availability:Public
 
 [VAR:Steps]
 [OBJECT:Navigation]
diff --git a/Engine/Documentation/Source/Gameplay/HowTo/CharacterMovement/Blueprints/Setup_1/Setup_1.INT.udn b/Engine/Documentation/Source/Gameplay/HowTo/CharacterMovement/Blueprints/Setup_1/Setup_1.INT.udn
--- a/Engine/Documentation/Source/Gameplay/HowTo/CharacterMovement/Blueprints/Setup_1/Setup_1.INT.udn
+++ b/Engine/Documentation/Source/Gameplay/HowTo/CharacterMovement/Blueprints/Setup_1/Setup_1.INT.udn
@@ -1,4 +1,5 @@
-﻿Availability: Public
+﻿INTSourceChangelist:2474065
+Availability: Public
 Title:1. Character Setup
 Crumbs: %ROOT%, Gameplay, Gameplay\HowTo\CharacterMovement, Gameplay\HowTo\CharacterMovement\Blueprints
 Description:Here we will set up our character and create some assets that we will need later in this guide. 
