#UE4doc #docs: Deleted "Placing Actors" as a related page to "Placing Actors"

diff --git a/Engine/Documentation/Source/Engine/Actors/Placement/PlacingActors.INT.udn b/Engine/Documentation/Source/Engine/Actors/Placement/PlacingActors.INT.udn
--- a/Engine/Documentation/Source/Engine/Actors/Placement/PlacingActors.INT.udn
+++ b/Engine/Documentation/Source/Engine/Actors/Placement/PlacingActors.INT.udn
@@ -2,7 +2,6 @@ Availability:Public
 Title:Placing Actors
 Crumbs: %ROOT%, Engine, Engine/Actors
 Description:Shows how you can place Actors such as props, lights, cameras, etc. in your levels. 
-Related: Engine/Actors/Placement
 Related: Engine/UI/ClassViewer
 Related: GettingStarted\HowTo
 
