#UE4doc #docs: Made Physic Asset singular

diff --git a/Engine/Documentation/Source/Engine/Actors/SkeletalMeshActors/SkeletalMeshActors.INT.udn b/Engine/Documentation/Source/Engine/Actors/SkeletalMeshActors/SkeletalMeshActors.INT.udn
--- a/Engine/Documentation/Source/Engine/Actors/SkeletalMeshActors/SkeletalMeshActors.INT.udn
+++ b/Engine/Documentation/Source/Engine/Actors/SkeletalMeshActors/SkeletalMeshActors.INT.udn
@@ -86,7 +86,7 @@ Materials that are to be used with Skeletal Mesh Actors need to have the materia
 [/region]
 
 ## Collision
-Normal collision creation and detection does not work with Skeletal Mesh Actors. In order for your Skeletal Meshes to collide with objects in the level, your Skeletal Mesh Actors need to have a [Physics Assets]( Engine/Physics/PhAT/UserGuide) that is specially created for your object. 
+Normal collision creation and detection does not work with Skeletal Mesh Actors. In order for your Skeletal Meshes to collide with objects in the level, your Skeletal Mesh Actors need to have a [Physics Asset]( Engine/Physics/PhAT/UserGuide) that is specially created for your object. 
 
 ### Creating Collision for a Skeletal Mesh Actor 
 
