#UE4doc #docs: Changed link overwrite text to match page title, "Level Designer Quick Start"

diff --git a/Engine/Documentation/Source/Engine/UMG/QuickStart/UMGQuickStart.INT.udn b/Engine/Documentation/Source/Engine/UMG/QuickStart/UMGQuickStart.INT.udn
--- a/Engine/Documentation/Source/Engine/UMG/QuickStart/UMGQuickStart.INT.udn
+++ b/Engine/Documentation/Source/Engine/UMG/QuickStart/UMGQuickStart.INT.udn
@@ -44,7 +44,7 @@ Description:Getting started with using Unreal Motion Graphics in Unreal Engine 4
 
 In this **Quick Start Guide**, you will learn how to implement some basic HUD elements using **Unreal Motion Graphics UI Designer (UMG)**. You will learn how to implement Health and Energy Bars, an Experience Bar, hiding/display text based on conditions, and an Action Bar style button that when pressed, activates a special character ability. 
 
-Over the course of this guide, links to additional documentation will be provided which will go deeper into each section's subject should you have additional questions. While this guide will cover the actions needed to proceed through each step, if you have not worked with the Unreal Editor before, it is recommended that you first look over the **[Unreal Editor Quick Start](Engine/QuickStart/)** guide to gain a better understanding of general editor usage, terms, and processes. 
+Over the course of this guide, links to additional documentation will be provided which will go deeper into each section's subject should you have additional questions. While this guide will cover the actions needed to proceed through each step, if you have not worked with the Unreal Editor before, it is recommended that you first look over the **[Level Designer Quick Start](Engine/QuickStart/)** guide to gain a better understanding of general editor usage, terms, and processes. 
 
 
 [OBJECT:TopicButtonList]
