#UE4doc #docs  Made minor edits to the doc

diff --git a/Engine/Documentation/Source/Engine/Blueprints/BP_HowTo/CollapsingGraphs/CollapsingGraphs.INT.udn b/Engine/Documentation/Source/Engine/Blueprints/BP_HowTo/CollapsingGraphs/CollapsingGraphs.INT.udn
--- a/Engine/Documentation/Source/Engine/Blueprints/BP_HowTo/CollapsingGraphs/CollapsingGraphs.INT.udn
+++ b/Engine/Documentation/Source/Engine/Blueprints/BP_HowTo/CollapsingGraphs/CollapsingGraphs.INT.udn
@@ -1,7 +1,7 @@
 Availability: Public
 Title:Collapsing Graphs
 Crumbs: %ROOT%, Engine, Engine\Blueprints, Engine\Blueprints\BP_HowTo
-Description:Collapsing nodes to a new Graph, a Function or a Macro. 
+Description:Collapsing nodes to a new Graph, a Function, or a Macro. 
 Related: Engine/Blueprints/UserGuide
 Related: Engine/Blueprints/UserGuide/Functions
 Related: Engine/Blueprints/UserGuide/Macros
