#UE4doc #docs  Made minor edits to the doc

diff --git a/Engine/Documentation/Source/Engine/Blueprints/BP_HowTo/Debugging/Debugging.INT.udn b/Engine/Documentation/Source/Engine/Blueprints/BP_HowTo/Debugging/Debugging.INT.udn
--- a/Engine/Documentation/Source/Engine/Blueprints/BP_HowTo/Debugging/Debugging.INT.udn
+++ b/Engine/Documentation/Source/Engine/Blueprints/BP_HowTo/Debugging/Debugging.INT.udn
@@ -12,7 +12,7 @@ Related: Engine/Blueprints
 
 When working with Blueprints, there may be instances where something is not working the way you had planned it to. To diagnose these problems, the Blueprint Visual Scripting Tool comes with its own debugger and one of the most powerful aspects of it is the ability to add **Breakpoints** to nodes in your graph. By adding Breakpoints to nodes in your graph, you can play (or simulate) the game normally in the editor, then when reaching the node you have added a Breakpoint to, the game will pause and jump to that node in your graph so that you can step through your script to see where issues are occurring. 
 
-Another useful debugging feature is **Watch Values**, which allows you mark a variable in your Blueprint (or multiple variables) to **Watch** which will allow you to see their values inside your Blueprint while the game is running. This is essential for debugging incorrect calculations or values for your variables as you can visually see where the error is coming from as your variables change in real time. 
+Another useful debugging feature is **Watch Values**, which allows you to mark a variable in your Blueprint (or multiple variables) to **Watch** which will allow you to see their values inside your Blueprint while the game is running. This is essential for debugging incorrect calculations or values for your variables as you can visually see where the error is coming from as your variables change in real time. 
 
 On this page, you will learn how to set up a Blueprint for debugging as well as other methods used for diagnosing problems.
 
