#UE4doc #docs  Made minor edits to the doc

diff --git a/Engine/Documentation/Source/Engine/Blueprints/BP_HowTo/MakingFunctions/MakingFunctions.INT.udn b/Engine/Documentation/Source/Engine/Blueprints/BP_HowTo/MakingFunctions/MakingFunctions.INT.udn
--- a/Engine/Documentation/Source/Engine/Blueprints/BP_HowTo/MakingFunctions/MakingFunctions.INT.udn
+++ b/Engine/Documentation/Source/Engine/Blueprints/BP_HowTo/MakingFunctions/MakingFunctions.INT.udn
@@ -78,7 +78,7 @@ The steps below will guide you through creating a Function which will print text
 	When **F** is pressed, it will call the Print Text function which is set to print text to the screen using a Print String. 
 
 	[REGION:note]
-	Don't see your function? Try clicking the **Compile** button from the Toolbar then try your search again. 
+	If you do not see your function, try clicking the **Compile** button from the Toolbar then try your search again. 
 	[/REGION]
 
 1.	Your network should look similar to below. 
