TypoCorrection

diff --git a/Engine/Documentation/Source/Gameplay/Networking/Actors/Properties/Properties.INT.udn b/Engine/Documentation/Source/Gameplay/Networking/Actors/Properties/Properties.INT.udn
--- a/Engine/Documentation/Source/Gameplay/Networking/Actors/Properties/Properties.INT.udn
+++ b/Engine/Documentation/Source/Gameplay/Networking/Actors/Properties/Properties.INT.udn
@@ -1,4 +1,4 @@
-Availability:Public
+ï»¿Availability:Public
 Title:Property Replication
 Crumbs:%ROOT%, Gameplay, Gameplay/Networking, Gameplay/Networking/Actors
 Description:Detailed information about how Actor properties are replicated.
@@ -19,7 +19,7 @@ Actor property replication is reliable. This means that the property of the clie
 ## Setting up properties for replication
 
 To replicate a property you need to do a few things:
-In the header of the actor class where the property is defined, you need to make sure you have the ‘replicated' keyword as one of the parameters to the UPROPERTY declaration:
+In the header of the actor class where the property is defined, you need to make sure you have the `replicated` keyword as one of the parameters to the UPROPERTY declaration:
 
 	class ENGINE_API AActor : public UObject
 	{
@@ -41,4 +41,4 @@ In the actor's constructor, make sure you have the bReplicates flag set to true:
 		bReplicates = true;
 	}
 
-That's about it.The member variable ‘Owner' will now be synchronized to all connected clients for every copy of this actor type that is currently instantiated (in this case, the base actor class).
+That's about it.The member variable `Owner` will now be synchronized to all connected clients for every copy of this actor type that is currently instantiated (in this case, the base actor class).
