TypoCorrection

diff --git a/Engine/Documentation/Source/Gameplay/Networking/Actors/ReplicationFlow/ReplicationFlow.INT.udn b/Engine/Documentation/Source/Gameplay/Networking/Actors/ReplicationFlow/ReplicationFlow.INT.udn
--- a/Engine/Documentation/Source/Gameplay/Networking/Actors/ReplicationFlow/ReplicationFlow.INT.udn
+++ b/Engine/Documentation/Source/Gameplay/Networking/Actors/ReplicationFlow/ReplicationFlow.INT.udn
@@ -1,4 +1,4 @@
-Availability:Public
+﻿Availability:Public
 Title: Detailed Actor Replication Flow
 Crumbs:%ROOT%, Gameplay, Gameplay/Networking, Gameplay/Networking/Actors
 Description:A detailed description of low-level Actor replication.
@@ -50,7 +50,7 @@ The high level flow looks like this:
 ## Replicating an Actor to a Connection
 
 `UChannel::ReplicateActor` is the workhorse for replicating an actor and all of its components to a connection. The flow looks something like this:
-* Determine if this is this is the first update since this actor channel was opened
+* Determine if this is the first update since this actor channel was opened
 	* If so, serialize specific information that is needed (initial location, rotation, etc)
 * Determine if this connection owns this actor
 	* If not owned, and this actor's role is `ROLE_AutonomousProxy`, then downgrade to `ROLE_SimulatedProxy`
