TypoCorrection

diff --git a/Engine/Documentation/Source/Gameplay/Networking/Actors/Roles/Roles.INT.udn b/Engine/Documentation/Source/Gameplay/Networking/Actors/Roles/Roles.INT.udn
--- a/Engine/Documentation/Source/Gameplay/Networking/Actors/Roles/Roles.INT.udn
+++ b/Engine/Documentation/Source/Gameplay/Networking/Actors/Roles/Roles.INT.udn
@@ -1,4 +1,4 @@
-Availability:Public
+﻿Availability:Public
 Title: Actor Role and RemoteRole
 Crumbs:%ROOT%, Gameplay, Gameplay/Networking, Gameplay/Networking/Actors
 Description:The various aspects of replicating Actor objects.
@@ -47,7 +47,7 @@ There are currently two types of simulation that occur.
 ### `ROLE_SimulatedProxy`
 This is the standard simulation path, and is generally based on extrapolating movement based on the last known velocity. When the server sends an update for a particular actor, the client will adjust its position towards the new location, and then in between updates, the client will continue to move the actor based on the most recent velocity sent from the server. 
 
-Simulating using last known velocity it just one example of how general simulation works. There is nothing stopping you from writing custom code to use some other information to extrapolate between server updates.
+Simulating using last known velocity is just one example of how general simulation works. There is nothing stopping you from writing custom code to use some other information to extrapolate between server updates.
 
 ### `ROLE_AutonomousProxy`
 This is generally only used on actors that are possessed by PlayerControllers. This just means that this actor is receiving inputs from a human controller, so when we extrapolate, we have a bit more information, and can use actual human inputs to fill in the missing info (rather than extrapolating based on the last known velocity).
