#UE4doc #docs  Made minor edits to the doc

diff --git a/Engine/Documentation/Source/Engine/Blueprints/BP_HowTo/WorkingWithArrays/WorkingWithArrays.INT.udn b/Engine/Documentation/Source/Engine/Blueprints/BP_HowTo/WorkingWithArrays/WorkingWithArrays.INT.udn
--- a/Engine/Documentation/Source/Engine/Blueprints/BP_HowTo/WorkingWithArrays/WorkingWithArrays.INT.udn
+++ b/Engine/Documentation/Source/Engine/Blueprints/BP_HowTo/WorkingWithArrays/WorkingWithArrays.INT.udn
@@ -167,7 +167,7 @@ _For this example, we are using the **Blueprint 3rd Person Template** project wi
 
 	![](HT25b.png)
 
-	We have two **Get Actor Transform** options above, due to the fact that we still have Actors selected in the level. The node in the red box will get the transform for all the selected Actors in the level, which is not what we want here. You can hop back into the World Outliner and deselect your Actors, or select the singular usage outlined in the yellow box above. 
+	We have two **Get Actor Transform** options above, due to the fact that we still have Actors selected in the level. The node in the red box will get the transform for all the selected Actors in the level, which is not what we want here. You can hop back into the **World Outliner** and deselect your Actors, or select the singular usage outlined in the yellow box above. 
 
 1.	Connect the **Return Value** from the **Get Actor Transform** to the **Spawn Transform** pin on the **Spawn Actor** node. 
 
