#UE4doc #docs  Made minor edits to the doc

diff --git a/Engine/Documentation/Source/Engine/Blueprints/Editor/UIComponents/FindResults/FindResults.INT.udn b/Engine/Documentation/Source/Engine/Blueprints/Editor/UIComponents/FindResults/FindResults.INT.udn
--- a/Engine/Documentation/Source/Engine/Blueprints/Editor/UIComponents/FindResults/FindResults.INT.udn
+++ b/Engine/Documentation/Source/Engine/Blueprints/Editor/UIComponents/FindResults/FindResults.INT.udn
@@ -14,7 +14,7 @@ The **Find Results** panel is a powerful search tool within the Bluerprint Edito
 * Property name
 * Property value
 
-As the **Find Results** panel tracks down search matches, it will display a list of results, each of which works like a hyperlink that will jump the graph view to the resulting node. This makes it an excellent way to track down a specific node or piece of information that may be buried somehwere within a complex network of Blueprint scripting nodes.
+As the **Find Results** panel tracks down search matches, it will display a list of results, each of which works like a hyperlink that will jump the graph view to the resulting node. This makes it an excellent way to track down a specific node or piece of information that may be buried somewhere within a complex network of Blueprint scripting nodes.
 
 Much like searching in a web browser, the **Find Results** panel is available by pressing **Ctrl-F** while working in the Blueprint Editor. By default, the panel will appear along the bottom of the **Graph** Panel. If the **Compiler Results** panel is showing, then the **Find Results** panel will dock alongside it.
 
