#UE4 Docs: fixes for publishing

diff --git a/Engine/Documentation/Source/Resources/Showcases/BlueprintMultiplayer/BlueprintMultiplayer.INT.udn b/Engine/Documentation/Source/Resources/Showcases/BlueprintMultiplayer/BlueprintMultiplayer.INT.udn
--- a/Engine/Documentation/Source/Resources/Showcases/BlueprintMultiplayer/BlueprintMultiplayer.INT.udn
+++ b/Engine/Documentation/Source/Resources/Showcases/BlueprintMultiplayer/BlueprintMultiplayer.INT.udn
@@ -6,7 +6,7 @@ Related: Engine/Blueprints
 Related: Engine/UMG
 Related: Programming/Online
 Related: Programming/Online/Steam
-Related: Gameplay/Networking/Replication
+Related: Gameplay/Networking/Actors
 Related: Resources/ContentExamples/Networking
 
 
