#UE4 Docs: fixes for publishing

diff --git a/Engine/Documentation/Source/Engine/AI/BehaviorTrees/UserGuide/BehaviorTreeUserGuide.INT.udn b/Engine/Documentation/Source/Engine/AI/BehaviorTrees/UserGuide/BehaviorTreeUserGuide.INT.udn
--- a/Engine/Documentation/Source/Engine/AI/BehaviorTrees/UserGuide/BehaviorTreeUserGuide.INT.udn
+++ b/Engine/Documentation/Source/Engine/AI/BehaviorTrees/UserGuide/BehaviorTreeUserGuide.INT.udn
@@ -40,19 +40,19 @@ Or, if you have started to edit a Behavior Tree (or simply want a new Blackboard
 
 ### Creating a Task (Blueprint)
 
-1. Click the New Task ![](Engine/AI/BehaviorTrees/QuickStart/10/newTask_button.png) button in the Behavior Tree Editor (while in Behavior Tree Edit Mode).
+1. Click the New Task ![](Engine/AI/BehaviorTrees/QuickStart/12/newTask_button.png) button in the Behavior Tree Editor (while in Behavior Tree Edit Mode).
 1. Choose a Behavior Tree Task class to inherit from, `BTTask_BlueprintBase` is the default "blank" class.
 1. The new Task Blueprint will be created in the same folder as your Behavior Tree asset.
 
 ### Creating a Decorator (Blueprint)
 
-1. Click the New Decorator ![](Engine/AI/BehaviorTrees/QuickStart/11/newDecorator_button.png) button in the Behavior Tree Editor (while in Behavior Tree Edit Mode).
+1. Click the New Decorator ![](Engine/AI/BehaviorTrees/QuickStart/13/newDecorator_button.png) button in the Behavior Tree Editor (while in Behavior Tree Edit Mode).
 1. Choose a Behavior Tree Decorator class to inherit from, `BTDecorator_BlueprintBase` is the default "blank" class.
 1. The new Decorator Blueprint will be created in the same folder as your Behavior Tree asset.
 
 ### Creating a Service (Blueprint)
 
-1. Click the New Service ![](Engine/AI/BehaviorTrees/QuickStart/9/newService_button.png) button in the Behavior Tree Editor (while in Behavior Tree Edit Mode).
+1. Click the New Service ![](Engine/AI/BehaviorTrees/QuickStart/11/newService_button.png) button in the Behavior Tree Editor (while in Behavior Tree Edit Mode).
 1. Choose a Behavior Tree Service class to inherit from, `BTService_BlueprintBase` is the default "blank" class.
 1. The new Service Blueprint will be created in the same folder as your Behavior Tree asset.
 
@@ -101,7 +101,7 @@ You can edit a Behavior Tree by either **double-clicking** it in the **Content B
 
 Or by switching a currently opened Behavior Tree Editor to Behavior Tree Edit Mode:
 
-![](Engine/AI/BehaviorTrees/QuickStart/12/btMode.png)(convert:false)
+![](Engine/AI/BehaviorTrees/QuickStart/14/btMode.png)(convert:false)
 
 ### Setting the Blackboard Asset to Use
 
@@ -134,7 +134,7 @@ To connect nodes:
 1. Now drag the cursor to an input bar.
 1. Release the left mouse button.
 
-![](Engine/AI/BehaviorTrees/QuickStart/12/connectNodes.png)
+![](Engine/AI/BehaviorTrees/QuickStart/14/connectNodes.png)
 
 ### Breaking Connections
 
