#UE4doc #docs  Made minor edits to the doc

diff --git a/Engine/Documentation/Source/Engine/Blueprints/HowTo/BPHT_4/BPHT_4.INT.udn b/Engine/Documentation/Source/Engine/Blueprints/HowTo/BPHT_4/BPHT_4.INT.udn
--- a/Engine/Documentation/Source/Engine/Blueprints/HowTo/BPHT_4/BPHT_4.INT.udn
+++ b/Engine/Documentation/Source/Engine/Blueprints/HowTo/BPHT_4/BPHT_4.INT.udn
@@ -58,7 +58,7 @@ Description:In this step we are using the Level Blueprint to get the location of
 	This is the **Event Dispatcher** we created. When you select this, two additional nodes will be added. 
 	
 	[REGION:note]
-	Don't see this node in the list? Make sure that you clicked **Compile** inside the **FirstPersonCharacter** Blueprint and try again. 
+	If you do not see this node in the list, make sure that you clicked **Compile** inside the **FirstPersonCharacter** Blueprint and try again. 
 	[/REGION]
 
 1. Make sure you connect the newly added nodes from the **Bind Event**. 
