#UE4doc #docs  Made minor edits to the doc

diff --git a/Engine/Documentation/Source/Engine/Blueprints/QuickStart/2/BPQS_2.INT.udn b/Engine/Documentation/Source/Engine/Blueprints/QuickStart/2/BPQS_2.INT.udn
--- a/Engine/Documentation/Source/Engine/Blueprints/QuickStart/2/BPQS_2.INT.udn
+++ b/Engine/Documentation/Source/Engine/Blueprints/QuickStart/2/BPQS_2.INT.udn
@@ -24,13 +24,13 @@ Description:An overview on the most commonly used viewport navigation controls.
 
 ---
 
-You'll build the launchpad in the Level Editor, and then convert it to a Blueprint Class so you can add your gameplay behavior to it. 
+You will build the launchpad in the Level Editor, and then convert it to a Blueprint Class so you can add your gameplay behavior to it. 
 
 1. First, move around in the viewport until you are looking at the top platforms in the level.
 
 	![](TopPlatforms.png)(w:500)
 
-	We're going to create the container to hold all the parts of our launchpad by using the Empty Actor. The two parts (or Components) you will need are a shape to represent the launchpad, and a trigger for when the character
+	We are going to create the container to hold all the parts of our launchpad by using the Empty Actor. The two parts (or Components) you will need are a shape to represent the launchpad, and a trigger for when the character
 	overlaps it. 
 
 1. In the **Modes** panel, click on **Basic**, then find the **Empty Actor**.
@@ -39,16 +39,16 @@ You'll build the launchpad in the Level Editor, and then convert it to a Bluepri
 1. Drag it into the level so that it is sitting on one of the top platforms.
 
 	![](DragEmptyActor.png)
-1. Now that you have the Actor selected in the level, its properties are visible in the Details panel. At the top of the Details panel, you can rename it. Go ahead and click in the box to enter a new name, like LaunchPad.
+1. Now that you have the Actor selected in the level, its properties are visible in the **Details** panel. At the top of the **Details** panel, you can rename it. Go ahead and click in the box to enter a new name, like LaunchPad.
 
 	![](NameLaunchpad.png)
 
-	The first component we are going to add is a cube, although we'll scale it down a bit to make it more platform-sized. Adding this component first makes it the Root component, so all other components
-	will follow its lead for position, scale, and movement around the scene. You can rearrange components as you work, as well, so this isn't a permanent decision.
+	The first component we are going to add is a cube, although we will scale it down a bit to make it more platform-sized. Adding this component first makes it the Root component, so all other components
+	will follow its lead for position, scale, and movement around the scene. You can rearrange components as you work, as well, so this is not a permanent decision.
 1. In the **Details** panel, click the **Add Component** button, then select **Cube** under **Basic Shapes**.
 1. Now you should have the **Cube** component selected, so you can tweak its properties. Resize the cube to (1, 1, 0.1).
 
-1. Now, we'll add a Box Collision Component, which will fire off an event whenever something overlaps it. In the **Details** panel, click the **Add Component** button, then select **Cube Collision** under **Collision**.
+1. Now, we will add a Box Collision Component, which will fire off an event whenever something overlaps it. In the **Details** panel, click the **Add Component** button, then select **Cube Collision** under **Collision**.
 1. Resize the box to be the size of the platform in X and Y, and extend a bit above the platform in Z. If you want to enter in your **Scale** directly, try (1.25, 1.25, 15.0).
 
 [REGION:tip]
@@ -58,7 +58,7 @@ If at any time you need to change your Actor's properties, such as to scale up t
 
 ![](LaunchPadActor.png)
 
-Now that you have your Actor the way you want it, we'll turn it into a Blueprint. You can add more components inside the Blueprint Editor, and tweak them just like you can here in the level.
+Now that you have your Actor the way you want it, we will turn it into a Blueprint. You can add more components inside the Blueprint Editor, and tweak them just like you can here in the level.
 
 
 ---
