#UE4doc #docs  Made minor edits to the doc

diff --git a/Engine/Documentation/Source/Engine/Blueprints/QuickStart/3/QS_3.INT.udn b/Engine/Documentation/Source/Engine/Blueprints/QuickStart/3/QS_3.INT.udn
--- a/Engine/Documentation/Source/Engine/Blueprints/QuickStart/3/QS_3.INT.udn
+++ b/Engine/Documentation/Source/Engine/Blueprints/QuickStart/3/QS_3.INT.udn
@@ -23,7 +23,7 @@ Description:Creating an empty level via the New Level dialog window.
 %Steps%
 
 ---
-Making changes inside the Blueprint means that each time you make a new launchpad, it will have the look and feel you've created in the Blueprint Editor. While you could duplicate your LaunchPad
+Making changes inside the Blueprint means that each time you make a new launchpad, it will have the look and feel you have created in the Blueprint Editor. While you could duplicate your LaunchPad
 Actor around the level, any changes you make on a particular launchpad would only affect that one copy. 
 
 1. In the **Details** panel, click on the **Blueprint/Add Script** button.
@@ -36,8 +36,8 @@ Actor around the level, any changes you make on a particular launchpad would onl
 
 	![](CreateBlueprint.png)
 
-Your Blueprint is now visible in the **Content Browser**. Right now, you could drag and drop from the Content Browser into the level to create lots of copies of your platform mesh and trigger, but
-there won't be any behavior on it yet. In the next step, you'll start setting up the graph nodes inside of your Blueprint to launch your Character when it crosses the launchpad.
+Your Blueprint is now visible in the **Content Browser**. Right now, you could drag and drop from the **Content Browser** into the level to create lots of copies of your platform mesh and trigger, but
+there will not be any behavior on it yet. In the next step, you will start setting up the graph nodes inside of your Blueprint to launch your Character when it crosses the launchpad.
 
 ![](LaunchPadBP_CB.png)
 
