#UE4doc #docs  Made minor edits to the doc

diff --git a/Engine/Documentation/Source/Engine/Blueprints/QuickStart/4/BPQS_4.INT.udn b/Engine/Documentation/Source/Engine/Blueprints/QuickStart/4/BPQS_4.INT.udn
--- a/Engine/Documentation/Source/Engine/Blueprints/QuickStart/4/BPQS_4.INT.udn
+++ b/Engine/Documentation/Source/Engine/Blueprints/QuickStart/4/BPQS_4.INT.udn
@@ -23,9 +23,9 @@ Description:Placing and Manipulating Actors is the foundation of building any le
 
 ---
 
-To start adding behavior to your Blueprint Class, you'll need to open it in the Blueprint Editor.  
+To start adding behavior to your Blueprint Class, you will need to open it in the Blueprint Editor.  
 
-1. Double-click the **Blueprint Class** in the Content Browser.
+1. **Double-click** the **Blueprint Class** in the **Content Browser**.
 
 1. The Blueprint Editor will open, and you can see your **Cube** and **Box** components in a viewport. 
 
@@ -46,10 +46,10 @@ To start adding behavior to your Blueprint Class, you'll need to open it in the
 
 1. We want an event that will execute when anything overlaps our **Box** component. First, select the **Box** component in the **Components** tab.
 
-1. Right-click in an empty spot in your graph to bring up the context menu of nodes you could add to the graph.
+1. **Right-click** in an empty spot in your graph to bring up the context menu of nodes you could add to the graph.
 
 	[REGION:tip]
-	To move around in the graph, right-click and drag around. At this point, you could drag the graph to the left, moving the pre-placed event nodes off the left side of the screen and creating more 
+	To move around in the graph, **right-click** and drag around. At this point, you could drag the graph to the left, moving the pre-placed event nodes off the left side of the screen and creating more 
 	blank space to create your launchpad logic in.
 	[/REGION]
 
