#UE4doc #docs  Made minor edits to the doc

diff --git a/Engine/Documentation/Source/Engine/Blueprints/QuickStart/5/BPQS_5.INT.udn b/Engine/Documentation/Source/Engine/Blueprints/QuickStart/5/BPQS_5.INT.udn
--- a/Engine/Documentation/Source/Engine/Blueprints/QuickStart/5/BPQS_5.INT.udn
+++ b/Engine/Documentation/Source/Engine/Blueprints/QuickStart/5/BPQS_5.INT.udn
@@ -40,9 +40,9 @@ We will do that by working with the Other Actor output from the **OnComponentBeg
 	![](EqualsOtherActor.png)
 
 	We could set **Side Scroller Char** as the other input to the **Equals** node, but then if we changed the Character we were using, we would need to re-open this Blueprint and manually update that. 
-	Instead, let's get a reference to the current Pawn we are using.
+	Instead, lets get a reference to the current Pawn we are using.
 
-1. Right-click in an empty part of the graph to bring up the context menu.
+1. **Right-click** in an empty part of the graph to bring up the context menu.
 1. Type "Pawn" in the menu's search box, then select **Get Player Pawn** (under **Game**).
 
 	![](GetPlayerPawnMenu.png)
@@ -50,7 +50,7 @@ We will do that by working with the Other Actor output from the **OnComponentBeg
 
 	![](EqualsComplete.png)
 
-	Now that we've got a node that will tell us if the Other Actor is the Pawn our player controls, we will use the answer to change the execution flow of our graph. That is, we will direct the execution flow
+	Now that we have got a node that will tell us if the Other Actor is the Pawn our player controls, we will use the answer to change the execution flow of our graph. That is, we will direct the execution flow
 	as it leaves the On Component Begin Overlap node.
 
 	For that, we want to use a Flow Control node, specifically the **Branch** node. 
@@ -65,7 +65,7 @@ We will do that by working with the Other Actor output from the **OnComponentBeg
 
 	![](BranchComplete.png)
 
-The graph is now ready for you to set up different behavior to execute depending on whether the overlapping Actor is your Pawn or not. In the next step, we'll do just that, and set up Blueprint nodes 
+The graph is now ready for you to set up different behavior to execute depending on whether the overlapping Actor is your Pawn or not. In the next step, we will do just that, and set up Blueprint nodes 
 to launch our Character if the result of the **Equals** comparison is **True**.
 
 
