#UE4doc #docs  Made minor edits to the doc

diff --git a/Engine/Documentation/Source/Engine/Blueprints/QuickStart/6/BPQS_6.INT.udn b/Engine/Documentation/Source/Engine/Blueprints/QuickStart/6/BPQS_6.INT.udn
--- a/Engine/Documentation/Source/Engine/Blueprints/QuickStart/6/BPQS_6.INT.udn
+++ b/Engine/Documentation/Source/Engine/Blueprints/QuickStart/6/BPQS_6.INT.udn
@@ -55,7 +55,7 @@ and it will succeed. However, not everything in your level is the Pawn represent
 
 	At this point, **Compile** and **Save** the Blueprint using the toolbar buttons, then close the Blueprint Editor.
 
-1. Drag several of the launchpads into your level from the Content Browser. 
+1. Drag several of the launchpads into your level from the **Content Browser**. 
 
 	![](MultipleLaunchpads.png)
 
