TypoCorrection

diff --git a/Engine/Documentation/Source/Engine/AI/BehaviorTrees/UserGuide/BehaviorTreeUserGuide.INT.udn b/Engine/Documentation/Source/Engine/AI/BehaviorTrees/UserGuide/BehaviorTreeUserGuide.INT.udn
--- a/Engine/Documentation/Source/Engine/AI/BehaviorTrees/UserGuide/BehaviorTreeUserGuide.INT.udn
+++ b/Engine/Documentation/Source/Engine/AI/BehaviorTrees/UserGuide/BehaviorTreeUserGuide.INT.udn
@@ -1,7 +1,7 @@
+﻿Availability:Public
 Title: Behavior Tree User Guide
 Crumbs:%ROOT%, Engine, Gameplay, Engine/AI/BehaviorTrees
 Description:Behavior Tree Editor User Guide
-Availability:Public
 
 [TOC(start:2 end:2)]
 
