#UE4doc #docs  Made minor edits to the doc

diff --git a/Engine/Documentation/Source/Engine/Blueprints/UserGuide/Components/ComponentsList.INT.udn b/Engine/Documentation/Source/Engine/Blueprints/UserGuide/Components/ComponentsList.INT.udn
--- a/Engine/Documentation/Source/Engine/Blueprints/UserGuide/Components/ComponentsList.INT.udn
+++ b/Engine/Documentation/Source/Engine/Blueprints/UserGuide/Components/ComponentsList.INT.udn
@@ -107,21 +107,21 @@ You can also select the Component in the window and press the **Delete** key to
 
 ## Transforming Components
 
-Components, when added to an instance in your level, are placed by default at the location of that instance. However, they can be Transformed, Rotated and Scaled as necessary in either the **Details** panel or the **Viewport** similar to the method in which you can [Transform Actors](GettingStarted/HowTo/ManipulatingActors). 
+Components, when added to an instance in your level, are placed by default at the location of that instance. However, they can be Transformed, Rotated, and Scaled as necessary in either the **Details** panel or the **Viewport** similar to the method in which you can [Transform Actors](GettingStarted/HowTo/ManipulatingActors). 
 
-You can select Components for transformation either by clicking on their name in the **Components** window or by clicking directly on the Component in the **Viewport**. When Transforming, Rotating and Scaling your Components in the **Viewport**, hold **Shift** to enable snapping.  This snapping requires that snapping be enabled in the **Level Editor** and uses the **Snap Grid** settings from the **Level Editor** (see [](Engine/UI/LevelEditor/Viewports/ViewportToolbar) for more information on Grid Snapping).
+You can select Components for transformation either by clicking on their name in the **Components** window or by clicking directly on the Component in the **Viewport**. When Transforming, Rotating, and Scaling your Components in the **Viewport**, hold **Shift** to enable snapping.  This snapping requires that snapping be enabled in the **Level Editor** and uses the **Snap Grid** settings from the **Level Editor** (see [](Engine/UI/LevelEditor/Viewports/ViewportToolbar) for more information on Grid Snapping).
 
-Exact values can also be entered for **Location**, **Rotation** and **Scale** in the **Details** panel for your selected component.
+Exact values can also be entered for **Location**, **Rotation**, and **Scale** in the **Details** panel for your selected component.
 
 ![](transform_details.png)
 
 [REGION:note]
-Transforming, Rotating or Scaling a parent component will propagate the transformation down to all child components as well.
+Transforming, Rotating, or Scaling a parent component will propagate the transformation down to all child components as well.
 [/REGION]
 
 ## Component Assets
 
-Once you've added a Component, you may need to specify the asset to occupy the Component (such as assigning a Static Mesh to use for a _StaticMeshComponent_). There are a couple of different ways you can assign an asset to use for a Component type as described below. 
+Once you have added a Component, you may need to specify the asset to occupy the Component (such as assigning a Static Mesh to use for a _StaticMeshComponent_). There are a couple of different ways you can assign an asset to use for a Component type as described below. 
 
 
 ### Assigning Component Assets
@@ -134,7 +134,7 @@ To assign an asset to a Component in the **Components** window:
 	![](emptyComponent.png)
 	[/REGION]
 
-	Above we've added a _StaticMeshComponent_ and under **StaticMesh** is where we would assign the asset to use.
+	Above we have added a _StaticMeshComponent_ and under **StaticMesh** is where we would assign the asset to use.
 
 1. Click the **Current Asset** drop-down box, then select the asset to use from the context menu. 
 
@@ -154,7 +154,7 @@ Another method of assigning an asset can be done using the **Content Browser**:
 	![](emptyComponent.png)
 	[/REGION]
 
-	Above we've added a _StaticMeshComponent_ and under **StaticMesh** is where we would assign the asset to use. 
+	Above we have added a _StaticMeshComponent_ and under **StaticMesh** is where we would assign the asset to use. 
 
 1.	Instead of clicking the **Current Asset** box, since an asset is selected in the **Content Browser**, click the ![](k2_button_use_scs.png) button. 
 
diff --git a/Engine/Documentation/Source/Engine/Blueprints/UserGuide/KismetNodeReference.INT.udn b/Engine/Documentation/Source/Engine/Blueprints/UserGuide/KismetNodeReference.INT.udn
--- a/Engine/Documentation/Source/Engine/Blueprints/UserGuide/KismetNodeReference.INT.udn
+++ b/Engine/Documentation/Source/Engine/Blueprints/UserGuide/KismetNodeReference.INT.udn
@@ -8,7 +8,7 @@ Description: The User Guide is the go-to source to learn the different parts of
 
 [INCLUDE:Engine\Blueprints#intro]
 
-This page outlines the various types of Blueprints and Graph Types, Blueprint Elements (such as _Variables_, _Components_ and _Graphs_), Communication Methods and how to pass information between Blueprints, as well as a breakdown of the different types of Graph Nodes you will be working with as you create your Blueprints. 
+This page outlines the various types of Blueprints and Graph Types, Blueprint Elements (such as _Variables_, _Components_, and _Graphs_), Communication Methods and how to pass information between Blueprints, as well as a breakdown of the different types of Graph Nodes you will be working with as you create your Blueprints. 
 
 [OBJECT:TopicButtonList]
 	[PARAM:icon]
