#UE4doc #docs  Made minor edits to the doc

diff --git a/Engine/Documentation/Source/Engine/Blueprints/UserGuide/Graphs/Graphs.INT.udn b/Engine/Documentation/Source/Engine/Blueprints/UserGuide/Graphs/Graphs.INT.udn
--- a/Engine/Documentation/Source/Engine/Blueprints/UserGuide/Graphs/Graphs.INT.udn
+++ b/Engine/Documentation/Source/Engine/Blueprints/UserGuide/Graphs/Graphs.INT.udn
@@ -64,7 +64,7 @@ properties, but the key principles of adding variable references, adding and con
 	[PARAM:description]
 		**EventGraphs** are the most general type of _Blueprint_ graph.  Each new _Blueprint Class_ will contain one **EventGraph** when it is created, although more can be added.  These additional **EventGraphs** can
 help to organize your _Blueprint_ networks.  **EventGraphs** generally contain the networks for the gameplay behavior of the _Blueprint_, with [**Events**](Engine/Blueprints\UserGuide\Events), 
-[**CustomEvents**](Engine/Blueprints\UserGuide\Events\Custom) and **Input** nodes serving to initiate the flow of execution through the **EventGraph**.
+[**CustomEvents**](Engine/Blueprints\UserGuide\Events\Custom), and **Input** nodes serving to initiate the flow of execution through the **EventGraph**.
 	[/PARAM]
 	[PARAM:path]
 		[RELATIVE:Engine/Blueprints/UserGuide/EventGraph]
