#UE4doc #docs  Made minor edits to the doc

diff --git a/Engine/Documentation/Source/Engine/Blueprints/UserGuide/OnlineNodes/OnlineNodes.INT.udn b/Engine/Documentation/Source/Engine/Blueprints/UserGuide/OnlineNodes/OnlineNodes.INT.udn
--- a/Engine/Documentation/Source/Engine/Blueprints/UserGuide/OnlineNodes/OnlineNodes.INT.udn
+++ b/Engine/Documentation/Source/Engine/Blueprints/UserGuide/OnlineNodes/OnlineNodes.INT.udn
@@ -54,7 +54,7 @@ Related: Resources/ContentExamples/Networking
 
 ## Online Session Nodes
 
-The nodes below are responsible for hosting, discovering, joining or leaving multiplayer games through **Blueprints**. 
+The nodes below are responsible for hosting, discovering, joining, or leaving multiplayer games through **Blueprints**. 
 
 Each of these nodes are _latent_, meaning they perform their operation in the background (usually because these operations must communicate over the network) and will finish at some point in the future. This is also denoted on the nodes themselves in the upper right corner by a clock icon.
 
@@ -78,7 +78,7 @@ Use the **Find Sessions** node to get a list of game sessions that are currently
 
 ![](JoinSession.png)
 
-Once you know which **BlueprintSessionResult** you'd like to join, pass it into a **Join Session** call. If the game successfully connects to the server, it will automatically travel to the server's map and join gameplay. 
+Once you know which **BlueprintSessionResult** you would like to join, pass it into a **Join Session** call. If the game successfully connects to the server, it will automatically travel to the server's map and join gameplay. 
 
 ### Destroy Session
 
