#UE4doc #docs  Made minor edits to the doc

diff --git a/Engine/Documentation/Source/Engine/Blueprints/UserGuide/Timelines/Examples/OpeningDoors/OpeningDoors.INT.udn b/Engine/Documentation/Source/Engine/Blueprints/UserGuide/Timelines/Examples/OpeningDoors/OpeningDoors.INT.udn
--- a/Engine/Documentation/Source/Engine/Blueprints/UserGuide/Timelines/Examples/OpeningDoors/OpeningDoors.INT.udn
+++ b/Engine/Documentation/Source/Engine/Blueprints/UserGuide/Timelines/Examples/OpeningDoors/OpeningDoors.INT.udn
@@ -31,6 +31,6 @@ Note that sound effects are not considered in this setup.
 
 ## Timeline Setup
 
-The timeline setup for this example is fairly simple. We have a single Vector curve in place. In this curve we have animated the Z component to go from 0,0 to 1,360. This means that we are moving the door upward 360 units over the course of a single second.
+The timeline setup for this example is fairly simple. We have a single Vector curve in place. In this curve, we have animated the Z component to go from 0,0 to 1,360. This means that we are moving the door upward 360 units over the course of a single second.
 
 ![DoorTimeline.png](DoorTimeline.png)
\ No newline at end of file
