#UE4doc #docs  Made minor edits to the doc

diff --git a/Engine/Documentation/Source/Engine/Blueprints/UserGuide/Types/ClassBlueprint/ClassBlueprint.INT.udn b/Engine/Documentation/Source/Engine/Blueprints/UserGuide/Types/ClassBlueprint/ClassBlueprint.INT.udn
--- a/Engine/Documentation/Source/Engine/Blueprints/UserGuide/Types/ClassBlueprint/ClassBlueprint.INT.udn
+++ b/Engine/Documentation/Source/Engine/Blueprints/UserGuide/Types/ClassBlueprint/ClassBlueprint.INT.udn
@@ -63,13 +63,13 @@ Below are the most common Parent Classes used when creating a new Blueprint:
 | ------- | ---------- |
 | **Actor** | An Actor is an object that can be placed or spawned in the world. |
 | **Pawn** | A Pawn is an Actor that can be "possessed" and receieve input from a Controller. |
-| **Character** | A Character is a Pawn that includes the ability to walk, run, jump and more. |
+| **Character** | A Character is a Pawn that includes the ability to walk, run, jump, and more. |
 | **PlayerController** | A Player Controller is an Actor responsible for controlling a Pawn used by the player. |
 | **Game Mode** | A Game Mode defines the game being played, its rules, scoring, and other faces of the game type. |
 
 While above are the most common, all existing classes can be used as the Parent Class for a new Blueprint (even other Blueprint Classes). 
 
-For example, say you created an **Actor Blueprint** called _Animals_ and in it provided some script that all animals share such as _Hunger_, _Thirst_, _Energy_ or whatever script you wanted. Then you created another Blueprint called _Dogs_ and selected your _Animals_ Blueprint Class as the Parent Class; you can then provide specific functionality that applies to only dogs within the _Dogs_ Blueprint such as _Play Fetch_ or _Roll Over_ while inheriting the functionality that all animals share from the _Animals_ Blueprint. 
+For example, say you created an **Actor Blueprint** called _Animals_ and in it provided some script that all animals share such as _Hunger_, _Thirst_, _Energy_, or whatever script you wanted. Then you created another Blueprint called _Dogs_ and selected your _Animals_ Blueprint Class as the Parent Class; you can then provide specific functionality that applies to only dogs within the _Dogs_ Blueprint such as _Play Fetch_ or _Roll Over_ while inheriting the functionality that all animals share from the _Animals_ Blueprint. 
 
 [REGION:note]
 Depending upon the method used to [create a Blueprint Class](Engine\Blueprints\UserGuide\Types\ClassBlueprint\Creation), a Parent Class may already be assigned upon creation. 
