#UE4doc #docs  Made minor edits to the doc

diff --git a/Engine/Documentation/Source/Engine/Blueprints/UserGuide/Types/Interface/InterfaceBlueprint.INT.udn b/Engine/Documentation/Source/Engine/Blueprints/UserGuide/Types/Interface/InterfaceBlueprint.INT.udn
--- a/Engine/Documentation/Source/Engine/Blueprints/UserGuide/Types/Interface/InterfaceBlueprint.INT.udn
+++ b/Engine/Documentation/Source/Engine/Blueprints/UserGuide/Types/Interface/InterfaceBlueprint.INT.udn
@@ -87,7 +87,7 @@ When you first open up a new Interface, the editor looks like this:
 
 ![](InterfaceEditor.png)
 
-If you have just created your Interface, you'll see that the Editor has created a new blank Function for you named `NewFunction_0` and it will be highlighted for you to rename it.
+If you have just created your Interface, you will see that the Editor has created a new blank Function for you named `NewFunction_0` and it will be highlighted for you to rename it.
 
 ### Adding Functions
 
