#UE4doc #docs  Made minor edits to the doc

diff --git a/Engine/Documentation/Source/Engine/Blueprints/UserGuide/Types/Interface/UsingInterfaces/UsingInterfaces.INT.udn b/Engine/Documentation/Source/Engine/Blueprints/UserGuide/Types/Interface/UsingInterfaces/UsingInterfaces.INT.udn
--- a/Engine/Documentation/Source/Engine/Blueprints/UserGuide/Types/Interface/UsingInterfaces/UsingInterfaces.INT.udn
+++ b/Engine/Documentation/Source/Engine/Blueprints/UserGuide/Types/Interface/UsingInterfaces/UsingInterfaces.INT.udn
@@ -130,7 +130,7 @@ Note in this image we are sending another Blueprint as the Target input. This wo
 ![](InterfaceMessageFunctionCalled.png)
 
 [REGION:tip]
-	To specify a Blueprint that exists within the level, it is best to make a Public variable of type 'Actor', and assign the appropriate Blueprint in the value for that variable using the Details tab of the Level Editor.
+	To specify a Blueprint that exists within the level, it is best to make a Public variable of type 'Actor', and assign the appropriate Blueprint in the value for that variable using the **Details** tab of the Level Editor.
 [/REGION]
 
 ### Overriding an Interface Function
