#UE4doc #docs Updated Landscape image and wording of a few items throughout the page

diff --git a/Engine/Documentation/Source/Engine/Landscape/Creation/Creation.INT.udn b/Engine/Documentation/Source/Engine/Landscape/Creation/Creation.INT.udn
--- a/Engine/Documentation/Source/Engine/Landscape/Creation/Creation.INT.udn
+++ b/Engine/Documentation/Source/Engine/Landscape/Creation/Creation.INT.udn
@@ -28,8 +28,8 @@ Description:Guide to creating new Landscape terrains.
 
 ## Overview
 
-Creating massive terrain based worlds for your projects is something that Unreal Engine 4 (UE4) can do quite well with its built in terrain editing tools, Landscape. 
-Landscape will allow you to create massive terrain based worlds that can still maintain playable frame rates across a multitude of different devices. 
+Creating massive terrain based worlds for your projects is something that Unreal Engine 4 (UE4) can do quite well with its built in terrain editing tools. 
+The Landscape tool will allow you to create massive terrain based worlds that can still maintain playable frame rates across a multitude of different devices. 
 Because of the flexibility of the Landscape system, you can create your Landscape using either of the following methods. 
 
 * You can create a new Landscape heightmap from scratch using the Landscape tools in-engine.
@@ -61,7 +61,7 @@ In Landscape Manage Mode, you can create new Landscapes and modify existing Land
 ![](Landscape_New_Manage_Tab.png)
 
 [REGION:note]
-If your level already contains one or more Landscapes, the **Manage** tab will look different With the **Landscape Editor** section displaying a drop-down menu with the **Selection** tool visible. 
+If your level already contains one or more Landscapes, the **Manage** tab will look different. The **Landscape Editor** section will display a drop-down menu with the **Selection** tool visible. 
 It is from this drop down menu that you will select the Landscape you want to work with. 
 
 ![](Landscape_Multiple_Landscapes_Actors.png)
@@ -87,6 +87,7 @@ Creating a new Landscape from scratch is done through the **New Landscape** sect
 |**Section Per Component**|Sections Per Component helps with the Landscapes LOD.  Each section is the basic unit for the Landscape LOD. One component could have 2 x 2 sections which means that it is possible that one component could be rendering four different LOD’s at once. With a larger section size, you get the added benefit of reduced CPU calculation time. However you could run into issues with the Landscape rendering to many vertices at once. This is could be especially prevalent when using very large areas of Landscape.  These problems could be even worse on mobile devices because of the limited amount of draw calls you can have due to hardware limitations.|
 |**Number of Components**|Along with section size, this sets the size of your landscape. This value is capped at 32 x 32 because each component does have a CPU cost associated with it and going over this cap could result in performance issues with the Landscape.|
 |**Overall Resolution:**| This is the number of vertices your Landscape is using.|
+|**Total Components**:| The total number of components that will be created for your landscape.|
 |**Fill World**:| This will make your Landscape be as big as it can possibly be.|
 |**Create:**| This will create your Landscape in the world using the settings you specified.|
 
@@ -101,7 +102,7 @@ For this example, we are going to leave all of the Landscape settings at their d
 | **Overall Resolution:**| 505 x 505|
 
 [REGION:tip]
-Be careful when choosing **Section Size** because as the number of components increase, build times and performance can be dramatically effected. 63x63 quads per section is generally recommended as it gives good performance and size.
+Be careful when choosing **Section Size** because as the number of components increase, build times and performance can be dramatically affected. 63x63 quads per section is generally recommended as it gives good performance and size.
 [/REGION]
 
 
@@ -138,7 +139,7 @@ If the Landscape is large or complex, it may take a few seconds for it to appear
 
 Once you have pressed the **Create** Button, you should have something that looks similar to the following. 
 
-![Landscape_NewTerrain.png](Landscape_NewTerrain.png)
+![Landscape_NewTerrain.png](Landscape_NewTerrain.png)(w:600)
 
 Now with your new Landscape now created, you can now start to [Sculpt](Engine/Landscape/Editing/SculptMode) or [Paint](Engine/Landscape/Editing/PaintMode) the Landscape to your liking. 
 
