documented new ShaderComplexity fix for decals with image

diff --git a/Engine/Documentation/Source/Engine/Actors/DecalActor/DecalUserGuide.INT.udn b/Engine/Documentation/Source/Engine/Actors/DecalActor/DecalUserGuide.INT.udn
--- a/Engine/Documentation/Source/Engine/Actors/DecalActor/DecalUserGuide.INT.udn
+++ b/Engine/Documentation/Source/Engine/Actors/DecalActor/DecalUserGuide.INT.udn
@@ -103,6 +103,15 @@ The mesh complexity of the objects affected by the decal is not affecting the pe
 
 We can further improve the per decal performance. Ideally the bounding box of the decal is small to get better per pixel performance. This can be done manually. An automated method is possible but a good designer can also adjust the placement to improve performance further.
 
+
+The view mode **ShaderComplexity** (editor UI or "viewmode ShaderComplexity") can be used to see the impact on the pixel shading cost, it uses a pixel shader cost estimate and accumulates where muleuiple decals overlap.
+At the moment the decal masking feature has no effect there (masked parts should have a small constnt cost because of using the stencil hardware feature).
+The following shows a scene without (left) and with decals (right), in normal rendering (top) and with the ShaderComplexity enabled (bottom):
+
+![](DecalShaderComplexity.jpg)
+
+The darker color indicates a higher performance cost of those pixels. This information can be used to optimize the right shaders, remove barely visible decals or place them more efficiently.
+
 ## Current limitations
 
 * We currently only support deferred decals and they only work on static objects.
