TypoCorrection

diff --git a/Engine/Documentation/Source/Engine/Actors/DecalActor/DecalUserGuide.INT.udn b/Engine/Documentation/Source/Engine/Actors/DecalActor/DecalUserGuide.INT.udn
--- a/Engine/Documentation/Source/Engine/Actors/DecalActor/DecalUserGuide.INT.udn
+++ b/Engine/Documentation/Source/Engine/Actors/DecalActor/DecalUserGuide.INT.udn
@@ -1,4 +1,4 @@
-Availability:Public
+﻿Availability:Public
 Title:Decal Actor User Guide
 Crumbs:%ROOT%, Engine, Engine/Actors
 Description:A guide to using the Deferred Decal actor.
@@ -104,8 +104,8 @@ The mesh complexity of the objects affected by the decal is not affecting the pe
 We can further improve the per decal performance. Ideally the bounding box of the decal is small to get better per pixel performance. This can be done manually. An automated method is possible but a good designer can also adjust the placement to improve performance further.
 
 
-The view mode **ShaderComplexity** (editor UI or "viewmode ShaderComplexity") can be used to see the impact on the pixel shading cost, it uses a pixel shader cost estimate and accumulates where muleuiple decals overlap.
-At the moment the decal masking feature has no effect there (masked parts should have a small constnt cost because of using the stencil hardware feature).
+The view mode **ShaderComplexity** (editor UI or "viewmode ShaderComplexity") can be used to see the impact on the pixel shading cost, it uses a pixel shader cost estimate and accumulates where multiple decals overlap.
+At the moment the decal masking feature has no effect there (masked parts should have a small constant cost because of using the stencil hardware feature).
 The following shows a scene without (left) and with decals (right), in normal rendering (top) and with the ShaderComplexity enabled (bottom):
 
 ![](DecalShaderComplexity.jpg)
