#UE4doc #docs  Made minor edits to the doc

diff --git a/Engine/Documentation/Source/Engine/Blueprints/UserGuide/Variables/Variables.INT.udn b/Engine/Documentation/Source/Engine/Blueprints/UserGuide/Variables/Variables.INT.udn
--- a/Engine/Documentation/Source/Engine/Blueprints/UserGuide/Variables/Variables.INT.udn
+++ b/Engine/Documentation/Source/Engine/Blueprints/UserGuide/Variables/Variables.INT.udn
@@ -279,7 +279,7 @@ You can create a **Set** node (1 above) or **Get** node (2 above) for a variable
 
 Variable values can be set either as a part of the Blueprint node network or as a default value prior to execution. To set variable default values:
 
-1. Click the **Class Defaults** ![](classDefaults_button.png) Button on the Blueprint Editor Tool bar to open the Defaults in the Details Panel.
+1. Click the **Class Defaults** ![](classDefaults_button.png) Button on the Blueprint Editor Tool bar to open the Defaults in the **Details** panel.
 
 1. In the **Details** panel, enter the desired default values to the right of your variable name.
 
