#UE4doc #docs  Made minor edits to the doc

diff --git a/Engine/Documentation/Source/Engine/Components/Components.INT.udn b/Engine/Documentation/Source/Engine/Components/Components.INT.udn
--- a/Engine/Documentation/Source/Engine/Components/Components.INT.udn
+++ b/Engine/Documentation/Source/Engine/Components/Components.INT.udn
@@ -43,7 +43,7 @@ Description:A Component is a special type of Object designed to be used as a sub
 When adding Components to an Actor, you are putting together the bits and pieces that will make up that Actor as a whole who will be placed in your level(s). 
 Continuing with the car example above, the wheels on a car, the steering wheel, the body of the car, the lights etc. would all be considered as Components while the car itself the Actor. 
 
-Once you've added the Components that define your Actor, it is important to note that you can place the Actor in your level even without providing any **Blueprint** script (or C++ Code) as to how the car should function (as in the real world, a car can exist as an object without a driver providing the instructions to tell the car what to do). 
+Once you have added the Components that define your Actor, it is important to note that you can place the Actor in your level even without providing any **Blueprint** script (or C++ Code) as to how the car should function (as in the real world, a car can exist as an object without a driver providing the instructions to tell the car what to do). 
 Since the "gas pedal" would be a Component, through script or code you would indicate that the gas pedal has been pressed and in turn provide the logic for the car to speed up. 
 Each Component could then be accessed in this manner (turning on the lights, honking the horn, turning up the radio, steering, etc.) 
 
