#UE4doc #docs  Made minor edits to the doc

diff --git a/Engine/Documentation/Source/Engine/Components/Camera/CameraComponents.INT.udn b/Engine/Documentation/Source/Engine/Components/Camera/CameraComponents.INT.udn
--- a/Engine/Documentation/Source/Engine/Components/Camera/CameraComponents.INT.udn
+++ b/Engine/Documentation/Source/Engine/Components/Camera/CameraComponents.INT.udn
@@ -13,7 +13,7 @@ The **CameraComponent** (which adds a camera perspective to view from) and **Spr
 
 The **CameraComponent** allows you to add a camera perspective as a sub-object to an Actor. The CameraComponent will provide information about the camera's properties if the **ViewTarget** is a **CameraActor** or an Actor that contains a CameraComponent and has its `bFindCameraComponentWhenViewTarget` option set to _True_. 
 
-For example, say you wanted to [switch between multiple cameras](Gameplay/HowTo/UsingCameras/Blueprints) in your level during gameplay. By using `SetViewTargetWithBlend` and a CameraActor, you can switch between each of your cameras and use the properties that you define within the CameraActor for each camera (including field of view, angle or any Post Processing effects, etc.) 
+For example, say you wanted to [switch between multiple cameras](Gameplay/HowTo/UsingCameras/Blueprints) in your level during gameplay. By using `SetViewTargetWithBlend` and a CameraActor, you can switch between each of your cameras and use the properties that you define within the CameraActor for each camera (including field of view, angle, or any Post Processing effects, etc.) 
 
 A related property that can be set for any **Pawn** is `bTakeCameraControlWhenPossessed`, where the Pawn will automatically become the ViewTarget upon possession by a **PlayerController**. So for example if you had multiple Characters (which are a form of Pawn) that you wanted to switch between, and each of them had their own CameraComponents assigned to them providing the camera perspective to view from, you could use set `bTakeCameraControlWhenPossessed` to _True_ for each of them and whenever you switch between them, the CameraComponent for each will be used (where each could have its own properties and settings defined). 
 
