#UE4doc #docs  Made minor edits to the doc

diff --git a/Engine/Documentation/Source/Engine/Components/Rendering/RenderingComponents.INT.udn b/Engine/Documentation/Source/Engine/Components/Rendering/RenderingComponents.INT.udn
--- a/Engine/Documentation/Source/Engine/Components/Rendering/RenderingComponents.INT.udn
+++ b/Engine/Documentation/Source/Engine/Components/Rendering/RenderingComponents.INT.udn
@@ -106,7 +106,7 @@ The **SceneCapture2DComponent** is used to capture a "snapshot" of the scene fro
 
 ![](2dRender.png)(w:440)
 
-Above we've created a **Blueprint** with a **SceneCapture2DComponent** and assigned a **Render Target** texture which was then used as a **Material** that we applied to the piece of geometry in the level. For more information, see [Scene Capture 2D Content Example](Resources/ContentExamples/Reflections/1_7). 
+Above we have created a **Blueprint** with a **SceneCapture2DComponent** and assigned a **Render Target** texture which was then used as a **Material** that we applied to the piece of geometry in the level. For more information, see [Scene Capture 2D Content Example](Resources/ContentExamples/Reflections/1_7). 
 
 ## Scene Capture Cube Component
 
@@ -130,7 +130,7 @@ Above we can move the **Start Transform** and **End Transform** independently of
 
 ## Text Render Component
 
-A **TextRenderComponent** renders text in the world with the given font. Contains usual font related attributes such as Scale, Alignment, Color etc. You might use this Component to indicate to a player that an object is intractable in your levels. 
+A **TextRenderComponent** renders text in the world with the given font. Contains usual font related attributes such as Scale, Alignment, Color, etc. You might use this Component to indicate to a player that an object is intractable in your levels. 
 
 For example, suppose you had a chair in your level that players can sit down in when they approach it if they press a button. You could add a **TextRenderComponent** containing the instructions needed to execute the command of sitting (visibility set to off) along with a **BoxComponent** used as a trigger which will set visibility of the text to true when the player enters it (shown below). 
 
