#UE4doc #docs  Made minor edits to the doc

diff --git a/Engine/Documentation/Source/Engine/Components/Shapes/ShapesComponents.INT.udn b/Engine/Documentation/Source/Engine/Components/Shapes/ShapesComponents.INT.udn
--- a/Engine/Documentation/Source/Engine/Components/Shapes/ShapesComponents.INT.udn
+++ b/Engine/Documentation/Source/Engine/Components/Shapes/ShapesComponents.INT.udn
@@ -24,7 +24,7 @@ A **BoxComponent** is a box generally used for simple collision (can also be use
 
 ![](box_trigger.png)
 
-A **BoxComponent** was added around a fire particle effect, with its collision settings set to **Generate Overlap Events**. When something overlaps the box, an event will be fired where you execute script that applies damage to the actor that overlaps it. You could also set the [Collision Response](Engine/Physics/Collision) to **BlockAll** which would prevent all Actors from entering the box (if you wanted to keep things from entering the fire). 
+A **BoxComponent** was added around a fire particle effect, with its collision settings set to **Generate Overlap Events**. When something overlaps the box, an event will be fired where you execute script that applies damage to the Actor that overlaps it. You could also set the [Collision Response](Engine/Physics/Collision) to **BlockAll** which would prevent all Actors from entering the box (if you wanted to keep things from entering the fire). 
 
 
 ## Capsule Component
