#UE4doc #docs  Made minor edits to the doc

diff --git a/Engine/Documentation/Source/Engine/Components/StaticMesh/StaticMeshComponents.INT.udn b/Engine/Documentation/Source/Engine/Components/StaticMesh/StaticMeshComponents.INT.udn
--- a/Engine/Documentation/Source/Engine/Components/StaticMesh/StaticMeshComponents.INT.udn
+++ b/Engine/Documentation/Source/Engine/Components/StaticMesh/StaticMeshComponents.INT.udn
@@ -11,7 +11,7 @@ Static Mesh Components allow a piece of geometry to be used as a sub-object for
 
 ## Static Mesh Component
 
-The **StaticMeshComponent** is used to create an instance of a `UStaticMesh`. A **Static Mesh** is a peice of geometry that consists of a static set of polygons and are the basic unit used to create world geometry for levels in Unreal Engine 4. In addition to building levels, **Static Meshes** can be used for creating movers such as doors or lifts, rigid body physics objects, foliage and terrain decorations, procedurally created buildings, game objectives, and many more visual elements. 
+The **StaticMeshComponent** is used to create an instance of a `UStaticMesh`. A **Static Mesh** is a piece of geometry that consists of a static set of polygons and are the basic unit used to create world geometry for levels in Unreal Engine 4. In addition to building levels, **Static Meshes** can be used for creating movers such as doors or lifts, rigid body physics objects, foliage and terrain decorations, procedurally created buildings, game objectives, and many more visual elements. 
 
 Below, a **StaticMeshComponent** representing a light fixture along with a **SpotLightComponent** and **PointLightComponent** are used to create an alarm that could cycle between two different types of lights.
 
