#UE4doc #docs  Made minor edits to the doc

diff --git a/Engine/Documentation/Source/Engine/Components/Widget/WidgetComponents.INT.udn b/Engine/Documentation/Source/Engine/Components/Widget/WidgetComponents.INT.udn
--- a/Engine/Documentation/Source/Engine/Components/Widget/WidgetComponents.INT.udn
+++ b/Engine/Documentation/Source/Engine/Components/Widget/WidgetComponents.INT.udn
@@ -13,7 +13,7 @@ Widget Components allow you to manifest 3D UI elements created through [Unreal M
 
 The **WidgetComponent** itself is a 3D instance of a **Widget Blueprint** that you can interact with in your game world. 
 
-In the example below, we've added a **WidgetComponent** which will now allow us to specify a **Widget Class Blueprint**. 
+In the example below, we have added a **WidgetComponent** which will now allow us to specify a **Widget Class Blueprint**. 
 
 ![](Engine/UMG/HowTo/Create3DWidgets/UMG_3D_21.png)
 
