TypoCorrection

diff --git a/Engine/Documentation/Source/Engine/Paper2D/Sprites/RenderGeometry/SpriteRenderGeometry.INT.udn b/Engine/Documentation/Source/Engine/Paper2D/Sprites/RenderGeometry/SpriteRenderGeometry.INT.udn
--- a/Engine/Documentation/Source/Engine/Paper2D/Sprites/RenderGeometry/SpriteRenderGeometry.INT.udn
+++ b/Engine/Documentation/Source/Engine/Paper2D/Sprites/RenderGeometry/SpriteRenderGeometry.INT.udn
@@ -1,7 +1,7 @@
+ï»¿Availability:Public
 Title:Sprite Source Region and Render Geometry
 Crumbs: %ROOT%, Engine, Engine/Paper2D, Engine/Paper2D/Sprites
 Description:Displays and allows editing of the texture Source Region and Sprite Render Geometry. 
-Availability:Public
 Related: Engine\Paper2D\Sprites\Collision
 Related: Engine\Paper2D\Flipbooks
 Related: Engine/Rendering/Materials
@@ -13,7 +13,7 @@ Related: Engine/Content/Types/Textures
 
 ## Overview
 
-Each sprite is roughly defined as a rectangle from a texture, which can be edited interactively using the ‘Edit Source Region’ mode in the sprite editor, or typed in manually in the details panel (Source Dimension and Source UV, both in pixels). The actual rendered geometry can be further adjusted in the **Edit RenderGeom** mode, and the collision shape can be adjusted in the [Edit Collision](Engine\Paper2D\Sprites\Collision) mode.
+Each sprite is roughly defined as a rectangle from a texture, which can be edited interactively using the ï¿½Edit Source Regionï¿½ mode in the sprite editor, or typed in manually in the details panel (Source Dimension and Source UV, both in pixels). The actual rendered geometry can be further adjusted in the **Edit RenderGeom** mode, and the collision shape can be adjusted in the [Edit Collision](Engine\Paper2D\Sprites\Collision) mode.
 
 **Edit RenderGeom** can also be useful to optimize transparent or masked sprites, trading off some extra triangles in exchange for reduced overdraw in transparent regions of the texture. The settings and tools are exactly the same as editing custom collision geometry. Be sure to use the **Snap to Pixel Grid** tool after editing custom geometry (this will happen automatically in a future release, but is a manual process for now).
 
