#UE4doc #docs: Removed extra link to Blueprint Editor Cheat Sheet under related topics

diff --git a/Engine/Documentation/Source/Engine/Blueprints/BP_HowTo/PlacingNodes/PlacingNodes.INT.udn b/Engine/Documentation/Source/Engine/Blueprints/BP_HowTo/PlacingNodes/PlacingNodes.INT.udn
--- a/Engine/Documentation/Source/Engine/Blueprints/BP_HowTo/PlacingNodes/PlacingNodes.INT.udn
+++ b/Engine/Documentation/Source/Engine/Blueprints/BP_HowTo/PlacingNodes/PlacingNodes.INT.udn
@@ -5,7 +5,7 @@ Description:Illustrates the methods in which you can add nodes to Blueprints.
 Related: Engine/Blueprints/BP_HowTo/ConnectingNodes
 Related: Engine/Blueprints/BP_HowTo/BasicUsage
 Related: Engine/Blueprints/UserGuide
-Related: Engine/Blueprints/UserGuide/CheatSheet
+
 
 
 On this page, you will learn how you can place nodes on your graphs in Blueprints. 
