#UE4doc #doc minor Updates to page and added in link to additional doc page.

diff --git a/Engine/Documentation/Source/Engine/Landscape/Landscape.INT.udn b/Engine/Documentation/Source/Engine/Landscape/Landscape.INT.udn
--- a/Engine/Documentation/Source/Engine/Landscape/Landscape.INT.udn
+++ b/Engine/Documentation/Source/Engine/Landscape/Landscape.INT.udn
@@ -125,7 +125,7 @@ Description:Landscape system for creating terrain for large, open, outdoor envir
 ## Overview
 
 [EXCERPT:Intro]
-The Landscape system enables you to create terrain for your world - mountains, valleys, uneven or sloped ground, even caves - and easily modify both its shape and its appearance by using a range of tools.
+The Landscape system enables you to create terrain for your world - mountains, valleys, uneven or sloped ground, even openings for caves - and easily modify both its shape and its appearance by using a range of tools.
 [/EXCERPT:Intro]
 
 For information about opening and using the Landscape tool, see [](Engine/Landscape\Creation).
@@ -200,13 +200,13 @@ This makes it possible to always use the highest resolution of the terrain for p
 
 ![Landscape_HighRes_Lighting.jpg](Landscape_HighRes_Lighting.jpg)(w:900 h:272)
 
-When this high-resoluton normal data is combined with detail normal maps, Landscape terrains can achieve extremely detailed lighting with very little overhead.
+When this high-resolution normal data is combined with detail normal maps, Landscape terrains can achieve extremely detailed lighting with very little overhead.
 
 ![Landscape_Normals_Comparison.jpg](Landscape_Normals_Comparison.jpg)(w:900 h:272)
 
 ### PhysX Collision
 
-Landscape uses a PhysX heightfield object for its collision for both Unreal and rigid-body collision. Physical Materials can be specified for each layer, and the collision system will use the dominant layer at each position to determine which Physical Material to use. It is possible to use a reduced resolution collision heightfield (e.g. 0.5x render resolution) to save on memory requirements for large Landscape terrains. The collision and render components for distant Landscapes can also be streamed out using the level streaming system.
+Landscape uses a PhysX heightfield object for its collision for both Unreal and rigid-body collision. [Physical Materials](https://docs.unrealengine.com/latest/INT/Engine/Physics/PhysicalMaterials/index.html) can be specified for each layer, and the collision system will use the dominant layer at each position to determine which Physical Material to use. It is possible to use a reduced resolution collision heightfield (e.g. 0.5x render resolution) to save on memory requirements for large Landscape terrains. The collision and render components for distant Landscapes can also be streamed out using the level streaming system.
 
  
 
