#UE4doc #doc Updates to page. minor fixes

diff --git a/Engine/Documentation/Source/Engine/Landscape/TechnicalGuide/LansdscapeTechnicalGuide.INT.udn b/Engine/Documentation/Source/Engine/Landscape/TechnicalGuide/LansdscapeTechnicalGuide.INT.udn
--- a/Engine/Documentation/Source/Engine/Landscape/TechnicalGuide/LansdscapeTechnicalGuide.INT.udn
+++ b/Engine/Documentation/Source/Engine/Landscape/TechnicalGuide/LansdscapeTechnicalGuide.INT.udn
@@ -28,7 +28,7 @@ Description:Technical Settings for Landscape.
 
 ## Overview
 
-To get the best performance possible form the Landscape system, there are certain technical restrictions that have to be adhered to. The following document aims to bring these restrictions to your attention as well as provide you with valuable information so that you can strike the best balance between beauty and performance with your Landscape. 
+To get the best performance possible from the Landscape system, there are certain technical restrictions that have to be adhered to. The following document aims to bring these restrictions to your attention as well as provide you with valuable information so that you can strike the best balance between beauty and performance with your Landscape. 
 
 
 ## Technical Details
