#UE4doc #docs: Made small text edit

diff --git a/Engine/Documentation/Source/Engine/Blueprints/Editor/UIComponents/Toolbar/KismetToolbar.INT.udn b/Engine/Documentation/Source/Engine/Blueprints/Editor/UIComponents/Toolbar/KismetToolbar.INT.udn
--- a/Engine/Documentation/Source/Engine/Blueprints/Editor/UIComponents/Toolbar/KismetToolbar.INT.udn
+++ b/Engine/Documentation/Source/Engine/Blueprints/Editor/UIComponents/Toolbar/KismetToolbar.INT.udn
@@ -21,6 +21,6 @@ The Toolbar contains two sections:
 
 ## Toolbar Options
 
-The Blueprint Toolbar contains many different options.  Different options are available depending on the type of Blueprint and, for Blueprint Classes, the mode the Blueprint is in.
+The Blueprint Toolbar contains many different options.  Different options are available depending on the type of Blueprint and the mode the Blueprint is in.
 
 For more information about Toolbar options, see [](Engine/Blueprints/Editor/UIComponents/Toolbar/ToolbarOptions).
