#UE4doc #docs: Changed "Let us blow stuff up!" to "Let's blow stuff up!"

diff --git a/Engine/Documentation/Source/Engine/Blueprints/HowTo/BPHT_5/BPHT_5.INT.udn b/Engine/Documentation/Source/Engine/Blueprints/HowTo/BPHT_5/BPHT_5.INT.udn
--- a/Engine/Documentation/Source/Engine/Blueprints/HowTo/BPHT_5/BPHT_5.INT.udn
+++ b/Engine/Documentation/Source/Engine/Blueprints/HowTo/BPHT_5/BPHT_5.INT.udn
@@ -24,7 +24,7 @@ Description:In this section we will use a Bind Event which is bound to OnHit to
 
 In this step, we will allow our rock to be destroyed when it is hit by the projectile the player fires. To do this, we will use another **Bind Event**, but we are not going to bind it to an Event Dispatcher. Instead we will use a collision event called **OnHit** which fires off when the collision of our rock is hit by something which is perfect for our situation. 
 
-Let us blow stuff up!
+Let's blow stuff up!
 
 1.	Back inside the **FirstPersonCharacter** Blueprint, find the **Spawn Actor** node where we spawn our rock. 
 
