#UE4doc #docs: Changed "let us finish it off" to "let's finish it off" and "That is it" to "That's it"

diff --git a/Engine/Documentation/Source/Engine/Blueprints/HowTo/BPHT_6/BPHT_6.INT.udn b/Engine/Documentation/Source/Engine/Blueprints/HowTo/BPHT_6/BPHT_6.INT.udn
--- a/Engine/Documentation/Source/Engine/Blueprints/HowTo/BPHT_6/BPHT_6.INT.udn
+++ b/Engine/Documentation/Source/Engine/Blueprints/HowTo/BPHT_6/BPHT_6.INT.udn
@@ -24,7 +24,7 @@ Description:In this step we will update our HUD based on values stored in our Ch
 
 In this step, we will update the values displayed on our HUD Widget Blueprint by communicating with our Character Blueprint to get values stored in it for Pulls, Shots, and Score. To do this, we will use a **Cast** node to access the player's Character Blueprint then Bind properties in our HUD to values in our Character Blueprint so they automatically get updated when those values are updated in the Character Blueprint. 
 
-Our game is almost complete, let us finish it off by updating the HUD. 
+Our game is almost complete, let's finish it off by updating the HUD. 
 
 1.	Open the **HUD** Widget Blueprint you created. 
 
@@ -79,7 +79,7 @@ Our game is almost complete, let us finish it off by updating the HUD.
 	![](GameplayGIF.gif)(convert:false)
 
 
-That is it; your simple shooting gallery game is complete. Over the course of this guide, we have learned how to communicate between three different types of Blueprints using an **Event Dispatcher**, **Cast** nodes, and **Binding**. We have learned how to pass information from the level to a Character Blueprint as well as how to communicate variables from the character to a HUD. 
+That's it; your simple shooting gallery game is complete. Over the course of this guide, we have learned how to communicate between three different types of Blueprints using an **Event Dispatcher**, **Cast** nodes, and **Binding**. We have learned how to pass information from the level to a Character Blueprint as well as how to communicate variables from the character to a HUD. 
 
 In the final step, you will see some examples you can take to alter this project and make it your own. You will also be provided with links to additional documentation that covers topics in this guide. 
 
