#UE4doc #docs: changed "graphs that references" to "graphs that reference"

diff --git a/Engine/Documentation/Source/Engine/Blueprints/UserGuide/Types/MacroLibrary/MacroLibrary.INT.udn b/Engine/Documentation/Source/Engine/Blueprints/UserGuide/Types/MacroLibrary/MacroLibrary.INT.udn
--- a/Engine/Documentation/Source/Engine/Blueprints/UserGuide/Types/MacroLibrary/MacroLibrary.INT.udn
+++ b/Engine/Documentation/Source/Engine/Blueprints/UserGuide/Types/MacroLibrary/MacroLibrary.INT.udn
@@ -50,9 +50,9 @@ A **Blueprint Macro Library** is a container that holds a collection of _Macros_
 be placed as nodes in other _Blueprints_. These can be time-savers as they can store commonly used sequences of nodes 
 complete with inputs and outputs for both execution and data transfer.
 
-_Macros_ are shared among all graphs that reference them, but they are auto-expanded into the graph as if they were a 
-collapsed node during compiling. This means that _Blueprint Macro Libraries_ do not need to be compiled, but changes 
-to a _Macro_ are only reflected in graphs that references that _Macro_ when the _Blueprint_ containing those graphs is 
+_Macros_ are shared among all graphs that reference them, but they are auto-expanded into graphs as if they were a 
+collapsed node during compiling. This means that _Blueprint Macro Libraries_ do not need to be compiled. However, changes 
+to a _Macro_ are only reflected in graphs that reference that _Macro_ when the _Blueprint_ containing those graphs is 
 recompiled.
 [/EXCERPT:Intro]
 
