Fixed some missing code at the bottom of a C++ Tutorial page.

diff --git a/Engine/Documentation/Source/Programming/Tutorials/PlayerInput/3/PlayerInput_3.INT.udn b/Engine/Documentation/Source/Programming/Tutorials/PlayerInput/3/PlayerInput_3.INT.udn
--- a/Engine/Documentation/Source/Programming/Tutorials/PlayerInput/3/PlayerInput_3.INT.udn
+++ b/Engine/Documentation/Source/Programming/Tutorials/PlayerInput/3/PlayerInput_3.INT.udn
@@ -227,6 +227,13 @@ Description:Binding input to C++ code
 			{
 				Super::SetupPlayerInputComponent(InputComponent);
 
+				// Respond when our "Grow" key is pressed or released.
+				InputComponent->BindAction("Grow", IE_Pressed, this, &AMyPawn::StartGrowing);
+				InputComponent->BindAction("Grow", IE_Released, this, &AMyPawn::StopGrowing);
+
+				// Respond every frame to the values of our two movement axes, "MoveX" and "MoveY".
+				InputComponent->BindAxis("MoveX", this, &AMyPawn::Move_XAxis);
+				InputComponent->BindAxis("MoveY", this, &AMyPawn::Move_YAxis);
 			}
 
 			void AMyPawn::Move_XAxis(float AxisValue)
