#UE4doc #docs Updated pages with requests from user helping with doc feedback forum

diff --git a/Engine/Documentation/Source/Engine/Landscape/Editing/PaintMode/PaintMode.INT.udn b/Engine/Documentation/Source/Engine/Landscape/Editing/PaintMode/PaintMode.INT.udn
--- a/Engine/Documentation/Source/Engine/Landscape/Editing/PaintMode/PaintMode.INT.udn
+++ b/Engine/Documentation/Source/Engine/Landscape/Editing/PaintMode/PaintMode.INT.udn
@@ -150,6 +150,7 @@ You can either create a layer info object from the layer itself, or reuse an exi
 
 1. Press the plus icon (![Create Layer Info](Engine/Landscape\Editing\icon_Landscape_CreateLayerInfo.png)) to the right of the Layer name.
 1. Choose **Weight-Blended Layer (normal)** or **Non Weight-Blended Layer**.
+![](WeightBlended_NonWeightBlended.png)
 1. Choose the location to save the layer info object.
 
 <!-- what's the difference between weight-blended and non weight-blended? -->
@@ -178,7 +179,7 @@ The primary purpose of layer info objects is to act as a unique key for painted
 
 | Option | Description |
 | ------ | ----------- |
-| **Phys Material** | The Physical Material (if any) assigned to areas of Landscape where this layer is dominant. |
+| **Phys Material** | The [Physical Material](https://docs.unrealengine.com/latest/INT/Engine/Physics/PhysicalMaterials/index.html) (if any) assigned to areas of Landscape where this layer is dominant. |
 | **Hardness** | The value used by the [Erosion](Engine/Landscape/Editing/SculptMode/#Erosion) tool. |
 
 ### Orphaned Layers
