#UE4doc #docs Minor updates and link addition from user requests within the Docs feedback forum.

diff --git a/Engine/Documentation/Source/Engine/Landscape/Editing/Editing.INT.udn b/Engine/Documentation/Source/Engine/Landscape/Editing/Editing.INT.udn
--- a/Engine/Documentation/Source/Engine/Landscape/Editing/Editing.INT.udn
+++ b/Engine/Documentation/Source/Engine/Landscape/Editing/Editing.INT.udn
@@ -37,13 +37,13 @@ After you have created a Landscape, you can now do any of the following to it:
 * Copy, paste, import, and export parts of the Landscape's heightmap with **[Landscape Gizmos](Engine/Landscape\Editing\Gizmos)**.
 * Modify the Landscape's **[Material](Engine/Landscape\Materials)**.
 * Create **[Splines](Engine/Landscape\Editing\Splines)** for any linear features you want to use in your Landscape.
-* Create holes in your Landscape for features such as caves.
+* Create **[holes](Engine/Landscape\Materials)** in your Landscape for features such as caves.
 
-For information about creating Landscapes, see **[](Engine/Landscape\Creation)**.
+For information about creating Landscape's, see **[](Engine/Landscape\Creation)**.
 
 For a quick reference guide to the various Landscape tool modes and options, see **[](Engine/Landscape/UI)**.
 
-For examples of the Landscape tool in action, see the **[](Resources\ContentExamples\Landscapes)** Content Example.
+For examples of the Landscape tool in action, see the **[](Resources\ContentExamples\Landscapes)**.
 
 
 [REGION:topics third]
