[UEDOC-1431] Typo fix.

diff --git a/Engine/Documentation/Source/Programming/Tutorials/Components/4/Components_4.INT.udn b/Engine/Documentation/Source/Programming/Tutorials/Components/4/Components_4.INT.udn
--- a/Engine/Documentation/Source/Programming/Tutorials/Components/4/Components_4.INT.udn
+++ b/Engine/Documentation/Source/Programming/Tutorials/Components/4/Components_4.INT.udn
@@ -84,7 +84,7 @@ Description:Tying the functionality of the Pawn, the input configuration, and th
 			{
 				if (OurParticleSystem && OurParticleSystem->Template)
 				{
-					OurParticleSystem->Activate(!OurParticleSystem->IsActive());
+					OurParticleSystem->ToggleActive();
 				}
 			}
 
