#UE4doc #docs Updates and image updates requested from user within the doc feedback forum

diff --git a/Engine/Documentation/Source/Engine/Landscape/Editing/Gizmos/Gizmos.INT.udn b/Engine/Documentation/Source/Engine/Landscape/Editing/Gizmos/Gizmos.INT.udn
--- a/Engine/Documentation/Source/Engine/Landscape/Editing/Gizmos/Gizmos.INT.udn
+++ b/Engine/Documentation/Source/Engine/Landscape/Editing/Gizmos/Gizmos.INT.udn
@@ -29,7 +29,7 @@ Description:Guide to the Gizmos available for modifying Landscapes.
 ## Overview
 
 
-Landscape Gizmos are editor-only Actors that are similar to volumes in that they define a specified area. Their purpose is to hold height and layer data for an area of a Landscape so that it can be copied to another location on the Landscape or exported for use in a different Landscape.
+Landscape Gizmos are editor-only Actors that are similar to volumes in that they define a specified area. Their purpose is to hold height and layer data for an area of a Landscape so that it can be copied to another location on the Landscape or exported for use in a different Landscape or heightmap generator (world machine, terresculptor,...).
 
 ![](Using_Gizom_Tool.gif)(convert:false)
 
