TypoCorrection

diff --git a/Engine/Documentation/Source/Engine/Landscape/Editing/PaintMode/PaintMode.INT.udn b/Engine/Documentation/Source/Engine/Landscape/Editing/PaintMode/PaintMode.INT.udn
--- a/Engine/Documentation/Source/Engine/Landscape/Editing/PaintMode/PaintMode.INT.udn
+++ b/Engine/Documentation/Source/Engine/Landscape/Editing/PaintMode/PaintMode.INT.udn
@@ -1,4 +1,4 @@
-Availability: Public
+﻿Availability: Public
 Title:Paint Mode
 Crumbs: %ROOT%, Engine, Engine/Landscape, Engine/Landscape/Editing
 Description:Guide to the Landscape tool's Paint mode.
@@ -179,7 +179,7 @@ The primary purpose of layer info objects is to act as a unique key for painted
 
 | Option | Description |
 | ------ | ----------- |
-| **Phys Material** | The [Physical Material](https://docs.unrealengine.com/latest/INT/Engine/Physics/PhysicalMaterials/index.html) (if any) assigned to areas of Landscape where this layer is dominant. |
+| **Phys Material** | The [Physical Material](Engine/Physics/PhysicalMaterials) (if any) assigned to areas of Landscape where this layer is dominant. |
 | **Hardness** | The value used by the [Erosion](Engine/Landscape/Editing/SculptMode/#Erosion) tool. |
 
 ### Orphaned Layers
