TypoCorrection

diff --git a/Engine/Documentation/Source/Engine/Landscape/Editing/Editing.INT.udn b/Engine/Documentation/Source/Engine/Landscape/Editing/Editing.INT.udn
--- a/Engine/Documentation/Source/Engine/Landscape/Editing/Editing.INT.udn
+++ b/Engine/Documentation/Source/Engine/Landscape/Editing/Editing.INT.udn
@@ -1,4 +1,4 @@
-Availability: Public
+﻿Availability: Public
 Title:Editing Landscapes
 Crumbs: %ROOT%, Engine, Engine/Landscape
 Description:Guide to using the tools for editing Landscape terrains.
@@ -39,7 +39,7 @@ After you have created a Landscape, you can now do any of the following to it:
 * Create **[Splines](Engine/Landscape\Editing\Splines)** for any linear features you want to use in your Landscape.
 * Create **[holes](Engine/Landscape\Materials)** in your Landscape for features such as caves.
 
-For information about creating Landscape's, see **[](Engine/Landscape\Creation)**.
+For information about creating Landscapes, see **[](Engine/Landscape\Creation)**.
 
 For a quick reference guide to the various Landscape tool modes and options, see **[](Engine/Landscape/UI)**.
 
