#UE4doc #docs: fixed shadow box link for UEDOC-1065

diff --git a/Engine/Documentation/Source/Resources/Showcases/Stylized/PostProcessing/PostProcessing.INT.udn b/Engine/Documentation/Source/Resources/Showcases/Stylized/PostProcessing/PostProcessing.INT.udn
--- a/Engine/Documentation/Source/Resources/Showcases/Stylized/PostProcessing/PostProcessing.INT.udn
+++ b/Engine/Documentation/Source/Resources/Showcases/Stylized/PostProcessing/PostProcessing.INT.udn
@@ -58,7 +58,7 @@ This effect can be seen within the M_PostProcess_Base Material, primarily in sec
 [/REGION]
 
 [REGION:lightbox fullwidth]
-[![](LineRender02.png)](LineRender.png)
+[![](LineRender02.png)](LineRender02.png)
 [/REGION]
 
 **For each of the images above**, it will be best to **right-click** each one and choose Save As if you would like to view them full size (the Material Network can also be found inside the Stylized Showcase).
