#UE4 Docs: Adding the first pass on the Gear VR Documents.

diff --git a/Engine/Documentation/Source/Platforms/GearVR/BestPractices/GearVRBestPractices.INT.udn b/Engine/Documentation/Source/Platforms/GearVR/BestPractices/GearVRBestPractices.INT.udn
new file mode 100644
--- /dev/null
+++ b/Engine/Documentation/Source/Platforms/GearVR/BestPractices/GearVRBestPractices.INT.udn
@@ -0,0 +1,161 @@
+Title:Samsung Gear VR Best Practices
+Crumbs:%ROOT%, Platforms
+Description:Information on how to get the most out of your Gear VR headset.
+Availability:Public
+Navigation:platform
+Platform:Gear VR
+Related: Platforms/Mobile/Performance
+Related: Engine/Performance
+Related: Engine/Performance/StatCommands
+Related: Platforms/Mobile/Textures
+Related: Platforms/Mobile/Materials
+Related: Engine/Rendering/Materials/MaterialProperties/BlendModes
+Related: Engine/Rendering/Materials/CustomizedUVs
+
+[TOC(start:2 end:3)]
+
+## Overview
+When developing content for Virtual Reality (VR) projects, certain steps must be taken to ensure that the user will have the best VR experience possible.
+In the following Best Practices guide we will cover all you need to know to help guarantee that the user has the best possible experience when viewing your project 
+using the Samsung Gear Virtual Reality(VR) Headset and Samsung Galaxy Note 4 Smartphone.
+
+
+## Gear VR Rendering Options
+While combining the Gear VR Headset with the Samsung Galaxy Note 4 Smartphone can produce some very immersive VR experiences, 
+there are some rendering features that UE4 can not perform while using the Gear VR Headset due to hardware limitations. 
+Many of the restrictions that you will encounter are similar to restrictions for [mobile devices](https://docs.unrealengine.com/latest/INT/Platforms/Mobile/Performance/index.html).
+However because of the unique rendering requirements for the Gear VR Headset there are few things extra things that must be considered.
+In the following sections we will cover what these requirements are and how they can effect your Gear VR project. 
+
+### Mobile HDR
+While you may be tempted to enable Mobile HDR so that you enable more advanced rendering features for your projects but **DO NOT** do this!
+Mobile HDR **will not work** with the Gear VR Headset and should be left off. 
+Enabling this feature will cause your project to not be rendered correctly when placed into the Gear VR Headset.
+
+### Post Processing
+Due to the demanding rendering requirements for the Gear VR Headset, there are many Post Processing features that unfortunately will not work with the Gear VR Headset. 
+Below you will find a list of all the Post Process features that will work with the Gear VR Headset.
+
+|Post Processing Section |Feature Name| Description |
+| -------- | ---------- |------|
+| Auto Exposure|Exposure Bias |B Can be used to increase or decrease the scenes brightness. |
+
+### Materials
+Creating and interacting with Materials for Gear VR projects is very similar to creating and interacting with Materials for other projects. 
+You still use the Material Editor to create Materials and you can still use many of the Material nodes you find in the Material Editor. 
+However there are a few restrictions / considerations that must be adhered to when making Materials for Gear VR to ensure that what you make will work with the Gear VR Headset. 
+For example only the following Material inputs will work with the Gear VR Headset.
+
+| Material Input Name | Description |
+| -------- | ---------- |
+| **Base Color** |Base Color defines the overall color of the Material. |
+| **Roughness** |Roughness defines how rough or smooth the Material is. |
+| **Metallic** |Metallic defines how metal-like the Material is. |
+| **Specular** |Specular defines the intensity of the reflections. |
+| **Normal**  |Normal helps to add detail that would otherwise be to costly to model in. |
+| **Emissive**  |Emissive helps to make parts of Materials glow like they are emitting light. |
+| **Opacity**  |Opacity helps to define how see through something is. |
+| **Opacity Mask** |Opacity Mask helps to define if a certain pixel is see through or not. |
+
+[REGION:note]
+Please keep in mind that the both **Translucent** and **Masked** Materials can be very expensive to render, especially on hardware like Smartphone's.
+Whenever possible you should use [Opaque Materials](https://docs.unrealengine.com/latest/INT/Engine/Rendering/Materials/MaterialProperties/BlendModes/index.html#opaque) as they are much cheaper to render.
+[/region]
+
+### Material Shading Models
+Due to hardware limitations only the following Material Shading Models are available for use in projects that use the Gear VR Headset.
+| Shading Model Name | Description |
+| -------- | ---------- |
+| **Default Lit** |Default Lit is, as the name suggests, the default shading model and will likely be the one that you will use most often. This shading model makes use of direct and indirect lighting, as well as specularity for reflections.. |
+| **Unlit** |The Unlit Shading Model only outputs Emissive for color, making this perfect for special effects such as fire or illuminating objects.|
+
+[region:note]
+Try to use the **Unlit** Shading Model as much as possible with your Gear VR projects as it is the cheapest to use. 
+[/region]
+
+### Material Blend Modes
+Material Blend Modes describe how the output of the current Material will blend over what is already being drawn in the background. 
+When creating content for the Gear VR Headset the following Material Blend Modes are available for use.
+
+| Blend Mode | Description |
+| ---- | ----- |
+| **Opaque** | For a Solid object, **Opaque** is the cheapest material. But, as its name implies, it has no support for translucency. That means it will never create overdraw, and rendering one pixel of the material will always take the same time depending on the complexity of the **Node Graph**. |
+| **Masked** |  If you need some pixels not drawn on your material, this is the way to go. The material only creates overdraw for the area that is masked out. But the material is 100\% opaque or 100\% see through. The bigger the area that is masked out the more overdraw is created. That means more expensive. |
+| **Translucent** | This is the most expensive **BlendMode**. It can look like a **Masked** and **Opaque** but will have an additional effect on performance regardless of actually being transparent or not. |
+| **Additive** | Is a cheaper alternative to **Translucent**. It does not need to sort the pixels, they just get added to the scene. But that means it can only brighten. Black pixels on the material have no effect on the scene.|
+| **Modulated** | Is the second cheapest alternative to **Translucent**, but it multiplies its Pixels with the scene. That means it can only darken the scene.|
+[region:note]
+When using the **Modulate** blend mode make sure to disable **Separate Translucency** or your Material will not compile correctly.
+  ![](T_Separate_Trans.png)
+[/region]
+
+### Mobile Specific Material Settings
+On the **Main Material Node** under the **Details Panel** you will find a section called **Mobile**.
+The settings contained here are for helping to lessen Material rendering costs on mobile devices.  
+| Property | Description |
+| -------- | ---------- |
+| bFullyRough | When enabled, saves a number of shader ALU instructions and one sampler.  This effectively disables Specular while maintaining overall reflective energy. |
+| bUseLightmapDirectionality | When disabled, lighting from light maps will be flat but cheaper. |
+
+### Customized UVs
+[Customized UV's](https://docs.unrealengine.com/latest/INT/Engine/Rendering/Materials/CustomizedUVs/index.html) allow you to perform certain shader operations in vertex shaders instead of pixel shaders to help avoid dependent texture fetches. 
+This offers an excellent way to claw back much needed performance when doing simple things like UV manipulation to make a texture title more than it currently does. 
+Whenever possible make sure to do all UV manipulation using Customized UVs especially when developing content for the Gear VR Headset.
+
+[region:note]
+Customized UVs are currently only implemented for certain component types such as Static Meshes, Skeletal Meshes, BSP, Landscape, and Mesh Particles.
+[/region]
+
+### Reflections
+The Gear VR does not support Screen Space Reflections(SSR) which means you will only be able generate reflections that come from [Reflection Capture Actors](https://docs.unrealengine.com/latest/INT/Resources/ContentExamples/Reflections/1_4/) that are placed in levels. 
+Also keep in mind that the reflections that are generated from the reflections actors will not be of very high fidelity due to hardware limitations. 
+
+## Gear VR & Content Setup
+When creating Art content for your Gear VR projects there are few things that you need to keep in mind to ensure that your project can run as smooth as possible on the Gear VR Headset. 
+In the following section we will cover what you need to be concerned with when making Art content for the Gear VR Headset.
+
+### Content Hard limits
+Due to certain hardware restrictions on mobile devices the following restrictions **Must Be** adhered to when creating content. 
+Failing to follow these restrictions will result in your project now working.
+
+* All mesh types can have up to but not over 65,000 vertices due to the lack of 32 bit index support on mobile hardware.
+* Skeletal Meshes can have up to but not 75 bones.
+
+### Mesh Considerations
+When creating Meshes to use with your Gear VR Project the following information needs to be kept in mind to ensure that what you are creating will work with the Gear VR Headset. 
+
+* Use as few [Material ID's](https://docs.unrealengine.com/latest/INT/Engine/Content/FBX/Materials/index.html) on 3D meshes as possible to save on draw calls.
+* Make sure to make use of [Static Mesh LOD's](https://docs.unrealengine.com/latest/INT/Engine/Content/FBX/StaticMeshes/index.html#lodsetup) whenever possible.
+* Keep an eye on your 3D meshes polygon count and try and keep it as [Low Poly](https://en.wikipedia.org/wiki/Low_poly) as possible.
+* Make sure that all static models have a 2nd UV set so they can make use of [Light Mapping](https://docs.unrealengine.com/latest/INT/Engine/Content/Types/StaticMeshes/LightmapUnwrapping/index.html).
+* Whenever possible try and make use if [Imposter Sprites](https://docs.unrealengine.com/latest/INT/Engine/Content/Tools/RenderToTextureTools/3/index.html) to replace far away Static Meshes as they are much cheaper to render.
+
+### Texture Considerations
+When creating Textures to use with your Gear VR Project the following information needs to be kept in mind to ensure that what you are creating will not be detrimental to your projects overall performance. 
+
+* Textures use the [ECT2](http://en.wikipedia.org/wiki/Ericsson_Texture_Compression) Texture compression format.
+* Textures **can not** be bigger than 2048 pixels in either the X or Y dimension.
+* Textures sizes **MUST BE** in a power of 2 (2, 4, 8, 16, 32, 64, 128, 256, 512, 1024, 2048)
+* Loading up to 512 MB of Textures into memory is possible with the Gear VR however it is recommended that you do not load more that 128 MB of Textures into memory.
+* Use square Textures whenever possible as them make a more efficient use of memory. 
+* Wherever possible try and use [Texture Atlasing](http://en.wikipedia.org/wiki/Texture_atlas) to reduce the number of Textures needed in memory.
+* [RGB Texture Masking / Packing](https://docs.unrealengine.com/latest/INT/Engine/Rendering/Materials/HowTo/Masking/index.html) is a great way to cut down on unique Textures in memory.
+* Only **TC_Default** and **TC_NormalMap** compression settings should be used as they are the only two that will work on mobile platforms. 
+
+### Level Considerations
+Getting the maximum performance out of your levels on the Gear VR is a delicate balancing act that will require time to get just right. 
+The following suggestions / things to look out for can help you quickly identify areas of your levels that might cause performance issues when viewed on the Gear VR Headset.
+
+* Always make sure that the level lighting has been [built](https://docs.unrealengine.com/latest/INT/Engine/Rendering/LightingAndShadows/Lightmass/index.html) to ensure you get the best performance possible. If you see a red **Lighting needs to be rebuilt** error message on the screen, this means that your lighting needs to be built.
+* Make sure that all the lights in you level are set to [Static](https://docs.unrealengine.com/latest/INT/Engine/Rendering/LightingAndShadows/LightMobility/StaticLights/index.html) as the other lighting options can be very expensive to render.
+* Check to see if your level is making use of [Precomputed Visibility](https://docs.unrealengine.com/latest/INT/Engine/Actors/Volumes/index.html) and if not make sure to add one to the level and re-build lighting to enable it in the level's World Settings. 
+* Draw calls for the entire scene should not be greater than 100 for any view. This can be seen with [stat openglrhi](https://docs.unrealengine.com/latest/INT/Engine/Performance/StatCommands/index.html) on the device, [stat d3d11rhi](https://docs.unrealengine.com/latest/INT/Engine/Performance/StatCommands/index.html) in ES2 previewer on PC.
+* The total triangle count of the level should not be greater than 100,000 for any view. This can be seen with [stat openglrhi](https://docs.unrealengine.com/latest/INT/Engine/Performance/StatCommands/index.html) on the device, [stat d3d11rhi](https://docs.unrealengine.com/latest/INT/Engine/Performance/StatCommands/index.html) in ES2 previewer on PC.
+* Be very careful in how Translucent and Opacity Masked Materials are used in the level as they are very expensive to render.
+* Make sure to [Performance Profile](https://docs.unrealengine.com/latest/INT/Engine/Performance/index.html) your level often to make sure you catch any performance issues sooner rather than later.
+
+
+
+
+
+
diff --git a/Engine/Documentation/Source/Platforms/GearVR/Debugging/GearVRDebugging.INT.udn b/Engine/Documentation/Source/Platforms/GearVR/Debugging/GearVRDebugging.INT.udn
new file mode 100644
--- /dev/null
+++ b/Engine/Documentation/Source/Platforms/GearVR/Debugging/GearVRDebugging.INT.udn
@@ -0,0 +1,220 @@
+Title:Samsung Gear VR Debugging
+Crumbs:%ROOT%, Platforms
+Description:Information on how to debug applications for the Gear VR headset.
+Availability:Public
+Navigation:platform
+Platform:Gear VR
+Related: Platforms/Mobile/Performance
+Related: Engine/Performance
+Related: Engine/Performance/StatCommands
+Related: Platforms/Mobile/Textures
+Related: Platforms/Mobile/Materials
+Related: Engine/Rendering/Materials/MaterialProperties/BlendModes
+Related: Engine/Rendering/Materials/CustomizedUVs
+
+[TOC(start:2 end:3)]
+
+Being able to debug your project on it's target device is critical to tracking down performance issues and other possible project breaking errors. 
+In the following section you will learn how what tools are available to assist you in debugging your Gear VR projects and how to use them.
+
+### Enabling Gear VR Developer Mode
+Having the ability to view your Gear VR Project without having to use the Gear VR Headset can greatly speed up the time it takes to debug an issue by allowing you quick and easy access to things like Unreal's console. 
+In the following sections we will cover how you can set up your Gear VR projects so that they can be viewed without the Gear VR Headset. 
+
+[region:note]
+Before you can do this you must make sure that you have some type of Gear VR application installed to your phone that has your device signature in it. 
+This can be any kind of application, download or compiled.
+[/region]
+
+1. First on your Samsung Galaxy Note 4 Smartphone go to **Settings** > **Application Manager** > **Gear VR Service**.
+	
+	![](T_GearVR_Dev_Mode_01.png)
+
+1. Next tap on the **Manage Storage** option.
+  
+	![](T_GearVR_Dev_Mode_02.png)
+
+1. Then tap on the **VR Service Version** **6** times.
+
+	![](T_GearVR_Dev_Mode_03.png)
+
+1. You should now see a scan process happen and when completed you should now see the **Developer Mode** toggle on the screen.
+
+	![](T_GearVR_Dev_Mode_04.png)
+
+1. Now you should be able to launch your Gear VR Project and have it display without having to use the Gear VR Headset.
+
+	[region:tip]
+	Be careful when enabling VR Developer mode as it can drastically drain your phone's battery as you are overriding any settings to turn off the Gear VR Headset when not in use.
+	[/region]
+
+1. If for any reason you need to disable VR Developer mode just slide the **Developer Mode** slider from **On** to **Off** and VR Developer mode will be disabled.
+
+### Enabling the UE4 Console
+The UE4 Console is a very powerful tool that allows you to look at and in some cases affect data that is currently being displayed in UE4. 
+Accessing the UE4 console on Mobile devices requires a slightly different approach than what you might be used to on PC and is explained below.
+
+[region:tip]
+You can get a full list of available console commands by entering **DumpConsoleCommands** into the console while on a PC and Unreal will then list all available console commands in the console.
+[/region]
+
+1. First make sure that you have enabled Gear VR Developer mode on your phone. 
+	This way you can launch your project and see it without having to use the Gear VR Headset which, if used, would prevent you from being able to touch your phone's screen.
+	
+	![](T_GearVR_Dev_Mode_04.png)
+
+1. Next launch your Gear VR project by tapping on the projects icon to launch it. 
+	When the project has been launched, you should see something looking like this on your phone.
+	
+	![](T_App_On_Phone.png)
+
+1. With your Gear VR project now running, tap the phone's screen with **Four (4)** fingers at once to open up the Unreal Engine Console. 
+	If done correctly you should see something like this on your phone's screen. 
+
+	![](T_GearR_Open_Console.png)
+
+1. Now input a console command you wish to use in the input box and then press the **Ok** button. 
+	For this example we will use the **Stat Memory** command to see how much memory the project is using. 
+
+	![](T_GearR_Input_Command.png)
+
+1. Once the **Ok** button has been pressed, your project should now be displaying how much memory is being used and by what.
+
+	![](T_GearR_Displaying_Memory.png)
+
+1. When you are done debugging or if you just wish to clear the screen of debug information, you can input **Stat None** to clear the screen off all debug information.
+
+	[region:tip]
+	The best option for console debugging with the Gear VR Headset is to use a Bluetooth keyboard. This way you can use the console with out having to remove the phone from the headset.
+	[/region]
+
+### Viewing Gear VR Debug Logs
+Having the ability to view debug logs is of critical importance when developing any kind of 3D project. 
+Debug logs can provide vital information and insight into why your project is not working and in some case's even make suggestions as to what you can do to fix things. 
+In the following section we will go over what you need to do in order to get debug logs from your Smartphone.
+
+[region:note]
+For much more detailed information on all of various things you can do with [Logcat](http://developer.android.com/tools/help/logcat.html), check out the official Android documentation on [Logcat](http://developer.android.com/tools/help/logcat.html).
+[/region]
+
+1. First make sure that your Smartphone with the application you want to get debug information from is plugged into your PC via USB.
+
+1. Next open up a Windows Command Shell by first pressing the **Windows Key + R** and then entering **CMD** in the **Run** dialogue box that is opened.
+	Then press either the **Enter** key on your keyboard or just click on the **OK** button with your mouse. 
+	 
+	 ![](T_Open_CMD_Prompt.png)
+
+1. Once the **OK** button has been toggled a Windows Command Line window like the following should now be displayed.
+	
+	![](T_Windows_CMD.png)
+
+1. Now check to make sure that your PC can see your Smartphone by entering **ADB Devices** into the command line window and then press the **Enter** key to display devices attached to your PC.
+	
+	![](T_Windows_ADB_Devices.png)
+
+	[region:tip]
+	If you enter **ADB Devices** but do not see anything show up, make sure that you have gone through the [Gear VR Headset Prerequisites](Platforms/GearVR) setup section and installed / setup everything that is listed there and then try this step again. 
+	[/region]
+
+1. Now that you have verified that you can see your Smartphone in ADB Devices it is now time to get the log information from your phone. 
+	But first we are going to clear the log cache to ensure that we are getting only current debug information and not old debug information. 
+	To clear the log enter the following command.
+
+		adb logcat -c
+	![](T_ADB_Clear_Log.png)
+
+1. With the log now cleared of old data, launch the application on your Smartphone and enter the following command to get the log information.
+	
+		adb logcat -s UE4 OVR VrApi VrLib VRSVC HmdInfo >logoutput.txt
+	![](T_ADB_Get_Log.png)
+ [region:note]
+ To see the live output of the log, omit the **>logoutput.txt** part and all log updates will happen in the command line window.
+ [/region]
+
+1. To complete the log capture press **CTRL + C** on the keyboard and when completed you should now have a log file named,** logoutput.txt**, that is saved to your **Users** folder in Windows
+	![](T_ADB_Log_Location.png)
+
+	[region:note]
+	Data will still continue to be output to the log file while the command line window is open and running the ADB logcat command. To stop this all you have to do is press **CTRL + C** and data will stop being written to the log file.
+	[/region] 
+
+1. The information that is contained within the log file was generated from adding **UE4**, **OVR**,  **VrApi**, **VrLib**,**VRSVC**, and **HmdInfo** to the logcat command. 
+	When you look at the log file that was created and search for anyone of those terms you will find all debug information related to those terms. 
+	You can see an example of this in the image below, some of the terms used, **OVR**, **UE4**, **Vrapi**, and **VrApi** all show up in the log with information that is related to those terms.
+	
+	![](T_ADB_Log_Info.png)
+
+If your Gear VR Project is not starting up for some reason, using the logcat command is a great way to see exactly where your project is failing and why. 
+It also provides you with detailed information about the problem that can then be posted to the UE4 Forums or Answer Hub so that you can get with the issue.
+
+### Wireless Deploying & Debugging
+Getting debug information or deploying to your Smartphone while it is in the Gear VR Headset is impossible due to the only USB port being used by the Gear VR Headset. 
+However you can work around this issue by using your Smartphone's Wifi connection to send debug information to your development PC. 
+You can even deploy projects in this manner however it is extremely slow when compared to using the USB cable. 
+In the following section we will cover all you need to know about setting up your Smartphone so that it can send debug information and be deployed to over wifi.
+
+[region:note]
+Please note that in order for this to work you have to have your development PC and Galaxy Note 4 on the same network.
+[/region]
+
+1. Next open up a Windows Command Shell by first pressing the **Windows Key + R** and then entering **CMD** in the **Run** dialogue box that is opened.
+	Then press either the **Enter** key on your keyboard or just click on the **OK** button with your mouse. 
+	 
+	 ![](T_Open_CMD_Prompt.png)
+
+1. Once the **OK** button has been toggled a Windows Command Line window like the following should now be displayed.
+	
+	![](T_Windows_CMD.png)
+
+1. Now enter the following command to see what your Smartphone's IP address is and make sure to take note of it. The IP address you want to use is the one that comes after **inet**.
+	
+		adb shell ip -f inet addr show wlan0
+	
+	![](T_ADB_Phone_IP.png)
+
+1. Next enter the following command to change the network connection to **TCP** and set the port to listen to.
+
+		adb tcpip 5555
+	
+	![](T_ADB_Set_Port.png)
+
+1. Once the port has been set, it is now time to connect to the Smartphone. To do this all you need to do is enter the following command and press the Enter key.
+	
+		adb connect 10.3.31.133 - Replace with your devices IP Address.
+		
+1. You should now see a message that says **connected to** and the IP address of your device.
+	
+	![](T_ADB_Connected.png)
+
+1. To make sure you are connected over Wifi, enter **adb devices** into the command line window and press the Enter key and you will now see that in the **List of devices attached** the IP Address used above is now the only attached device.
+
+	![](T_ADB_Attached_IP.png)
+
+With your Smartphone and PC now connected to one another via the network you can now begin to get debug information from your Smartphone or even deploy projects to your phone with out having to use the USB calbe. 
+However keep in mind that deploying to your Smartphone from your PC over the network is four to five times slower than using the USB cable.
+
+### Return to Wired Deploying & Debugging
+1. To return to using the USB cable for debugging and deployment, input the following command into the command window.
+	
+		adb usb
+
+	![](T_ADB_Detached.png)
+
+
+
+
+
+
+
+
+
+
+
+
+
+
+
+
+
+
+
diff --git a/Engine/Documentation/Source/Platforms/GearVR/GearVR.INT.udn b/Engine/Documentation/Source/Platforms/GearVR/GearVR.INT.udn
new file mode 100644
--- /dev/null
+++ b/Engine/Documentation/Source/Platforms/GearVR/GearVR.INT.udn
@@ -0,0 +1,83 @@
+Title: Samsung Gear VR Headset Prerequisites
+Crumbs:%ROOT%, Platforms
+Description:Instructions on how to get the Samsung Gear VR Headset to work with UE4.
+Availability:Public
+Navigation:platform
+Platform:Gear VR
+
+![](Platforms\GearVR\QuickStart\GVR_Headset_00.png)
+
+[REGION:warning]
+Developing for the Gear VR Headset requires **Unreal Engine 4.7.4** or later.
+[/REGION]
+
+## Overview 
+
+With the 4.7 release of Unreal Engine 4 (UE4) you can now build Virtual Reality(VR) projects for the Samsung Gear VR Headset that can transform your Samsung Galaxy Note 4 Smartphone into an immersive VR device. 
+In the following sections we will go over everything you need to know to get your Samsung Galaxy Note 4 Smartphone properly setup so that it can work with the Samsung Gear VR headset.
+
+## Gear VR Headset Pre-requisites 
+
+Before we can get into UE4 Gear VR development we first must make sure that the following prerequisites are full filled as if we do not, we run the risk of our Gear VR projects not working with the Gear VR Headset. 
+
+* Make sure you have the USB cord used for charging your Samsung Galaxy Note 4 Smartphone available and plugged into a USB port on the PC you are developing on. 
+The USB cord is of vital importance as this is how you transfer data between your development PC and your Smartphone.
+
+* Only certain types of Samsung Galaxy Note 4's will work with the Gear VR Headset. 
+	The [Oculus website](https://support.oculus.com/hc/en-us/articles/204113096-Samsung-Gear-VR-Innovator-Edition-Samsung-Galaxy-Note-4-Compatibility) has a list of phone models that will work with the Gear VR Headset.
+	Make sure to double check that the Samsung Galaxy Note 4 Smartphone you are buying is the correct model that will work with the Gear VR Headset.
+
+* Make sure to download and install the specific USB drivers for your Galaxy Note 4 from the [Samsung Website](http://www.samsung.com/us/support/downloads). 
+
+* If you do not have it all ready, make sure to download and install [Microsoft Visual Studio 2013](https://www.visualstudio.com/en-us/downloads/download-visual-studio-vs.aspx).
+
+* Download Unreal Engine version **4.7.4** or later from the [UE4 Launcher](https://www.unrealengine.com/) or [Github](https://www.unrealengine.com/ue4-on-github). 
+
+* Install the **Tegra Android Development Pack** which can be found in **Engine\Extras\Android**.  
+
+* Check to see if [Developer Mode](https://docs.unrealengine.com/latest/INT/Platforms/Android/GettingStarted/2/index.html ) has been enabled on your Samsung Galaxy Note 4 Smartphone and if not enable it. 
+
+* Also make sure to check that USB Debugging has been enabled to ensure that you can load data onto your Smartphone from your PC.
+	
+	![](GVR_Dev_Options_Enabled.png)
+
+* With your Smartphone plugged into your development PC via USB, make sure to accept the RSA Fingerprint key from your PC.
+This way your development PC and Smartphone can send data back and forth.
+Also make sure to check **Always allow this computer** to ensure that your development PC can always interact with this Smartphone when plugged in. 
+
+	![](GVR_USB_Debuging.png)
+
+* Generate and download the **OSIG File** for your Samsung Galaxy Note 4 Smartphone using your Smartphone's Device ID and the [Oculus](https://developer.oculus.com/tools/osig/) website.
+If you are not sure how to get your Smartphone's **Device ID** the [Oculus website](https://developer.oculus.com/tools/osig/) has detailed directions on how to do this.
+	
+	![](GVR_OSIG_Website.png)
+
+* After you download the OSIG file, create a folder called **assets** and place the **OSIG** file inside of it. 
+Then place that folder in **Engine/Build/Android/Java**. 
+
+	[region:warning]
+	**If you do not do this your project will not be able to run on your Galaxy Note 4 Smartphone.**
+	[/region]
+	
+	![](GVR_OSIG_File_Location.png)
+
+It is very important that you make sure all items listed in the Gear VR Prerequisites section are completed before moving on to the next section. 
+Skipping or failing to do a step could result in your project not working with the Gear VR Headset. 
+Once everything from above has been completed you can now move on to the Gear VR Headset Development Quick Start.
+
+
+[OBJECT:TopicButtonList]
+	[PARAM:icon]
+		![](%ROOT%/vr_icon.png)(convert:false)
+	[/PARAM]
+	[PARAM:title]
+		Gear VR Development
+	[/PARAM]
+	[PARAM:description]
+	[/PARAM]
+	[PARAM:links]
+		* [](Platforms/GearVR/QuickStart "%Platforms/GearVR/QuickStart:description%")
+		* [](Platforms/GearVR/BestPractices "%Platforms/GearVR/BestPractices:description%")
+		* [](Platforms/GearVR/Debugging "%Platforms/GearVR/BestPractices:description%")
+	[/PARAM]
+[/OBJECT]
\ No newline at end of file
diff --git a/Engine/Documentation/Source/Platforms/GearVR/QuickStart/1/GearVRQuickStart_1.INT.udn b/Engine/Documentation/Source/Platforms/GearVR/QuickStart/1/GearVRQuickStart_1.INT.udn
new file mode 100644
--- /dev/null
+++ b/Engine/Documentation/Source/Platforms/GearVR/QuickStart/1/GearVRQuickStart_1.INT.udn
@@ -0,0 +1,59 @@
+Availability:Public
+Title:1. Gear VR C++ Project Creation
+Crumbs:%ROOT%, Platforms/GearVR
+Description:How to create a UE4 C++ based project that can be used with the Gear VR Headset.
+Navigation:platform
+
+
+[VAR:Steps]
+[OBJECT:Navigation]
+	[PARAM:previous]
+	[/PARAM]
+	[PARAM:current]
+	[/PARAM]
+	[PARAM:home]
+		[Quick Start Home](Platforms/GearVR/QuickStart)
+	[/PARAM]
+	[PARAM:next]
+		[Next Step](Platforms/GearVR/QuickStart/2 "%Platforms/GearVR/QuickStart/2:title%")
+	[/PARAM]
+[/OBJECT]
+[/VAR]
+
+
+%Steps%
+
+Setting up a UE4 project that works with the Samsung Gear VR Headset requires a slightly different setup than a normal UE4 project. 
+This is because a Gear VR Headset project needs to be created as a C++ project and then compiled to ensure that the proper files for the Gear VR Headset are compiled into the projects binary files. 
+In the following section we will go over everything you need to know about how to properly setup your projects so they work with the Gear VR Headset.
+
+[REGION:warning]
+Currently projects that are created using the Blueprint based templates **WILL NOT** work with the Gear VR Headset. 
+This means that all newly created Gear VR projects **MUST** use one of the C++ based templates. 
+However this does not mean that you can not use Blueprints in your projects even if they were created using the C++ based templates.
+[/REGION]
+
+1. First make sure that you have UE4 version **4.7.4** or later downloaded and installed on your development PC. 
+You can get UE4 version **4.7.4** or later from the [Unreal Engine Launcher](http://tinyurl.com/oawb7x2) or from [Github](https://www.unrealengine.com/ue4-on-github). 
+
+	[Region:note]
+	If you are using the [Github](https://www.unrealengine.com/ue4-on-github) version of UE4 to run the **Setup.Bat** file before going any further to ensure that you have the proper files downloaded to your development PC.
+	[/REGION]
+
+	![](GVR_New_Project_03.png)
+
+1. Once you have a version of UE4 **4.7.4** or later installed on your development PC create a new **C++ Project** from the **New Project Tab**.
+	![](GVR_New_Project_00.png)
+
+1. Next select the **First Person Template** as the template to be used and make sure to change the following options.
+	* **Mobile / Tablet**
+	* **Scalable 3D / 2D** 
+	* Also make sure to not include the **Starter Content** to help reduce project size.
+	![](GVR_New_Project_01.png)
+
+1. Finally give the project a location to be saved to and a name. The name for this example the project name will be **GearVR_QS**. Once everything has been completed press the green **Create Project** button to create the project.
+	![](GVR_New_Project_02.png)
+
+
+
+%Steps%
\ No newline at end of file
diff --git a/Engine/Documentation/Source/Platforms/GearVR/QuickStart/2/GearVRQuickStart_2.INT.udn b/Engine/Documentation/Source/Platforms/GearVR/QuickStart/2/GearVRQuickStart_2.INT.udn
new file mode 100644
--- /dev/null
+++ b/Engine/Documentation/Source/Platforms/GearVR/QuickStart/2/GearVRQuickStart_2.INT.udn
@@ -0,0 +1,81 @@
+Availability:Public
+Title:2. Gear VR Project Setup
+Crumbs:%ROOT%, Platforms/GearVR
+Description:How to setup your C++ based UE4 project to be used with the Gear VR Headset.
+Navigation:platform
+
+[VAR:Steps]
+[OBJECT:Navigation]
+	[PARAM:previous]
+	[Previous Step](Platforms/GearVR/QuickStart/1 "%Platforms/GearVR/QuickStart/1:title%")
+	[/PARAM]
+	[PARAM:current]
+	[/PARAM]
+	[PARAM:home]
+		[Quick Start Home](Platforms/GearVR/QuickStart)
+	[/PARAM]
+	[PARAM:next]
+		[Next Step](Platforms/GearVR/QuickStart/3 "%Platforms/GearVR/QuickStart/3:title%")
+	[/PARAM]
+[/OBJECT]
+[/VAR]
+
+%Steps%
+
+
+1. Once the project has finished being created and compiled, Visual Studio 2013 will automatically open (If already installed) along with the editor.
+
+	![](GVR_Editor_Opened.png)
+
+1. Once the UE4 Editor has loaded, from the **Main Toolbar**, open the **Windows** tab and select the **Plugins** menu item. 
+
+	![](GVR_Plugins_Menu.gif)(convert:false)
+
+1. From the Plug-In menu go to the **Head Mounted Displays** section and **disable** the **Oculus** option and **enable** the **Gear VR** option. Then click the **Restart Now** button to re-start the editor so the changes can take place. 
+
+	![](GVR_Plugin_Options_Toggle.png)
+
+1. Once the Editor has re-started, select **Edit** from the **Main Toolbar** and then select **Project Settings**  and click the **Configure Now** button for Android support in the **APKPacking** section.
+
+	![](GVR_Android_Setup_01.gif)(convert:false)
+
+1. Next change the **Minimum SDK Version** from **9** to **19**.
+
+	![](GVR_Android_SDK_Version.png)
+
+1. Next scroll down till you see the **Configure the Android Manifest for deployment to Gear VR** checkbox and enable it. 
+
+	[Region:note]
+	This step is of vital importance as without this on your project will not be able to run on the Gear VR Headset.
+	[/region]
+
+	![](GVR_Deployment_Checkbox.png)
+
+1. Next, under the **Android SDK** section fill, in the location of the required Android SDK files on your development PC. 
+	You can find the needed folders in the location where you installed the Tegra files. 
+	
+	[Region:note]
+	If you have already installed TADP this information should be already filled in for you. The only reason that you should change something here is if the information is missing or point to the wrong location.
+	[/region]
+
+	![](GVR_Android_SDK_Setup.png)
+
+1. Finally under **File** > **Package Project** > **Packing Settings** un-check **Full Rebuild** to help speed up iteration times. 
+
+	[Region:note]
+	Remember to **Re-Enable** this option when your project is ready for its final publish to ensure that your project will work as intended on the Gear VR Headset. 
+	[/region]
+
+	![](GVR_Disable_FullReBuild.gif)(convert:false)
+
+	<!-- This was step six, removed based on Josh Adams advice.
+	1. Next in the **Android Package Name** section give your project a name. For this example we will be using the following name **com.Company.GearVR_QS**. 
+
+	![](GVR_Android_Game_Name.png)
+
+	[REGION:note]
+	If using UE 4.7.0 or earlier the name that is used here must match the name that is used for the project or you app could have issues running on your Smartphone.
+	[/region]
+	-->
+
+%Steps%
\ No newline at end of file
diff --git a/Engine/Documentation/Source/Platforms/GearVR/QuickStart/3/GearVRQuickStart_3.INT.udn b/Engine/Documentation/Source/Platforms/GearVR/QuickStart/3/GearVRQuickStart_3.INT.udn
new file mode 100644
--- /dev/null
+++ b/Engine/Documentation/Source/Platforms/GearVR/QuickStart/3/GearVRQuickStart_3.INT.udn
@@ -0,0 +1,62 @@
+Title:3. Gear VR Project Packaging & Deployment
+Crumbs:%ROOT%, Platforms/GearVR
+Description:Instructions on how to deploy projects form UE4 to use with the Samsung Gear VR Headset.
+Availability:Public
+Navigation:platform
+Platform:Gear VR
+
+[VAR:Steps]
+[OBJECT:Navigation]
+	[PARAM:previous]
+	[Previous Step](Platforms/GearVR/QuickStart/2 "%Platforms/GearVR/QuickStart/2:title%")
+	[/PARAM]
+	[PARAM:current]
+	[/PARAM]
+	[PARAM:home]
+		[Quick Start Home](Platforms/GearVR/QuickStart)
+	[/PARAM]
+	[PARAM:next]
+		
+	[/PARAM]
+[/OBJECT]
+[/VAR]
+
+%Steps%
+
+
+1. First make sure that your Samsung Galaxy Note 4 Smartphone is connected to a USB cord and plugged into a USB port on the PC you are developing on.
+
+1. Next start the packaging process by going to the **Main Toolbar** and selecting **File** > **Package Project** > **Android** > **Android ETC 2**.
+
+	![](GVR_Deploy_To_Android.png)
+
+1. Once you click on the **Android ETC 2** option, a Windows dialogue box will be displayed. From this dialogue box, select or create a folder to save your packaged project to and then press the **OK** button to begin the packaging process.
+
+	![](GVR_Save_Project.png)
+
+1. Depending on the size of your project, the packaging process could take a few minutes or a few hours to complete. 
+	When it is completed the folder that you selected/created to save the build to should now have three files in it and should look similar to the image below. 
+
+	![](GVR_Deployed_Files.png)
+
+1. Next, run the **.BAT** file which is located in the folder that the build was packaged to by **Double - Clicking** on it with your **Left Mouse Button**.
+	Once you run the .BAT file, a command line window will be displayed that shows the current progress of the deployment to your Smartphone.
+	When your project has been successfully deployed to your Smartphone, the command line window will shut itself down.
+	
+	![](GVR_Run_Bat_File.png)
+
+1. Once you run the .BAT, file a command line window will be displayed that will show you the deployment progress as well as any errors that may have occurred. 
+
+	![](GVR_Deployment_Window.png)
+
+	[region:note]
+	The command line window that is created when running the .BAT file will close down on it's own when it is completed. Closing this window manually could cause issue with your deployment and should be avoided. 
+	[/region]
+
+1. Finally, after the command line window has closed, check your Smartphone for your newly created application which should use the same name used for your UE4 C++ project that it was created from.
+	Launch the application on your and then place your Smartphone into the Gear VR Headset.
+
+	![](GVR_On_Phone.png)
+
+
+%Steps%
diff --git a/Engine/Documentation/Source/Platforms/GearVR/QuickStart/GearVRQuickStart.INT.udn b/Engine/Documentation/Source/Platforms/GearVR/QuickStart/GearVRQuickStart.INT.udn
new file mode 100644
--- /dev/null
+++ b/Engine/Documentation/Source/Platforms/GearVR/QuickStart/GearVRQuickStart.INT.udn
@@ -0,0 +1,29 @@
+Title:Samsung Gear VR UE4 Quick Start
+Crumbs:%ROOT%, Platforms/GearVR
+Description:Instructions on how to quickly setup and deploy projects for use with the Samsung Gear VR Headset.
+Availability:Public
+Navigation:platform
+Platform:Gear VR
+
+
+The Gear VR Quick Start will walk you through how to setup your development PC and Samsung Galaxy Note 4 Smartphone to work with the Samsung Gear VR headset.
+
+![](GVR_Headset_00.png)
+
+
+[OBJECT:TopicButtonList]
+	[PARAM:icon]
+		![](%ROOT%/build_icon.png)(convert:false)
+	[/PARAM]
+	[PARAM:title]
+		Steps
+	[/PARAM]
+	[PARAM:description]
+	[/PARAM]
+	[PARAM:links]
+	* [](Platforms/GearVR/QuickStart/1 "%Platforms/GearVR/QuickStart/1:description%")
+	* [](Platforms/GearVR/QuickStart/2 "%Platforms/GearVR/QuickStart/2:description%")
+	* [](Platforms/GearVR/QuickStart/3 "%Platforms/GearVR/QuickStart/3:description%")
+	[/PARAM]
+[/OBJECT]
+
diff --git a/Engine/Documentation/Source/Platforms/Platforms.INT.udn b/Engine/Documentation/Source/Platforms/Platforms.INT.udn
--- a/Engine/Documentation/Source/Platforms/Platforms.INT.udn
+++ b/Engine/Documentation/Source/Platforms/Platforms.INT.udn
@@ -88,6 +88,9 @@ platforms
 		* [](Platforms/Oculus "%Platforms/Oculus:description%")
 			* [](Platforms/Oculus/BestPractices "%Platforms/Oculus/BestPractices:description%")
 			* [](Platforms/Oculus/QuickStart "%Platforms/Oculus/QuickStart:description%")
+		* [](Platforms/GearVR "%Platforms/GearVR:Description%")
+			* [](Platforms/GearVR/QuickStart "%Platforms/GearVR:Description%")
+			* [](Platforms/GearVR/BestPractices "%Platforms/GearVR:Description%")
 	[/PARAM]
 [/OBJECT]
 
diff --git a/Engine/Documentation/Source/Platforms/VR/VR.INT.udn b/Engine/Documentation/Source/Platforms/VR/VR.INT.udn
--- a/Engine/Documentation/Source/Platforms/VR/VR.INT.udn
+++ b/Engine/Documentation/Source/Platforms/VR/VR.INT.udn
@@ -74,3 +74,18 @@ vr
 		* [](Platforms/Oculus/QuickStart "%Platforms/Oculus/QuickStart:description%")
 	[/PARAM]
 [/OBJECT]
+
+[OBJECT:TopicButtonList]
+	[PARAM:icon]
+		![](%ROOT%/vr_icon.png)(convert:false)
+	[/PARAM]
+	[PARAM:title]
+		[](Platforms/GearVR "%Platforms/GearVR:description%")
+	[/PARAM]
+	[PARAM:description]
+	[/PARAM]
+	[PARAM:links]
+		* [](Platforms/GearVR/BestPractices "%Platforms/GearVR/BestPractices:description%")
+		* [](Platforms/GearVR/QuickStart "%Platforms/GearVR/QuickStart:description%")
+	[/PARAM]
+[/OBJECT]
\ No newline at end of file
