#UE4doc #docs: Made minor edits to new Gear VR docs

diff --git a/Engine/Documentation/Source/Platforms/GearVR/Debugging/GearVRDebugging.INT.udn b/Engine/Documentation/Source/Platforms/GearVR/Debugging/GearVRDebugging.INT.udn
--- a/Engine/Documentation/Source/Platforms/GearVR/Debugging/GearVRDebugging.INT.udn
+++ b/Engine/Documentation/Source/Platforms/GearVR/Debugging/GearVRDebugging.INT.udn
@@ -38,7 +38,7 @@ This can be any kind of application, download or compiled.
 
 	![](T_GearVR_Dev_Mode_03.png)
 
-1. You should now see a scan process happen and when completed you should now see the **Developer Mode** toggle on the screen.
+1. You should now see a scan process happen and when completed you should see the **Developer Mode** toggle on the screen.
 
 	![](T_GearVR_Dev_Mode_04.png)
 
@@ -85,7 +85,7 @@ You can get a full list of available console commands by entering **DumpConsoleC
 1. When you are done debugging or if you just wish to clear the screen of debug information, you can input **Stat None** to clear the screen off all debug information.
 
 	[region:tip]
-	The best option for console debugging with the Gear VR Headset is to use a Bluetooth keyboard. This way you can use the console with out having to remove the phone from the headset.
+	The best option for console debugging with the Gear VR Headset is to use a Bluetooth keyboard. This way you can use the console without having to remove the phone from the headset.
 	[/region]
 
 ### Viewing Gear VR Debug Logs
@@ -94,7 +94,7 @@ Debug logs can provide vital information and insight into why your project is no
 In the following section we will go over what you need to do in order to get debug logs from your Smartphone.
 
 [region:note]
-For much more detailed information on all of various things you can do with [Logcat](http://developer.android.com/tools/help/logcat.html), check out the official Android documentation on [Logcat](http://developer.android.com/tools/help/logcat.html).
+Check out the official Android documentation for much more detailed information on all of various things you can do with [Logcat](http://developer.android.com/tools/help/logcat.html).
 [/region]
 
 1. First make sure that your Smartphone with the application you want to get debug information from is plugged into your PC via USB.
@@ -131,7 +131,7 @@ For much more detailed information on all of various things you can do with [Log
  To see the live output of the log, omit the **>logoutput.txt** part and all log updates will happen in the command line window.
  [/region]
 
-1. To complete the log capture press **CTRL + C** on the keyboard and when completed you should now have a log file named,** logoutput.txt**, that is saved to your **Users** folder in Windows
+1. To complete the log capture press **CTRL + C** on the keyboard and when completed you should now have a log file named,** logoutput.txt**, that is saved to your **Users** folder in Windows.
 	![](T_ADB_Log_Location.png)
 
 	[region:note]
@@ -145,7 +145,7 @@ For much more detailed information on all of various things you can do with [Log
 	![](T_ADB_Log_Info.png)
 
 If your Gear VR Project is not starting up for some reason, using the logcat command is a great way to see exactly where your project is failing and why. 
-It also provides you with detailed information about the problem that can then be posted to the UE4 Forums or Answer Hub so that you can get with the issue.
+It also provides you with detailed information about the problem that can then be posted to the UE4 Forums or Answer Hub so that you can get help with the issue.
 
 ### Wireless Deploying & Debugging
 Getting debug information or deploying to your Smartphone while it is in the Gear VR Headset is impossible due to the only USB port being used by the Gear VR Headset. 
@@ -190,7 +190,7 @@ Please note that in order for this to work you have to have your development PC
 
 	![](T_ADB_Attached_IP.png)
 
-With your Smartphone and PC now connected to one another via the network you can now begin to get debug information from your Smartphone or even deploy projects to your phone with out having to use the USB calbe. 
+With your Smartphone and PC now connected to one another via the network you can now begin to get debug information from your Smartphone or even deploy projects to your phone without having to use the USB cable. 
 However keep in mind that deploying to your Smartphone from your PC over the network is four to five times slower than using the USB cable.
 
 ### Return to Wired Deploying & Debugging
diff --git a/Engine/Documentation/Source/Platforms/GearVR/GearVR.INT.udn b/Engine/Documentation/Source/Platforms/GearVR/GearVR.INT.udn
--- a/Engine/Documentation/Source/Platforms/GearVR/GearVR.INT.udn
+++ b/Engine/Documentation/Source/Platforms/GearVR/GearVR.INT.udn
@@ -18,7 +18,7 @@ In the following sections we will go over everything you need to know to get you
 
 ## Gear VR Headset Pre-requisites 
 
-Before we can get into UE4 Gear VR development we first must make sure that the following prerequisites are full filled as if we do not, we run the risk of our Gear VR projects not working with the Gear VR Headset. 
+Before we can get into UE4 Gear VR development we first must make sure that the following prerequisites are fulfilled as if we do not, we run the risk of our Gear VR projects not working with the Gear VR Headset. 
 
 * Make sure you have the USB cord used for charging your Samsung Galaxy Note 4 Smartphone available and plugged into a USB port on the PC you are developing on. 
 The USB cord is of vital importance as this is how you transfer data between your development PC and your Smartphone.
@@ -35,7 +35,7 @@ The USB cord is of vital importance as this is how you transfer data between you
 
 * Install the **Tegra Android Development Pack** which can be found in **Engine\Extras\Android**.  
 
-* Check to see if [Developer Mode](https://docs.unrealengine.com/latest/INT/Platforms/Android/GettingStarted/2/index.html ) has been enabled on your Samsung Galaxy Note 4 Smartphone and if not enable it. 
+* Check to see if [Developer Mode](https://docs.unrealengine.com/latest/INT/Platforms/Android/GettingStarted/2/index.html ) has been enabled on your Samsung Galaxy Note 4 Smartphone and if not, enable it. 
 
 * Also make sure to check that USB Debugging has been enabled to ensure that you can load data onto your Smartphone from your PC.
 	
@@ -62,7 +62,7 @@ Then place that folder in **Engine/Build/Android/Java**.
 	![](GVR_OSIG_File_Location.png)
 
 It is very important that you make sure all items listed in the Gear VR Prerequisites section are completed before moving on to the next section. 
-Skipping or failing to do a step could result in your project not working with the Gear VR Headset. 
+Skipping a step could result in your project not working with the Gear VR Headset. 
 Once everything from above has been completed you can now move on to the Gear VR Headset Development Quick Start.
 
 
diff --git a/Engine/Documentation/Source/Platforms/GearVR/QuickStart/3/GearVRQuickStart_3.INT.udn b/Engine/Documentation/Source/Platforms/GearVR/QuickStart/3/GearVRQuickStart_3.INT.udn
--- a/Engine/Documentation/Source/Platforms/GearVR/QuickStart/3/GearVRQuickStart_3.INT.udn
+++ b/Engine/Documentation/Source/Platforms/GearVR/QuickStart/3/GearVRQuickStart_3.INT.udn
@@ -45,16 +45,16 @@ Platform:Gear VR
 	
 	![](GVR_Run_Bat_File.png)
 
-1. Once you run the .BAT, file a command line window will be displayed that will show you the deployment progress as well as any errors that may have occurred. 
+1. Once you run the .BAT file, a command line window will be displayed that will show you the deployment progress as well as any errors that may have occurred. 
 
 	![](GVR_Deployment_Window.png)
 
 	[region:note]
-	The command line window that is created when running the .BAT file will close down on it's own when it is completed. Closing this window manually could cause issue with your deployment and should be avoided. 
+	The command line window that is created when running the .BAT file will close down on it's own when it is completed. Closing this window manually could cause an issue with your deployment and should be avoided. 
 	[/region]
 
 1. Finally, after the command line window has closed, check your Smartphone for your newly created application which should use the same name used for your UE4 C++ project that it was created from.
-	Launch the application on your and then place your Smartphone into the Gear VR Headset.
+	Launch the application and then place your Smartphone into the Gear VR Headset.
 
 	![](GVR_On_Phone.png)
 
