#UE4doc #docs: Made minor edits

diff --git a/Engine/Documentation/Source/Engine/Components/Paper2D/Paper2DComponents.INT.udn b/Engine/Documentation/Source/Engine/Components/Paper2D/Paper2DComponents.INT.udn
--- a/Engine/Documentation/Source/Engine/Components/Paper2D/Paper2DComponents.INT.udn
+++ b/Engine/Documentation/Source/Engine/Components/Paper2D/Paper2DComponents.INT.udn
@@ -14,7 +14,7 @@ Paper 2D Components are used to add (or modify) 2D **Sprites** or **Flipbooks**
 A **PaperSpriteComponent** handles rendering and collision for a single instance of a `UPaperSprite` asset. 
 This Component is created when you drag a Sprite asset from the **Content Browser** into a **Blueprint** or contained inside the Actor created when you drag one into a level. 
 
-An example of using this type of Component may come in the form of Sprite assets used to construct a level, for instance ledges or platforms, ladders, ramps etc. Placing those Sprite assets in your level creates a **PaperSpriteActor** that uses and instanced **PaperSpriteComponent** (or unique copy) of the selected Sprite asset. 
+An example of using this type of Component may come in the form of Sprite assets used to construct a level, for instance ledges or platforms, ladders, ramps etc. Placing those Sprite assets in your level creates a **PaperSpriteActor** that uses an instanced **PaperSpriteComponent** (or unique copy) of the selected Sprite asset. 
 
 For information on creating Sprites in Unreal Engine 4, see the [](Engine\Paper2D\Sprites\Creation) documentation. 
 
diff --git a/Engine/Documentation/Source/Engine/Components/Rendering/RenderingComponents.INT.udn b/Engine/Documentation/Source/Engine/Components/Rendering/RenderingComponents.INT.udn
--- a/Engine/Documentation/Source/Engine/Components/Rendering/RenderingComponents.INT.udn
+++ b/Engine/Documentation/Source/Engine/Components/Rendering/RenderingComponents.INT.udn
@@ -46,7 +46,7 @@ The **MaterialBillboardComponent** is a 2D Material that will be rendered always
 
 A **CableComponent** allows you attach two Components together while rendering a cable between them. On the cable itself, you can assign a Material and define parameters that affect the way the cable is displayed. 
 
-An example use case for this type of Component is if you were in need of a rope or something that a player character could swing on. By attaching one point to a fixed location and attaching the other to the player character through script, the character could swing on the rope. Below use have a first person character that can shoot a cable out from their gun and attach it to the impact location, creating a "grappling gun" of sorts. 
+An example use case for this type of Component is if you were in need of a rope or something that a player character could swing on. By attaching one point to a fixed location and attaching the other to the player character through script, the character could swing on the rope. Below you have a first person character that can shoot a cable out from their gun and attach it to the impact location, creating a "grappling gun" of sorts. 
 
 ![](cable.png)(w:440)
 
diff --git a/Engine/Documentation/Source/Engine/Components/Shapes/ShapesComponents.INT.udn b/Engine/Documentation/Source/Engine/Components/Shapes/ShapesComponents.INT.udn
--- a/Engine/Documentation/Source/Engine/Components/Shapes/ShapesComponents.INT.udn
+++ b/Engine/Documentation/Source/Engine/Components/Shapes/ShapesComponents.INT.udn
@@ -33,7 +33,7 @@ The **CapsuleComponent** is a capsule shape which can also be used as simple col
 
 ![](capsule_component.png)
 
-Above a **CapsuleComponent** is automatically included when creating a new **Character Blueprint** and provides the collision settings for the character, keeping them from intersecting with objects in the world. The **CapsuleComponent** can also **Generate Overlap Events** or **Generate Hit Events** allow you to provide script for when either of those events occur. 
+Above a **CapsuleComponent** is automatically included when creating a new **Character Blueprint** and provides the collision settings for the character, keeping them from intersecting with objects in the world. The **CapsuleComponent** can also **Generate Overlap Events** or **Generate Hit Events** to allow you to provide scripts for when either of those events occur. 
 
 
 ## Sphere Component
