TypoCorrection

diff --git a/Engine/Documentation/Source/Engine/Rendering/Materials/HowTo/Fresnel/Fresnel.INT.udn b/Engine/Documentation/Source/Engine/Rendering/Materials/HowTo/Fresnel/Fresnel.INT.udn
--- a/Engine/Documentation/Source/Engine/Rendering/Materials/HowTo/Fresnel/Fresnel.INT.udn
+++ b/Engine/Documentation/Source/Engine/Rendering/Materials/HowTo/Fresnel/Fresnel.INT.udn
@@ -1,4 +1,4 @@
-Availability: Public
+﻿Availability: Public
 Title:Material - How To Use Fresnel in your Materials
 Crumbs: %ROOT%, Engine, Engine/Rendering/Materials
 Description:A How To Guide for using the Fresnel Material node. 
@@ -80,7 +80,7 @@ This node should only be used if the Fresnel that is under **Utility** does not
 | **Camera Vector(V3)** |This allows you to effect the Camera that is used to determine the Dot Product between the Camera and the surface normal. |
 | **Invert Fresnel(B)** | This will allow you to invert the effect of the Fresnel. This is useful if you want the Fresnel to only be applied in the center or at the edges. |
 | **Power(S)** | This allows you to control the power of the Fresnel effect. Entering bigger numbers will tighten up the areas affected. Smaller numbers will allow more areas to be affected by the Fresnel.  |
-| **Use Cheap Contrast(B)** | This will allow the Fresnel to muse a cheaper method for calculating contrast between lights and darks. |
+| **Use Cheap Contrast(B)** | This will allow the Fresnel to use a cheaper method for calculating contrast between lights and darks. |
 | **Cheap Contrast Dark (S)** | This controls how dark the Fresnel will be. This will only have an affect if **Use Cheap Contrast** is enabled.  |
 | **Cheap Contrast Bright (S)** | This controls how bright the Fresnel effect will be. This will only have an affect if **Use Cheap Contrast** is enabled.|
 
