#UE4doc #docs: Changed "MyActor.h"  to "MyPawn.h" to correctly correspond with tutorial

diff --git a/Engine/Documentation/Source/Programming/Tutorials/PlayerInput/3/PlayerInput_3.INT.udn b/Engine/Documentation/Source/Programming/Tutorials/PlayerInput/3/PlayerInput_3.INT.udn
--- a/Engine/Documentation/Source/Programming/Tutorials/PlayerInput/3/PlayerInput_3.INT.udn
+++ b/Engine/Documentation/Source/Programming/Tutorials/PlayerInput/3/PlayerInput_3.INT.udn
@@ -21,7 +21,7 @@ Description:Binding input to C++ code
 
 %Steps%
 
-1. In **Visual Studio**, open "MyActor.h" and add the following code to the bottom of MyPawn's class definition:
+1. In **Visual Studio**, open "MyPawn.h" and add the following code to the bottom of MyPawn's class definition:
 
 		//Input functions
 		void Move_XAxis(float AxisValue);
