#UE4doc #docs Switched around hotkey commands as "C" and "Shift+C" were reveresed

diff --git a/Engine/Documentation/Source/Engine/Rendering/Materials/Editor/Interface/MaterialEditorInterface.INT.udn b/Engine/Documentation/Source/Engine/Rendering/Materials/Editor/Interface/MaterialEditorInterface.INT.udn
--- a/Engine/Documentation/Source/Engine/Rendering/Materials/Editor/Interface/MaterialEditorInterface.INT.udn
+++ b/Engine/Documentation/Source/Engine/Rendering/Materials/Editor/Interface/MaterialEditorInterface.INT.udn
@@ -212,7 +212,7 @@ drop in the node.  The hotkeys are as follows:
 | --- | --- |
 | **A** | Add |
 | **B** | BumpOffset |
-| **C** | ComponentMask |
+| **C** | Comment |
 | **D** | Divide |
 | **E** | Power |
 | **F** | MaterialFunctionCall |
@@ -231,5 +231,5 @@ drop in the node.  The hotkeys are as follows:
 | **2** | Constant2Vector |
 | **3** | Constant3Vector |
 | **4** | Constant4Vector |
-| **Shift + C**  | Comment |
+| **Shift + C**  | ComponentMask |
 
