TypoCorrection

diff --git a/Engine/Documentation/Source/Engine/Rendering/Materials/HowTo/Making_Parameters/Making_Parameters.INT.udn b/Engine/Documentation/Source/Engine/Rendering/Materials/HowTo/Making_Parameters/Making_Parameters.INT.udn
--- a/Engine/Documentation/Source/Engine/Rendering/Materials/HowTo/Making_Parameters/Making_Parameters.INT.udn
+++ b/Engine/Documentation/Source/Engine/Rendering/Materials/HowTo/Making_Parameters/Making_Parameters.INT.udn
@@ -1,4 +1,4 @@
-Availability: Public
+﻿Availability: Public
 Title:Material Editor - How To make Material Parameters
 Crumbs: %ROOT%, Engine, Engine/Rendering/Materials, Engine/Rendering/Materials/HowTo
 Description:A How To Guide for Setting Up and using Material Parameters
@@ -35,7 +35,7 @@ The simplest way to make a Material Parameter out of a Material Expression is to
 Not all Material Expressions can be converted into Material Parameters. The best way to see if an Material Expression can be converted into a Material Parameter is to **Right-Click** on it and check to see if the **Convert to Parameter** option is available. If you do not see the **Convert to Parameter** option, then the Material Expression cannot be converted.
 [/region]
 
-### Adding Parameters form the Palette
+### Adding Parameters from the Palette
 
 You can also add a Material Parameter from the **Palette** search menu. 
 To do this, first search for the word **Parameters** in the Palette search menu.
