In-progress checking for Tile Sets / Tile Maps page (needs section on Tile Map Collision still which I'm still working on).

diff --git a/Engine/Documentation/Source/Engine/Paper2D/Paper2D.INT.udn b/Engine/Documentation/Source/Engine/Paper2D/Paper2D.INT.udn
--- a/Engine/Documentation/Source/Engine/Paper2D/Paper2D.INT.udn
+++ b/Engine/Documentation/Source/Engine/Paper2D/Paper2D.INT.udn
@@ -26,6 +26,7 @@ Refer to each section below for more information and the **Paper 2D** map inside
 		* [](Engine\Paper2D\Flipbooks"%Engine\Paper2D\Flipbooks:description%")
 			* [](Engine/Paper2D/Flipbooks/Editor "%Engine/Paper2D/Flipbooks/Editor:description%")
 			* [](Engine\Paper2D\Flipbooks\Components "%Engine\Paper2D\Flipbooks\Components:description%")
+		* [](Engine\Paper2D\TileMaps"%Engine\Paper2D\TileMaps:description%")
 		* [](Resources/ContentExamples/Paper2D "%Resources/ContentExamples/Paper2D:description%")
 		* [Paper 2D Video Tutorials](https://docs.unrealengine.com/latest/INT/Videos/index.html?category=2D)
 	[/PARAM]
diff --git a/Engine/Documentation/Source/Engine/Paper2D/TileMaps/Paper2DTilesMaps.INT.udn b/Engine/Documentation/Source/Engine/Paper2D/TileMaps/Paper2DTilesMaps.INT.udn
new file mode 100644
--- /dev/null
+++ b/Engine/Documentation/Source/Engine/Paper2D/TileMaps/Paper2DTilesMaps.INT.udn
@@ -0,0 +1,233 @@
+Title:Paper 2D Tile Sets / Tile Maps
+Crumbs: %ROOT%, Engine, Engine/Paper2D
+Description:Overview of how to create Tile Sets and Tile Maps for use within Paper 2D. 
+Availability:Public
+Related: Engine\Paper2D\Sprites\Collision
+Related: Engine\Paper2D\Flipbooks
+Related: Engine/Rendering/Materials
+Related: Engine/Content/Types/Textures
+
+[REGION:warning]
+The features on this page are currently in the **experimental stage** of development. Aspects may not yet be fully supported, use at your own risk. 
+[/REGION]
+
+
+[TOC (start:2 end:2)]
+
+## Overview
+
+**Tile Sets** and **Tile Maps** within **Paper 2D**, provide a quick and easy way to layout the framework or "general layout" of your 2D levels. 
+By creating and using a Tile Set (a collection of tiles pulled from a texture) with a Tile map (a 2D grid with a defined width and height in tiles), you can select various tiles to "paint" onto the Tile Map which can be used for your level layout. 
+You can also paint tiles on multiple layers, each of which can specify which tile should appear in each cell of the map for that particular layer. 
+
+
+## Creating Tile Sets
+
+You can create Tile Sets inside the **Content Browser**, similar to other assets. They can be created from scratch as blank assets or generated from other existing assets. 
+
+### Blank Tile Set
+
+**To create a new blank Tile Set asset:**
+
+1. Click the ![](addNewButton.png)(h:24) button in the **Content Browser**, then under _Miscellaneous_, select **Tile Set**.
+	
+	![](TileSetButton.png)
+
+1. Enter in a name for the new Tile Set asset. 
+
+	![](PaperTitleSet.png)
+
+1. **Double-click** on the asset to open it up in the **Tile Set Editor** window. 
+
+	![](PaperTitleSetEditor.png)(w:640)
+
+1. Specify a **Tile Sheet** texture asset to use inside the **Details** panel. 
+
+	![](TileSetData.png)
+
+	
+### Tile Set from Texture
+
+**To create a Tile Set from an existing Texture:**
+
+1. In the **Content Browser**, **Right-click** on a Texture asset, then under _Sprite Actions_ choose **Create Tile Set**.
+
+	![](RightClickMethod.png)
+
+1. The Tile Set will automatically be created/named based on the Texture you supplied. 
+
+	![](AutoCreate.png)
+
+	[REGION:note]
+	You can rename the asset by selecting it and pressing **F2** or **Right-clicking** on it and choosing **Rename**. 
+	[/REGION]
+
+1. **Double-click** on the asset to open it up in the **Tile Set Editor** window. 
+
+	![](CreatedTitleSet.png)(w:640)
+
+	
+## Tile Set Editor
+
+Once you have created a new Tile Set asset, **Double-clicking** on the asset will open it in the **Tile Set Editor**. 
+
+![](TileSetEditorWindow.png)(w:800)
+
+Above in the **Viewport** window (1), you will see the Tile Sheet you've specified to use with your Tile Set. In the **Tile Editor** window (2), the currently selected tile will be displayed. In the **Details** panel (3), you can set various settings that will affect the Tile Set or the individual Sprites that make up the set (see _Configuring Your Tile Set_ below). 
+
+### Configuring Your Tile Set
+
+You will more than likely need to configure your Tile Set to ensure that you are accurately selecting each tile within your sheet. 
+
+Inside the **Viewport**, **Left-clicking** on your Tile Sheet will display a white tile selection box indicating the currently selected tile. 
+
+![](DefaultSettings.png)
+
+[REGION:tip]
+**Right Mouse Button** drag inside the Viewport to pan and use the **Mouse Wheel** to zoom in/out on your Tile Sheet. 
+[/REGION]
+
+Above, the default **Tile Width** and **Tile Height** (in pixels) of 32 is being used represented by the white selection box. Unfortunately our tiles are larger than the default values and we want our Tile Width and Height to encompass our entire tile (represented by the yellow box outline). 
+
+Inside the **Tile Editor** preview window the currently selected Tile and its associated Tile Number (see below) are displayed. Just below it, inside the **Details** panel, you can adjust the Tile Width and Tile Height so that it properly encompasses your entire tile. 
+
+![](CurrentlySelectedTile.png)
+
+In our example, we've increased the Tile Width and Tile Height to 128. 
+
+[REGION:lightbox]
+[![](UpdatedTileSheet.png)(w:800)](UpdatedTileSheet.png)
+[/REGION]
+
+_Click for full view._
+
+In the Viewport, the white tile selection box now fully highlights an individual tile on the Tile Sheet (you may need to experiment with different values depending on how your Tile Sheet is set up). The Tile Editor window also updates to show a single tile from our Tile Sheet that is selected. 
+
+In addition to setting Tile Width and Tile Height, you can use the following settings to further adjust your Tile Set from the **Details** panel. 
+
+
+| Setting | Description |
+| -------------- | ----------- |
+| **Tile Width** | The width of a single tile (in pixels). |
+| **Tile Height** | The height of a single tile (in pixels). |
+| **Margin** | The amount of padding around the perimeter of the tile sheet (in pixels).  |
+| **Spacing** | The amount of padding between tiles in the tile sheet (in pixels). |
+| **Drawing Offset** | The drawing offset for tiles from this set (in pixels). |
+| **Background Color** | The background color displayed in the tile set viewer. |
+
+[REGION:note]
+There are additional settings inside the **Details** panel for per tile adjustments, however these are considered **experimental** in the current release.
+[/REGION]
+
+
+
+## Creating Tile Maps
+
+Tile Maps are also created inside the **Content Browser** and edited inside the **Tile Map Editor**. 
+
+### New Tile Map
+
+**To create a new Tile Map asset:**
+
+1. Click the ![](addNewButton.png)(h:24) button in the **Content Browser**, then under _Miscellaneous_, select **Tile Map**.
+	
+	![](NewTileMap.png)
+
+1. Enter in a name for the new Tile Map asset. 
+
+	![](NewPaperTileMap.png)
+
+1. **Double-click** on the asset to open it up in the **Tile Map Editor** window. 
+
+	![](TileMapEditorWindow.png)(w:640)
+
+
+## Tile Map Editor
+
+After opening a new Tile Map asset, before you can work within the Tile Map editor you will need to assign an **Active Tile Set** to work with. 
+
+1. You can specify the **Active Tile Set** to use inside the **Toolbox** window by clicking the **Selection** box. 
+
+	![](SelectTileSetWindow.png)
+
+	Then selecting the Tile Set you wish to currently work with. 
+
+	![](SelectTileSet.png)
+
+	Any Tile Sets that you have created will be displayed in the window above, allowing you to switch between the "Active Tile Set". 
+
+1. Once you have selected a Tile Set, the viewer will become updated with your selection. 
+
+	![](TileMapReadyToEdit.png)(w:640)
+
+The three windows that comprise the Tile Map Editor are outlined below. 
+
+![](TileMapEditorOverview.png)(w:800)
+
+Above in the **Toolbox** window (1), you can switch between different Tile Sets to use for the creation of your Tile Map. The Tools used to **Paint**, **Erase** or **Fill** tiles on the map are also presented (in addition to options for flipping or rotating the currently selected tile). The bottom section of this window contains the Viewer, which uses the same controls for navigation as the Tile Set viewer (**Right Mouse Button** drag to pan and **Mouse Wheel** to zoom in/out). 
+
+The center window, the **Viewport** (2), is used for filling out the Tile Map layout. Using the tools from the Toolbox, you can fill in tiles on the map to create the desired effect by painting, erasing or filling the tiles that comprise your Tile Map. This is where you can "paint" how you want your map to appear in-game. 
+
+In the **Details** panel (3), you can configure how many layers your Tile Map consists of as well as adjust various settings such as the width/height of the Tile Map itself, the width/height of the Tiles that comprise the map, the scaling factor between pixels and Unreal Units, etc. 
+
+
+## Example Tile Map Workflow
+
+Below is an example of a typical workflow for creating a Tile Map. 
+
+1.	Open your Tile Map asset. 
+
+1.	In the **Details** panel, set the **Map Width** and **Map Height** for your Tile Map. 
+
+	![](SetMapSize.png)
+
+	This is essentially the size you want your map or level to be. 
+
+1.	Also set the **Tile Width** and **Tile Height** of the individual tiles in your map. 
+
+	![](TileSizesInMap.png)
+
+	These are usually going to be the same size you defined in your Tile Set. 
+
+1.	Assign an **Active Tile Set** to work with for the construction of your Tile Map. 
+
+1.	Also add some additional layers to work with (one for background, one for foreground and one for the playable level geometry). 
+
+	![](ReadyToPaint.png)
+
+	This is optional, but usually good to have a layer that has your background elements and another for any foreground elements and having your level geometry (the ground, boxes, platforms, etc.) that the player will walk around on placed on a separate layer. 
+
+1.	Select a tile in the viewer, then in the map **Left-click** to paint it onto the map. 
+
+	![](BeginPainting.png)
+
+	You can also select multiple tiles at once and paint multiple tiles in the map by **Left-click** and dragging over tiles in the viewer. 
+
+	![](MultipleTiles.png)
+
+	In the map, you can then paint with the selected tiles. 
+
+	![](InMapMultipleTiles.png)
+
+1.	Continue Painting, Filling or Erasing Tiles on our map to reach your desired layout. 
+
+	![](ExampleMap.png)(w:800)
+
+1.	Drag the Tile Map asset into your level and scale or move based on your needs. 
+
+	![](ExampleGameImage.png)(w:800)
+
+
+	
+
+
+
+
+
+
+
+
+
+
+
+
