TypoCorrection

diff --git a/Engine/Documentation/Source/Engine/Paper2D/Paper2D.INT.udn b/Engine/Documentation/Source/Engine/Paper2D/Paper2D.INT.udn
--- a/Engine/Documentation/Source/Engine/Paper2D/Paper2D.INT.udn
+++ b/Engine/Documentation/Source/Engine/Paper2D/Paper2D.INT.udn
@@ -1,7 +1,7 @@
+﻿Availability:Public
 Title: Paper 2D
 Crumbs: %ROOT%, Engine
 Description:Paper 2D is a sprite-based system for creating 2D and 2D/3D hybrid games entirely within Unreal Engine 4.
-Availability:Public
 
 [REGION:banner]
 ![](paper2d_banner.png)
