TypoCorrection

diff --git a/Engine/Documentation/Source/Engine/Rendering/Materials/HowTo/Subsurface_Profile/Subsurface_Profile.INT.udn b/Engine/Documentation/Source/Engine/Rendering/Materials/HowTo/Subsurface_Profile/Subsurface_Profile.INT.udn
--- a/Engine/Documentation/Source/Engine/Rendering/Materials/HowTo/Subsurface_Profile/Subsurface_Profile.INT.udn
+++ b/Engine/Documentation/Source/Engine/Rendering/Materials/HowTo/Subsurface_Profile/Subsurface_Profile.INT.udn
@@ -1,4 +1,4 @@
-Availability:Public
+﻿Availability:Public
 Title: Material Editor - How To use a Subsurface Profile in your Materials 
 Crumbs: %ROOT%, Engine, Engine/Rendering/Materials
 Description: Setting up Materials to use Subsurface Profiles as they offer a better way to render skin and wax like surfaces.
@@ -113,7 +113,7 @@ When completed, your **Content Browser** should look something like this.
 	![](SSP_SSP_In_Material_Instance.png)
 
 	[region:note]
-	If you have already supplied a Subsurface Profile in the Master Material, then you do not need to override it in the Material Instance. This is only if you wish you use a different Subsurface Profile that what is already being used. 
+	If you have already supplied a Subsurface Profile in the Master Material, then you do not need to override it in the Material Instance. This is only if you wish you use a different Subsurface Profile than what is already being used. 
 	[/region]
 
 [/EXCERPT:USSP]
