#UE4 Docs, #Doc

Updated to set to Docs only (not ready for public release).

diff --git a/Engine/Documentation/Source/Engine/Paper2D/Paper2D.INT.udn b/Engine/Documentation/Source/Engine/Paper2D/Paper2D.INT.udn
--- a/Engine/Documentation/Source/Engine/Paper2D/Paper2D.INT.udn
+++ b/Engine/Documentation/Source/Engine/Paper2D/Paper2D.INT.udn
@@ -26,7 +26,9 @@ Refer to each section below for more information and the **Paper 2D** map inside
 		* [](Engine\Paper2D\Flipbooks"%Engine\Paper2D\Flipbooks:description%")
 			* [](Engine/Paper2D/Flipbooks/Editor "%Engine/Paper2D/Flipbooks/Editor:description%")
 			* [](Engine\Paper2D\Flipbooks\Components "%Engine\Paper2D\Flipbooks\Components:description%")
+		<!-- Not yet ready for release.
 		* [](Engine\Paper2D\TileMaps"%Engine\Paper2D\TileMaps:description%")
+		-->
 		* [](Resources/ContentExamples/Paper2D "%Resources/ContentExamples/Paper2D:description%")
 		* [Paper 2D Video Tutorials](https://docs.unrealengine.com/latest/INT/Videos/index.html?category=2D)
 	[/PARAM]
diff --git a/Engine/Documentation/Source/Engine/Paper2D/TileMaps/Paper2DTilesMaps.INT.udn b/Engine/Documentation/Source/Engine/Paper2D/TileMaps/Paper2DTilesMaps.INT.udn
--- a/Engine/Documentation/Source/Engine/Paper2D/TileMaps/Paper2DTilesMaps.INT.udn
+++ b/Engine/Documentation/Source/Engine/Paper2D/TileMaps/Paper2DTilesMaps.INT.udn
@@ -1,7 +1,7 @@
 Title:Paper 2D Tile Sets / Tile Maps
 Crumbs: %ROOT%, Engine, Engine/Paper2D
 Description:Overview of how to create Tile Sets and Tile Maps for use within Paper 2D. 
-Availability:Public
+Availability:Docs
 Related: Engine\Paper2D\Sprites\Collision
 Related: Engine\Paper2D\Flipbooks
 Related: Engine/Rendering/Materials
