Added Paper2D Import Options page, updated several of the other Paper 2D pages.

diff --git a/Engine/Documentation/Source/Engine/Paper2D/Flipbooks/Components/Blueprints/Components_BP.INT.udn b/Engine/Documentation/Source/Engine/Paper2D/Flipbooks/Components/Blueprints/Components_BP.INT.udn
--- a/Engine/Documentation/Source/Engine/Paper2D/Flipbooks/Components/Blueprints/Components_BP.INT.udn
+++ b/Engine/Documentation/Source/Engine/Paper2D/Flipbooks/Components/Blueprints/Components_BP.INT.udn
@@ -23,13 +23,9 @@ This page will show you how to add, access and modify Flipbook Components in a B
 
 ## Adding a Flipbook Component
 
-You can add a Flipbook Component to a Blueprint from the **Components** tab of a Blueprint. 
+You can add a Flipbook Component to a Blueprint from the **Components** window of a Blueprint. 
 
-1. With a Blueprint open, in the upper-right corner, click the **Components** tab. 
-
-	![](Components0.png)
-
-1. From the **Components** tab, click the **Add Component** button inside the **Components** window. 
+1. Inside a Blueprint, click the **Add Component** button inside the **Components** window and search for and add a **Paper Flipbook**. 
 
 	![](Components1.png)	
 
@@ -39,15 +35,15 @@ You can add a Flipbook Component to a Blueprint from the **Components** tab of a
 
 	You can then begin adjusting the **Properties** of the Flipbook Component inside the **Details** panel. 
 
-You can also add Flipbook Components to Blueprints by using the **Replace With Composited Blueprint** option. 
+You can also add Flipbook Components to Blueprints by using **Blueprint/Add Script** option. 
 
 1.	With a **Flipbook** placed inside your level and selected, head to the **Details** panel. 
 
-1.	Towards the bottom of the **Details** panel, select the **Replace With Composited Blueprint** option. 
+1.	Click the **Blueprint/Add Script** button. 
 
 	![](Components3.png)
 
-	This will create a new Blueprint and assign the Flipbook selected as a component of it automatically. 
+	This will create a new Blueprint (prompting you to enter a name) and assign the Flipbook selected as a component of it automatically. 
 
 ## Flipbook Component Properties
 
@@ -73,15 +69,19 @@ And finally, set the color for the Sprite to use with the **Sprite Color** optio
 
 ### Collision Properties
 
-As of this time since Flipbooks are still not quite production ready and are in the Early Access state and as such do not currently support any form of collision. You can use a **Box**, **Capsule** or **Sphere** Component in conjunction with a Flipbook to provide the associated collision results you need (in yellow below). 
+Although it is possible to apply collision directly to a Flipbook (see [](Engine\Paper2D\Flipbooks\Editor) for more information), it is recommended that a collision shape be used instead for performance reasons. For example, you can use a **Box**, **Capsule** or **Sphere** Component in conjunction with a Flipbook to provide the associated collision results you need. 
+
+For example, clicking the **Add Component** button inside your Blueprint and search for **Collision**: 
 
 ![](Collision1.png)
 
+You can now apply different shapes to use as the collision for your Flipbook rather than specifying it on the Flipbook itself. 
+
 An example of such is indicated below. 
 
-![](Collision2.png)(w:720)
+![](Collision2.png)(w:820)
 
-Above we've parented our Flipbook to a Sphere (called **Collision**), adjusted the size of the collision with **Sphere Radius**, enabled Physics with the **Simulate Physics** option and used the **Locked Axis** option to lock physics simulation to only the Y axis, then finally set its collision to **BlockAll** so that it collides with everything in the level. With this setup, we have a circular Actor that can be placed in our level that reacts to other Actors in the level and simulates physics appropriately along one axis. 
+Above we've created a simple blinking box Blueprint and added a Box Collision component to provide the collision for the box by changing the **Box Extent** values to encapsulate our box and changing the **Collision Presets** to **Block All** to provide the actual collision. 
 
 ## Working with Flipbook Components
 
diff --git a/Engine/Documentation/Source/Engine/Paper2D/Flipbooks/Components/FlipbookComponents.INT.udn b/Engine/Documentation/Source/Engine/Paper2D/Flipbooks/Components/FlipbookComponents.INT.udn
--- a/Engine/Documentation/Source/Engine/Paper2D/Flipbooks/Components/FlipbookComponents.INT.udn
+++ b/Engine/Documentation/Source/Engine/Paper2D/Flipbooks/Components/FlipbookComponents.INT.udn
@@ -7,7 +7,7 @@ Description:Information on working with Flipbook Components in Blueprints or C++
 
 **Flipbook Components** are regular primitive components, and can be posed arbitrarily in 3D, attached to other components, or have other components attached to them. Each Flipbook Component instance can specify a custom color that will be passed down to the Flipbook Material as a Vertex Color. They can also have a custom Material specified that will override the default Material defined in the Flipbook. 
 
-You can change the current Flipbook asset by calling **SetFlipbook**, but note that this will only work if the **Mobility** property is set to **Moveable** (or if it is called during the construction of the Actor). You can also control play rate, play direction, looping, etc…, with various other methods on the component. 
+You can change the current Flipbook asset by calling **SetFlipbook**, but note that this will only work if the **Mobility** property is set to **Moveable** (or if it is called during the construction of the Actor). You can also control play rate, play direction, looping, etc. with various other methods on the component. 
 
 [REGION:note]
 While C++ documentation for working with Flipbook Components is still in development, you can refer to [UPaperFlipbookComponent](https://docs.unrealengine.com/latest/INT/API/Plugins/Paper2D/UPaperFlipbookComponent/index.html) for more information on using them in your code until more detailed documentation comes online. 
diff --git a/Engine/Documentation/Source/Engine/Paper2D/Flipbooks/Editor/FlipbookEditor.INT.udn b/Engine/Documentation/Source/Engine/Paper2D/Flipbooks/Editor/FlipbookEditor.INT.udn
--- a/Engine/Documentation/Source/Engine/Paper2D/Flipbooks/Editor/FlipbookEditor.INT.udn
+++ b/Engine/Documentation/Source/Engine/Paper2D/Flipbooks/Editor/FlipbookEditor.INT.udn
@@ -20,7 +20,7 @@ After you've created a Flipbook, you can then open it for editing in the **Flipb
 
 1. **Double-click** on a Flipbook inside the **Content Browser** (or **Right-click** on a Flipbook and select **Edit**). 
 
-	![](EditFlipbook.png)(w:440)
+	![](EditFlipbook.png)
 
 1.	This will open the **Flipbook Editor** window. 
 
@@ -54,10 +54,6 @@ Each of the **Tool Bar** options are described in the table below:
 |![](Toolbar5.png)(w:60)| Toggles the display of the simplified collision mesh of the static mesh, if one has been assigned (**Alt+C**). |
 |![](Toolbar6.png)(w:60)| Displays the pivot location of the sprite in the viewport. |
 
-[REGION:note]
-Flipbooks do not currently support any form of collision. You can use a Capsule or Box Component in conjunction with a Flipbook to provide the associated collision shape.
-[/REGION]
-
 ### Viewport
 
 ![](FlipbookViewport.png)
@@ -108,6 +104,27 @@ The **Details** panel inside the **Flipbook Editor** is where you determine the
 
 ![](DetailsPanel.png)
 
+
+#### Collision Source
+
+![](CollisionSource.png)
+
+While it is possible to set collision on a **Flipbook** by using the **Collision Source** dropdown to either: 
+
+* Have no collision applied to the Flipbook. 
+* Have non-animated collision only applied to the first frame of the animation. 
+* Have the collision changed each frame of the animation (very expensive in terms of performance and not recommended). 
+
+In many cases it may make more sense to attach a Collision Component (for example, a Box Component) to a socket on the Flipbook Component rather than directly generating collision from source frames, especially if you want some collision regions to overlap while others block. 
+
+Character Movement still relies on the collision capsule, in most cases you want to use a different collision channel for the Flipbook than BlockAll to prevent the capsule from getting stuck on the Flipbook (move/be blocked based on the capsule, and test for overlap or collision with projectiles based on the Flipbook). 
+
+[REGION:note]
+See the section on Collision in the [](Engine\Paper2D\Flipbooks\Components) documentation for more information. 
+[/REGION]
+
+
+
 ## Editing Flipbooks
 
 If you are working with a blank Flipbook, you can add sprites to the Flipbook in a couple of different ways. 
@@ -118,7 +135,7 @@ You can **Drag-drop** a Sprite (or Flipbook) asset onto the Timeline from the **
 
 Which will add the Sprite(s) to the Timeline as well as update the **Key Frames** section inside the **Details** panel. 
 
-![](Adding2.png)
+![](Adding2.png)(w:720)
 
 You can also click the **+** sign in the **Key Frames** section of the **Details** panel. 
 
diff --git a/Engine/Documentation/Source/Engine/Paper2D/Flipbooks/Paper2DFlipbooks.INT.udn b/Engine/Documentation/Source/Engine/Paper2D/Flipbooks/Paper2DFlipbooks.INT.udn
--- a/Engine/Documentation/Source/Engine/Paper2D/Flipbooks/Paper2DFlipbooks.INT.udn
+++ b/Engine/Documentation/Source/Engine/Paper2D/Flipbooks/Paper2DFlipbooks.INT.udn
@@ -66,7 +66,7 @@ A **Frames Per Second** option determines how fast the frames will be displayed,
 There are a couple of ways you can go about creating a Flipbook. You can either create a blank Flipbook that you populate with sprites yourself or you can auto-generate a Flipbook based off a series of selected sprites. 
 
 [REGION:note]
-You can import a JSON formatted sprite sheet description, which will import associated textures and create sprites and a Flipbook for the described frames as well. 
+You can import a JSON formatted sprite sheet description, which will import associated textures and create sprites and a Flipbook for the described frames as well. See [](Engine\Paper2D\Importing) for more information. 
 [/REGION]
 
 ### Blank Flipbooks
@@ -75,21 +75,21 @@ Creating a blank Flipbook can be acheived by following the steps below.
 
 **From the Content Browser**: 
 
-1.	Click the **Create** button, then in the context menu under _Animation_, select the **Paper Flipbook** option. 
+1.	Click the **Add New** button, then in the context menu under _Animation_, select the **Paper Flipbook** option. 
 
-	![](FlipbookCreate1.png)(w:440)
+	![](FlipbookCreate1.png)
 
 	[REGION:note]
-	You can also **Right-click** inside the **Content Browser** to open the context menu instead of clicking **Create**. 
+	You can also **Right-click** inside the **Content Browser** to open the context menu instead of clicking **Add New**. 
 	[/REGION]
 
 1.	You will then be prompted to enter a name for your new Flipbook. 
 
-	![](FlipbookCreate2.png)(w:440)
+	![](FlipbookCreate2.png)
 
 1.	After you select a name, your Flipbook asset has been created. 
 
-	![](FlipbookCreate3.png)(w:440)
+	![](FlipbookCreate3.png)
 
 	The asterik symbol in the lower-left corner indicates that the asset has not been saved yet and will go away when you save it. 
 
@@ -101,7 +101,7 @@ To create an automatically generated Flipbook, follow the steps below.
 
 1.	Locate and select each of the sprites you would like to include in the Flipbook in the **Content Browser**. 
 
-	![](Autocreate1.png)(w:440)
+	![](Autocreate1.png)
 
 1.	**Right-click** on any of the sprites, then select the **Create Flipbook** option from the context menu. 
 
@@ -109,11 +109,11 @@ To create an automatically generated Flipbook, follow the steps below.
 
 1.	You will then be prompted to enter a name for your new Flipbook. 
 
-	![](Autocreate3.png)(w:440)
+	![](Autocreate3.png)
 
 1.	After you select a name, your Flipbook asset has been created. 
 
-	![](Autocreate4.png)(w:440)
+	![](Autocreate4.png)
 
 	When you mouse over the Flipbook in the **Content Browser**, you will also be able to preview the Flipbook animation. 
 
@@ -121,3 +121,6 @@ To create an automatically generated Flipbook, follow the steps below.
 When auto-generating a Flipbook, the naming convention you use for the sprites included is **very important** as the sprites will be added to the Flipbook in alphabetical order. In the example above, we've named each of our sprites for creating an Idle to **Idle_x** where X is the order in which we want the sequence to appear. 
 [/REGION]
 
+
+
+
diff --git a/Engine/Documentation/Source/Engine/Paper2D/Importing/Paper2DImporting.INT.udn b/Engine/Documentation/Source/Engine/Paper2D/Importing/Paper2DImporting.INT.udn
new file mode 100644
--- /dev/null
+++ b/Engine/Documentation/Source/Engine/Paper2D/Importing/Paper2DImporting.INT.udn
@@ -0,0 +1,84 @@
+Title:Paper 2D Import Options
+Crumbs: %ROOT%, Engine, Engine/Paper2D
+Description:Importing assets from third party applications for use with Paper 2D. 
+Availability:Public
+Related: Engine\Paper2D\Sprites
+Related: Engine\Paper2D\Flipbooks
+Related: Engine\Paper2D\TileMaps
+
+
+[TOC (start:2 end:2)]
+
+## Overview
+
+**Paper 2D** in Unreal Engine 4 supports a number of different import options for your externally created content. Below are some examples of importing assets from some of the more commonly used content creation applications. 
+
+
+## Importing from Texture Packer
+
+If you are using the third party sprite sheet generation tool **Texture Packer** for content creation, it has a plug-in built into it which will allow you to easily import your content into Unreal Engine 4 by creating a **.paper2dsprites** file. This file when imported into UE4 will automatically import your sprite sheet as a texture, extract and create sprites from it as well as create a **Sprite Asset** which can be used to auto-generate **Flipbooks** out of your sprites. 
+
+**Example Workflow**
+
+1.	Inside **Texture Packer**, select the **UnrealEngine - Paper2d** framework option. 
+
+	![](NewProjectTP.png)(w:640)
+
+1.	Add your assets that will be included in the sprite sheet. 
+
+	![](Assets_Added.png)(w:640)
+
+1.	Under **Layout** set **Size Constraints** to **POT (Power of 2)**. 
+
+1.	Under **Sprites**, show advanced and set **Border** and **Shape Padding** to **0**. 
+
+	![](ExportSettings.png)(w:800)
+
+	This will eliminate the spaces between your sprites, making it easier to tile them together or create **Flipbook** animations. 
+
+1.	Click the **Publish** sprite option (this will create a **.paper2dsprites** file along with a **.png** sprite sheet) 
+
+1.	Drag the **.paper2dsprites** file into the **Content Browser** of your project to import it. 
+
+	![](DragAndDrop.png)(w:800)
+
+	The .png file is already included with the .paper2dsprites file so there is no need to import that as well. 
+
+1.	This will create a **Sprite Sheet** asset and two folders **Frames** (extracted sprites) and **Textures** (your sprite sheet). 
+
+	![](ContentAdded.png)
+
+From here you can **Right-click** on...
+
+* Your **Sprite Sheet** asset to automatically generate **Flipbooks** based on your sprites. 
+* Your imported texture to adjust its settings, configure it for retro sprites or create a **Tile Set** based on it. 
+* Your individual frames to manually add them to a **Flipbook**. 
+
+
+## Importing from Tiled / Adobe Flash CS6
+
+If you are building your content in external content generation tools such as **Tiled** or **Adobe Flash CS6**, you can import your work into Unreal Engine 4 in the **.json** file format. This will automatically import the sprite sheet(s) used to create your map and create **Tile Sets** and a **Tile Map** based off those assets for you. 
+
+**Example Workflow for Tiled**
+
+1.	Inside **Tiled**, choose the File/Save As option, then Save as type **.json**. 
+	[REGION:lightbox]
+	[![](TiledExport.png)(w:640)](TiledExport.png)
+	[/REGION]
+
+	_Click for full view._ 
+
+1.	Drag-and-drop the **.json** asset into the **Content Browser** of your project. 
+
+	![](JSONImport.png)(w:640)
+
+1.	This will import your sprite sheet textures, create **Tile Sets** and a **Tile Map** from those textures. 
+
+	![](JSONAssets.png)
+
+	Opening the **Tile Map** created, you will see the map constructed externally which can be modified inside Unreal Engine 4. 
+
+
+
+
+
diff --git a/Engine/Documentation/Source/Engine/Paper2D/Paper2D.INT.udn b/Engine/Documentation/Source/Engine/Paper2D/Paper2D.INT.udn
--- a/Engine/Documentation/Source/Engine/Paper2D/Paper2D.INT.udn
+++ b/Engine/Documentation/Source/Engine/Paper2D/Paper2D.INT.udn
@@ -26,9 +26,8 @@ Refer to each section below for more information and the **Paper 2D** map inside
 		* [](Engine\Paper2D\Flipbooks"%Engine\Paper2D\Flipbooks:description%")
 			* [](Engine/Paper2D/Flipbooks/Editor "%Engine/Paper2D/Flipbooks/Editor:description%")
 			* [](Engine\Paper2D\Flipbooks\Components "%Engine\Paper2D\Flipbooks\Components:description%")
-		<!-- Not yet ready for release.
 		* [](Engine\Paper2D\TileMaps"%Engine\Paper2D\TileMaps:description%")
-		-->
+		* [](Engine\Paper2D\Importing"%Engine\Paper2D\Importing:description%")
 		* [](Resources/ContentExamples/Paper2D "%Resources/ContentExamples/Paper2D:description%")
 		* [Paper 2D Video Tutorials](https://docs.unrealengine.com/latest/INT/Videos/index.html?category=2D)
 	[/PARAM]
diff --git a/Engine/Documentation/Source/Engine/Paper2D/Sprites/Collision/SpriteCollision.INT.udn b/Engine/Documentation/Source/Engine/Paper2D/Sprites/Collision/SpriteCollision.INT.udn
--- a/Engine/Documentation/Source/Engine/Paper2D/Sprites/Collision/SpriteCollision.INT.udn
+++ b/Engine/Documentation/Source/Engine/Paper2D/Sprites/Collision/SpriteCollision.INT.udn
@@ -49,7 +49,7 @@ The toolbar buttons can be used to add additional polygons or snap vertices to t
 (#2dcollision)
 ## 2D Collision
 
-There is an initial *experimental* integration of Box2D 2.3.1 and various associated changes to enable 2D physics in the 
+There is an initial **experimental** integration of Box2D 2.3.1 and various associated changes to enable 2D physics in the 
 engine. This is a totally unsupported and undocumented prototype, exercise great caution and do not use in production. 
 The current build only includes precompiled libraries of Box2D for Win32 and Win64, so 2D collision will not work on 
 other platforms. Collision detection and response is automatic when the 2D domain is selected in a sprite, but queries 
@@ -57,7 +57,7 @@ like point traces must be separately enabled in the Physics project settings (bE
 
 The integration supports collision detection and response (including Unreal collision channels/filtering), rigid body 
 simulation, and ray casts. Non-zero extent queries, sweep tests, and overlap tests are not implemented yet. Things like 
-CharacterMovementComponent and MoveComponent with bSweep=true rely on these sorts of queries and don’t work correctly yet.
+CharacterMovementComponent and MoveComponent with bSweep=true rely on these sorts of queries and don't work correctly yet.
 
 The long term integration strategy is to make it a first class citizen, where the same techniques and knowledge used in 
 3D scenes will directly apply to 2D scenes; for example, there'll only be one Overlap event, not a separate 2D and 3D 
diff --git a/Engine/Documentation/Source/Engine/Paper2D/Sprites/Editor/SpriteEditor.INT.udn b/Engine/Documentation/Source/Engine/Paper2D/Sprites/Editor/SpriteEditor.INT.udn
--- a/Engine/Documentation/Source/Engine/Paper2D/Sprites/Editor/SpriteEditor.INT.udn
+++ b/Engine/Documentation/Source/Engine/Paper2D/Sprites/Editor/SpriteEditor.INT.udn
@@ -89,11 +89,11 @@ This shows all of the editable properties of the opened Sprite asset.
 | **Source Texture** | The source texture that contains the image to be used for the sprite. |
 | **Source Dimension** | The size of the source region (in pixels) within the Source Texture. |
 | **Source UV** | The offset of the source region (in pixels) within the Source Texture. |
-| **Default Material** | This is the default material that will be used to render the sprite if a particular component doesn’t have a custom material override. |
+| **Default Material** | This is the default material that will be used to render the sprite if a particular component doesn't have a custom material override. |
 | **Pixels per unit** | Adjusts the scaling factor between pixels in the texture and Unreal Units (centimeters). For example, a value of 1.28 will size a 128 pixel wide sprite to be 1 meter. This value can be adjusted per sprite, but you can also set a global default value for newly created sprites in the project settings (Paper2D plugin page). Adjusting your scale to be in a realistic range helps greatly with the stability of physics simulations, among other things. |
 | **Pivot Mode** | Controls the way that the sprite pivot is calculated. It defaults to being computed as the center of the source region, but can be set to any of the corners or mid-points on the edges. You can also define a custom pivot point by setting the mode to Fully Custom and typing in a value for the Custom Pivot Point property. |
-| **Sockets** | Sprites can also define sockets, which can be used as attachment points for other components or reference points to spawn effects or do queries. These can be created by adding to the “Sockets” array in the Details panel and edited in the viewport using the standard gizmos. |
-| **Sprite Collision Domain** | Controls whether or not the sprite has collision shapes, and which domain they’ll be created in. See the ‘Editing sprite collision’ section for more details. When this is set to ‘Use 3D Physics’, Collision Thickness controls the thickness of 3D collision shapes. |
+| **Sockets** | Sprites can also define sockets, which can be used as attachment points for other components or reference points to spawn effects or do queries. These can be created by adding to the **Sockets** array in the Details panel and edited in the viewport using the standard gizmos. |
+| **Sprite Collision Domain** | Controls whether or not the sprite has collision shapes, and which domain they'll be created in. See the _Editing sprite collision_ section for more details. When this is set to _Use 3D Physics_, Collision Thickness controls the thickness of 3D collision shapes. |
 
 ### Sprite List
 
diff --git a/Engine/Documentation/Source/Engine/Paper2D/Sprites/Paper2DSprites.INT.udn b/Engine/Documentation/Source/Engine/Paper2D/Sprites/Paper2DSprites.INT.udn
--- a/Engine/Documentation/Source/Engine/Paper2D/Sprites/Paper2DSprites.INT.udn
+++ b/Engine/Documentation/Source/Engine/Paper2D/Sprites/Paper2DSprites.INT.udn
@@ -130,7 +130,11 @@ See the [](Engine\Paper2D\Sprites\Editor) documentation for information on worki
 
 ### Importing from a JSON file
 
-**Paper 2D** includes an importer for JSON formatted sprite sheet descriptions, which can be exported from tools like **Adobe Flash CS6** or **Texture Packer**. It will import any referenced textures and create sprite assets for each sprite. The importer also assumes that all of the sprites are frames of an animation, so it will always create a Flipbook in addition to the individual sprites. You can delete the generated flipbook if it isn’t needed.
+**Paper 2D** includes an importer for JSON formatted sprite sheet descriptions, which can be exported from tools like **Adobe Flash CS6** or **Texture Packer**. It will import any referenced textures and create sprite assets for each sprite. The importer also assumes that all of the sprites are frames of an animation, so it will always create a Flipbook in addition to the individual sprites. You can delete the generated flipbook if it isn't needed.
+
+[REGION:note]
+See [](Engine\Paper2D\Importing) for more information on Importing Options. 
+[/REGION]
 
 
 ## Creating Sprite Materials
diff --git a/Engine/Documentation/Source/Engine/Paper2D/Sprites/RenderGeometry/SpriteRenderGeometry.INT.udn b/Engine/Documentation/Source/Engine/Paper2D/Sprites/RenderGeometry/SpriteRenderGeometry.INT.udn
--- a/Engine/Documentation/Source/Engine/Paper2D/Sprites/RenderGeometry/SpriteRenderGeometry.INT.udn
+++ b/Engine/Documentation/Source/Engine/Paper2D/Sprites/RenderGeometry/SpriteRenderGeometry.INT.udn
@@ -13,7 +13,7 @@ Related: Engine/Content/Types/Textures
 
 ## Overview
 
-Each sprite is roughly defined as a rectangle from a texture, which can be edited interactively using the �Edit Source Region� mode in the sprite editor, or typed in manually in the details panel (Source Dimension and Source UV, both in pixels). The actual rendered geometry can be further adjusted in the **Edit RenderGeom** mode, and the collision shape can be adjusted in the [Edit Collision](Engine\Paper2D\Sprites\Collision) mode.
+Each sprite is roughly defined as a rectangle from a texture, which can be edited interactively using the **Edit Source Region** mode in the sprite editor, or typed in manually in the details panel (Source Dimension and Source UV, both in pixels). The actual rendered geometry can be further adjusted in the **Edit RenderGeom** mode, and the collision shape can be adjusted in the [Edit Collision](Engine\Paper2D\Sprites\Collision) mode.
 
 **Edit RenderGeom** can also be useful to optimize transparent or masked sprites, trading off some extra triangles in exchange for reduced overdraw in transparent regions of the texture. The settings and tools are exactly the same as editing custom collision geometry. Be sure to use the **Snap to Pixel Grid** tool after editing custom geometry (this will happen automatically in a future release, but is a manual process for now).
 
diff --git a/Engine/Documentation/Source/Engine/Paper2D/TileMaps/Paper2DTilesMaps.INT.udn b/Engine/Documentation/Source/Engine/Paper2D/TileMaps/Paper2DTilesMaps.INT.udn
--- a/Engine/Documentation/Source/Engine/Paper2D/TileMaps/Paper2DTilesMaps.INT.udn
+++ b/Engine/Documentation/Source/Engine/Paper2D/TileMaps/Paper2DTilesMaps.INT.udn
@@ -1,14 +1,14 @@
 Title:Paper 2D Tile Sets / Tile Maps
 Crumbs: %ROOT%, Engine, Engine/Paper2D
 Description:Overview of how to create Tile Sets and Tile Maps for use within Paper 2D. 
-Availability:Docs
+Availability:Public
 Related: Engine\Paper2D\Sprites\Collision
 Related: Engine\Paper2D\Flipbooks
 Related: Engine/Rendering/Materials
 Related: Engine/Content/Types/Textures
 
 [REGION:warning]
-The features on this page are currently in the **experimental stage** of development. Aspects may not yet be fully supported, use at your own risk. 
+The features on this page are currently in the **experimental stage** and as such, some aspects may not yet be fully supported. 
 [/REGION]
 
 
@@ -93,7 +93,7 @@ Inside the **Tile Editor** preview window the currently selected Tile and its as
 
 ![](CurrentlySelectedTile.png)
 
-In our example, we've increased the Tile Width and Tile Height to 128. 
+In our example, we've increased the Tile Width and Tile Height to 64 to match our tile size. 
 
 [REGION:lightbox]
 [![](UpdatedTileSheet.png)(w:800)](UpdatedTileSheet.png)
@@ -115,15 +115,36 @@ In addition to setting Tile Width and Tile Height, you can use the following set
 | **Drawing Offset** | The drawing offset for tiles from this set (in pixels). |
 | **Background Color** | The background color displayed in the tile set viewer. |
 
-[REGION:note]
-There are additional settings inside the **Details** panel for per tile adjustments, however these are considered **experimental** in the current release.
-[/REGION]
+### Per Tile Settings / Applying Tile Collision
+
+Also inside the **Details** panel, you can specify per tile settings (such as applying collision to a particluar tile). 
+
+**To apply collision to a tile**: 
+
+1. Select the tile in the viewport (1) then click the type of shape you want to use for collision (2). 
+
+	![](Collision_1.png)
+
+1. The shape will be applied and can be further edited inside the **Tile Editor** window. 
+
+	![](Collision_2.png)
 
+	Note the tile number you are editing (green box above), you will need this to apply settings to this particluar tile. 
+
+1. In the **Details** panel under the **Sprite** section, find the selected tile and change to the desired **Geometry Type**. 
+
+	![](Collision_3.png)
+
+	[REGION:note]
+	For information on sprite collision types, see the [](Engine\Paper2D\Sprites\Collision) documentation. 
+	[/REGION]
+	
+This will set this tile up for collsion based on the specificed geometry type and shape used. For the collision to actually take affect in-game, when creating your **Tile Map** you will need to set the **Sprite Collision Domain** (see below for information on creating/using Tile Maps). 
 
 
 ## Creating Tile Maps
 
-Tile Maps are also created inside the **Content Browser** and edited inside the **Tile Map Editor**. 
+Tile Maps are created inside the **Content Browser** and edited inside the **Tile Map Editor**. 
 
 ### New Tile Map
 
@@ -168,8 +189,25 @@ Above in the **Toolbox** window (1), you can switch between different Tile Sets
 
 The center window, the **Viewport** (2), is used for filling out the Tile Map layout. Using the tools from the Toolbox, you can fill in tiles on the map to create the desired effect by painting, erasing or filling the tiles that comprise your Tile Map. This is where you can "paint" how you want your map to appear in-game. 
 
-In the **Details** panel (3), you can configure how many layers your Tile Map consists of as well as adjust various settings such as the width/height of the Tile Map itself, the width/height of the Tiles that comprise the map, the scaling factor between pixels and Unreal Units, etc. 
+In the **Details** panel (3), you can configure how many layers your Tile Map consists of as well as adjust various settings such as the width/height of the Tile Map itself, the width/height of the Tiles that comprise the map, the scaling factor between pixels and Unreal Units, the type of collision to use, etc. 
 
+### Enabling Collision for your Tile Map
+
+As noted in the Tile Set section, once you've specified which tiles use collision and which do not, you will also need to set your Tile Map to use a collision type. You can do so from the **Details** panel under the **Collision** section. 
+
+![](Collision_4.png)
+
+You can set the collision thickness (which is the extrusion thickness of collision geometry when using a 3D collision domain) as well as the Collision Domain itself (which determines if the Tile Map should have collision geometry and participate in the physical world). 
+
+[REGION:warning]
+The **Use 2D Physics** option is currently experimental and may not work properly, it is suggested that the **Use 3D Physics** option be used instead. 
+[/REGION]
+
+When you play in your level, you can check which tiles have collision and which do not with the **Show Collision** console command.
+ 
+![](ShowCollision.png)
+
+Above, the upper most ground tiles that the player walks on have collision applied to them while the lower ground and water tiles do not. 
 
 ## Example Tile Map Workflow
 
@@ -213,6 +251,14 @@ Below is an example of a typical workflow for creating a Tile Map.
 
 	![](ExampleMap.png)(w:800)
 
+1.	Be sure to enable collision from the **Collision** section of the **Details** panel. 
+
+	![](Collision_4.png)
+
+	[REGION:note]
+	You must have specified in your Tile Set which tiles have collision and which do not for collision to take effect. 
+	[/REGION]
+	
 1.	Drag the Tile Map asset into your level and scale or move based on your needs. 
 
 	![](ExampleGameImage.png)(w:800)
