#UE4 DOCS: fixing UEDOC-1418

diff --git a/Engine/Documentation/Source/Engine/AI/BehaviorTrees/QuickStart/1/1.INT.udn b/Engine/Documentation/Source/Engine/AI/BehaviorTrees/QuickStart/1/1.INT.udn
--- a/Engine/Documentation/Source/Engine/AI/BehaviorTrees/QuickStart/1/1.INT.udn
+++ b/Engine/Documentation/Source/Engine/AI/BehaviorTrees/QuickStart/1/1.INT.udn
@@ -28,7 +28,7 @@ Description:Step 1. Essential Setup
 	The **Starter Assets** are optional for this quick start guide.
 
 
-If you need a more in depth explanation of starting a new project, see the first step of the [Editor Quick Start Guide](Engine/QuickStart/1).
+If you need a more in depth explanation of starting a new project, see the [Project Browser](Engine/Basics/Projects/Browser) User Guide.
 
 [/EXCERPT:AIProjectStart]
 %Steps%
