#UE4 Docs: Blueprint Best Practices Check-in

diff --git a/Engine/Documentation/Source/Engine/Blueprints/BestPractices/BestPractices.INT.udn b/Engine/Documentation/Source/Engine/Blueprints/BestPractices/BestPractices.INT.udn
new file mode 100644
--- /dev/null
+++ b/Engine/Documentation/Source/Engine/Blueprints/BestPractices/BestPractices.INT.udn
@@ -0,0 +1,127 @@
+Availability:Public
+Title:Blueprint Best Practices
+Crumbs:%ROOT%, Engine/Blueprints
+Description:Tips and tricks to help you make decisions about when to use Blueprints and how best to set them up.
+
+
+[TOC(start:2)]
+
+## Overview
+
+You can do many things with Blueprints. Everything from making a small game or procedural content tools, to prototyping new functionality, to tweaking and polishing things made by programmers,
+is possible with the Blueprints visual scripting system.
+
+However, there are certain things that will impact your performance more if they are done in Blueprints. 
+If you have a Blueprint that's doing a lot of operations and complex math every tick, you might want to consider using
+native C++ code.  Blueprints are best suited to making event-driven functionality, such as handling damage taking, controls, and other things that don't get called every frame.
+
+If you would like to read more about coding for Blueprints, or the technical details of Blueprint compiling, see the [](Engine\Blueprints\TechnicalGuide).
+
+Even if your functionality is well-suited for Blueprints, there are decisions you can make in setting up your Blueprints that will make things go more smoothly. This guide will go
+over some of the more common decisions you will have to make, as well as tips and tricks for Blueprint users.
+
+## Level Blueprints vs. Blueprint Classes
+
+Level Blueprints may be very familiar to users of UE3's Kismet,  because you can select objects directly in the level, and operate on them. They're great for one-off prototypes, and 
+getting familiar with the Blueprints system, but they are specific to the level that they're used in.  This does mean that Level Blueprints can be great places to set up functionality specific 
+to the level, or Actors in it.  Some examples would be kicking off a cinematic when a certain trigger is touched, or opening a particular door after you kill all the enemies.
+
+In general, Blueprint Classes are the best way to get reusable behavior in your project. If you create a Blueprint Class, you can add it to any of your levels, and you can also add as
+many copies as you would like to the level, without needing to copy script around. 
+
+If you do start out in the Level Blueprint, and then decide to move behavior into a Blueprint Class, this should be a pretty simple process. When you copy functions from a Level Blueprint that 
+operate on Actors (say, calling Set Brightness on a light), and paste them into a Blueprint Class based off of a Light Actor, the functions will update to be in the proper scope.  
+
+**Read More:**
+* [](Engine\Blueprints\UserGuide\Types\LevelBlueprint)
+* [](Engine\Blueprints\UserGuide\Types\ClassBlueprint)
+
+## Class Variables vs. Local Variables
+
+When you're working in your Blueprint, you have the ability to add [variables](Engine\Blueprints\UserGuide\Variables) using the [](Engine\Blueprints\Editor\UIComponents\MyBlueprint)
+tab. When you're inside a function or macro, though, you'll see an additional section of the **My Blueprint** tab for **Local Variables**.
+
+![](VariablesLocal.png)
+
+Local variables are scoped, which means that they only exist where you define them.  So, a local variable in a function or macro is only visible to that function, and not to other functions, or the 
+Event Graph.  This helps to reduce clutter for things that are only relevant within the context of that function.  You can consider local variables a "scratch pad" to work with in a function.  You 
+can use them to help perform what the function needs to, but they get tossed away after the function has completed.  
+
+Class variables are for things that you are potentially going to be interested in accessing from multiple places in the Blueprint.  So, things like references to meshes or other components, or variables
+that need to be public so you can access them from other Blueprints, should be class variables.  They're the things you don't want to forget about after you use them.
+
+## Functions vs. Macros
+
+Functions and macros both accomplish the same goal externally:  they let you send input into a node, something happens to it, and then outputs come out.  They accomplish these in two different ways, 
+but they have some similarities:
+
+-	Both have a central place that you can edit the functionality of what the node does (macros in their macro graph, and functions in their function graph), and every node that calls the 
+function or macro will be updated whenever that central place is changed.
+-	Both are reusable.
+-	Both allow for local variables. Macros use "anonymous" local variables, which don't have a specific name, but just hold a value.
+-	Both are awesome ways to encapsulate functionality for reusability and clarity.
+
+There are some key differences that will become more apparent as you get into more advanced use cases.
+
+-	Functions are actually called when you place nodes to call the function.  That means you can target them (i.e. "call a function on another object"), and they allow for communication 
+between Blueprints.
+-	Macros take the nodes from the macro graph, and actually replace the macro node with a copy of all those nodes.  Basically, when the Blueprint is compiled, the macro copies all the
+ graph nodes, and pastes them in where the macro node is.
+
+Because of how they work under the hood, that means there's a few differences in what you can do:
+
+-	Macros can use any node for the class they're scoped to (the class you pick when you create a new Macro Blueprint, or if it's a local macro, the class the macro is defined in).  
+This means that you have a little more versatility in the nodes you can put in a macro.  
+- One of the biggest differences between a function and a macro is that you can place latent nodes in a macro, but not in a function.
+-	Functions allow you to override their functionality in child Blueprints.  For instance, suppose you had a car Blueprint that had a function "PlayerInteractedWithMe".  In that function, 
+you might play a horn honk. Now, suppose you had two children of the car Blueprint, for a police car and a fire truck.  In the police car, you could override the function to play a siren noise and turn 
+on the lights.  In the fire truck, you could have it squirt water.  That sort of overriding functionality isn't possible with macros.
+-	Since macros are just copy-and-pasted into the graph at compile-time, you can have multiple execution wires in and out of a macro.  You can't do that in a function.
+
+In general, a good rule of thumb might be to use a macro if there's a quick bit of functionality that you want to reuse everywhere.  But, if you might want to tweak that behavior in children of the Blueprint, or you want to 
+access it directly from another Blueprint, it's better to use a function!
+
+**Read More:**
+* [](Engine\Blueprints\UserGuide\Functions)
+* [](Engine\Blueprints\UserGuide\Macros)
+
+## Blueprint Communication Types
+
+There are several different methods for communicating between two Blueprints. While the most common use case is direct Blueprint communication, sometimes the functionality you are going for
+means that you will use Event Dispatchers or Blueprint Interfaces instead. For an overview of each type of communication, some example use cases, and tutorials to help you get started, see
+[](Engine\Blueprints\UserGuide\BlueprintCommsUsage). To create a project that uses all three types together, try out the [](Engine/Blueprints/UserGuide/BlueprintCommsUsage/BPComHowTo) guide.
+
+
+
+## General Tips
+
+### Working in the Blueprint Editor
+
+When working in the Blueprint Editor, there are a few tricks you can use to find the nodes that you want to use, as well as nodes and comments you've already created.
+
+* Use the **Search** button in the Blueprint Editor when you're looking for something in an existing Blueprint, whether it's a variable, function, or comment. It searches
+all Blueprints, even unloaded ones, if you uncheck the **Find in Current Blueprint** only box, so it can help you track down just where you implemented something.  
+	
+	![](SearchResults.png)
+
+* The context menu only shows nodes that you can connect to the pin you dragged from (e.g. when you drag from a light, it shows you functions relevant to the light).  
+When you're trying to discover all the possible functionality available to act on a variable or component, just browse through the context menu and see what's avaiable.
+
+	![](ContextMenu.png)
+
+* Whether you're in the context menu or the palette, the search bar can be very helpful. We tag functions with lots of keywords to help you find what you're looking for, even if you don't know what it's called!
+
+	![](search_rotate.png)
+
+### Setting Up Your Graph
+
+While everyone has their own graph setup preferences, here are a few tips that can help you keep things organized and easy to follow:
+
+* Keep things tidy from the get go!  It's much harder to clean up after you've made a bunch of code than to work clean as you go.
+* If you find yourself using the same set of nodes more than twice in a graph, consider making it a function or a macro so you can reuse it.
+* You can stack nodes with their context, because it keeps things more spatially compact. For example, if you have a reference to a light, then you access its point light component, you 
+can stack those two nodes one on top of the other, and think of them as a single block.
+
+	![](stackednodes.png)
+* Use [comments](Engine\Blueprints\UserGuide\Comments) often! Also, don't forget that you can change the color of comment boxes, for extra help identifying sections of your graph.
+
diff --git a/Engine/Documentation/Source/Engine/Blueprints/Blueprints.INT.udn b/Engine/Documentation/Source/Engine/Blueprints/Blueprints.INT.udn
--- a/Engine/Documentation/Source/Engine/Blueprints/Blueprints.INT.udn
+++ b/Engine/Documentation/Source/Engine/Blueprints/Blueprints.INT.udn
@@ -198,6 +198,17 @@ Through the use of _Blueprints_, designers can prototype, implement, or modify v
 		[RELATIVE:Engine\Blueprints\UserGuide\CheatSheet]
 		[/PARAM]
 	[/OBJECT]
+	[OBJECT:TopicCompact]
+		[PARAM:image]
+		![](BestPractices_Topic.png)
+		[/PARAM]
+		[PARAM:title]
+			Blueprint Best Practices
+		[/PARAM]
+		[PARAM:path]
+		[RELATIVE:Engine\Blueprints\BestPractices]
+		[/PARAM]
+	[/OBJECT]
 	[OBJECT:TopicCompact]
 		[PARAM:image]
 			![](UsingBPComms_Topic.png)
diff --git a/Engine/Documentation/Source/Engine/Blueprints/TechnicalGuide/BlueprintsTechnicalGuide.INT.udn b/Engine/Documentation/Source/Engine/Blueprints/TechnicalGuide/BlueprintsTechnicalGuide.INT.udn
--- a/Engine/Documentation/Source/Engine/Blueprints/TechnicalGuide/BlueprintsTechnicalGuide.INT.udn
+++ b/Engine/Documentation/Source/Engine/Blueprints/TechnicalGuide/BlueprintsTechnicalGuide.INT.udn
@@ -6,7 +6,7 @@ Description:Technical guide for programmers working with Blueprints.
 
 Blueprints are a powerful new feature introduced in Unreal Engine 4.  
 
-In short, Blueprints are a way to create new UClasses without the need for writing or compiling code.  When you create a blueprint, you can choose to extend a 'native' C++ class or another Blueprint 
+In short, Blueprints are a way to create new UClasses without the need for writing or compiling code.  When you create a Blueprint, you can choose to extend a 'native' C++ class or another Blueprint 
 class.  You can then add, arrange, and customize components, implement custom logic using a visual scripting language, respond to events and interactions, define custom variables, handle input, and 
 generally create a fully custom object type.
 
@@ -26,6 +26,7 @@ of factors, such as a fence that automatically sizes itself to fill a gap betwee
 	[PARAM:links]
 		* [](Engine/Blueprints/TechnicalGuide/Compiler "%Engine/Blueprints/TechnicalGuide/Compiler:description%")
 		* [](Engine/Blueprints/TechnicalGuide/ExtendingBlueprints "%Engine/Blueprints/TechnicalGuide/ExtendingBlueprints:description%")
+		* [](Engine/Blueprints/TechnicalGuide/Guidelines "%Engine/Blueprints/TechnicalGuide/Guidelines:description%")
 	[/PARAM]
 [/OBJECT]
 
diff --git a/Engine/Documentation/Source/Engine/Blueprints/TechnicalGuide/Compiler/BlueprintCompilerInternals.INT.udn b/Engine/Documentation/Source/Engine/Blueprints/TechnicalGuide/Compiler/BlueprintCompilerInternals.INT.udn
--- a/Engine/Documentation/Source/Engine/Blueprints/TechnicalGuide/Compiler/BlueprintCompilerInternals.INT.udn
+++ b/Engine/Documentation/Source/Engine/Blueprints/TechnicalGuide/Compiler/BlueprintCompilerInternals.INT.udn
@@ -1,7 +1,7 @@
 Availability:Public
 Title:Blueprint Compiler Overview
 Crumbs: %ROOT%, Programming, Engine/Blueprints/TechnicalGuide
-Description:Guide for programmers extending Blueprints.
+Description:The steps of the Blueprint compilation process
 
 [TOC(start:2 end:2)]
 
diff --git a/Engine/Documentation/Source/Engine/Blueprints/TechnicalGuide/ExtendingBlueprints/ExposingGameplay.INT.udn b/Engine/Documentation/Source/Engine/Blueprints/TechnicalGuide/ExtendingBlueprints/ExposingGameplay.INT.udn
--- a/Engine/Documentation/Source/Engine/Blueprints/TechnicalGuide/ExtendingBlueprints/ExposingGameplay.INT.udn
+++ b/Engine/Documentation/Source/Engine/Blueprints/TechnicalGuide/ExtendingBlueprints/ExposingGameplay.INT.udn
@@ -2,6 +2,7 @@ Availability:Public
 Title:Exposing Gameplay Elements to Blueprints
 Crumbs: %ROOT%, Programming, Engine/Blueprints/TechnicalGuide
 Description:Technical guide for gameplay programmers exposing gameplay elements to Blueprints.
+Related: Engine/Blueprints/TechnicalGuide/Guidelines
 
 [TOC(start:2)]
 
@@ -16,6 +17,8 @@ rate of the gun to work on a curve. Rather than having to recode the fire rate o
 game, the designer can simply go into the Blueprint and change the fire rate directly, saving both the designer 
 and the programmer time. 
 
+For more details about optimizing the C++/Blueprints distribution in your project, or tips to keep in mind when
+creating an API exposed to Blueprints, see the [](Engine/Blueprints/TechnicalGuide/Guidelines).
 
 ## Making Classes Blueprintable
 
@@ -29,7 +32,7 @@ The simplest form of declaration for a Blueprintable class is similar to the fol
 	UCLASS(Blueprintable)
 	class AMyBlueprintableClass : AActor
 	{
-		GENERATED_UCLASS_BODY()
+		GENERATED_BODY()
 	}
 
 
diff --git a/Engine/Documentation/Source/Engine/Blueprints/TechnicalGuide/Guidelines/Guidelines.INT.udn b/Engine/Documentation/Source/Engine/Blueprints/TechnicalGuide/Guidelines/Guidelines.INT.udn
new file mode 100644
--- /dev/null
+++ b/Engine/Documentation/Source/Engine/Blueprints/TechnicalGuide/Guidelines/Guidelines.INT.udn
@@ -0,0 +1,115 @@
+Availability:Public
+Title:Guidelines for Programming for Blueprints
+Crumbs: %ROOT%, Programming, Engine/Blueprints/TechnicalGuide
+Description:Tips and tricks for how best to write a Blueprint-friendly API
+Related:Engine/Blueprints/TechnicalGuide/ExtendingBlueprints 
+
+[TOC (start:2 end:2)]
+
+## Overview
+There are two driving considerations when deciding to use C++ or Blueprints:
+-	Speed
+-	Complexity of expression
+
+Beyond those two factors, a lot of it comes down to the complexity of the game, and the composition of the team.  If you have many more artists than programmers, you'll probably have significantly 
+more Blueprints than C++ code.  In contrast, if you have a lot of programmers, they're probably keen to keep things in C++.  We expect people to fall somewhere in the middle.  At Epic, 
+a lot of the workflow is that content creators will make a very complex Blueprint, and a programmer can see how he can compress a lot of that work into C++ by coding a new Blueprint node, 
+so he moves that chunk of functionality into a new C++ function.  A good practice to follow would be to use Blueprints extensively, and then push things into C++ as they reach a complexity 
+where that would enable a more concise expression of the functionality (or it becomes too complex for a non-programmer), or speed of execution dictates a move to C++.
+
+
+### Speed
+For speed, the fact is that Blueprint execution is slower than C++ execution.  That's not to say performance is bad, but if you're doing something that requires a lot of calculations, or operates 
+at a high frequency, it might be better to use C++ instead of Blueprints.  However, it is possible to combine the two in the way that works best for your team and your project's performance.
+If you have a Blueprint that has a lot of functionality, you can push some of that functionality into C++ to speed it up, but keep the rest in Blueprints to retain the flexibility.  
+If your profiling shows that one operation is taking a lot of time in Blueprints, consider moving just that section into C++, while keeping the rest in Blueprints.  
+
+An example of a system that would take a lot of time to execute if done with Blueprints visual scripting is a crowd system that controls a thousand Actors  In this case, it would be better for performance 
+to handle decision making, pathing, and other crowd functionality in C++, and then maybe expose some tweaking parameters and controlling functions to Blueprints.
+
+### Complexity
+For complexity of expression, there are just some things that are easier to do in C++ than in Blueprints.  Blueprints do a lot of things well, but some things are just not very 
+easy to express in nodes.  Operating on large sets of data, doing string manipulation, complex math over large data sets, etc. are all very complex, and aren't easy to follow in a visual system. 
+Those things are better kept in C++ than in Blueprints, just because they're easier to look at and figure out what's going on.  Complexity of expression is another reason a crowd system would
+be better to implement in C++ code than in Blueprints. 
+
+
+
+## Examples
+
+Since there are different pieces of functionality that are best handled in C++ or in Blueprints, here are some examples of how C++ programmers and Blueprint authors can work
+together while creating a game. 
+
+* The programmer could create a Character class in C++ that defines some custom events, and then Blueprints could be used to extend that Character class and actually assign meshes and set defaults. 
+Check out the player character and the enemy bots in the ShooterGame sample project to see implementations like this.
+
+* A ability system where the base class is implemented in C++ but then designers can create Blueprints which actually do something. In the StrategyGame sample, there is a base turret defined in C++,
+but the behavior of the flamethrower, cannon, and arrow turrets is all defined in those Blueprints.
+
+* A pickup where the "Collect" or "Respawn" functions are Blueprint Implementable Events that can be overridden for designers to spawn different particle emitters and sound effects. Both ShooterGame and 
+StrategyGame have examples of pickups created this way.
+
+
+## Creating a Blueprint API : Tips and Tricks
+
+Here are some points to consider as a programmer creating a Blueprint-exposed API:
+
+* Optional parameters are handled well in Blueprints:
+
+		/**
+		 * Prints a string to the log, and optionally, to the screen
+		 * If Print To Log is true, it will be visible in the Output Log window.  Otherwise it will be logged only as 'Verbose', so it generally won't show up.
+		 *
+		 * @param	InString		The string to log out
+		 * @param	bPrintToScreen	Whether or not to print the output to the screen
+		 * @param	bPrintToLog		Whether or not to print the output to the log
+		 * @param	bPrintToConsole	Whether or not to print the output to the console
+		 * @param	TextColor		Whether or not to print the output to the console
+		 */
+		UFUNCTION(BlueprintCallable, meta=(WorldContext="WorldContextObject", CallableWithoutWorldContext, Keywords = "log print", AdvancedDisplay = "2"), Category="Utilities|String")
+		static void PrintString(UObject* WorldContextObject, const FString& InString = FString(TEXT("Hello")), bool bPrintToScreen = true, bool bPrintToLog = true, FLinearColor TextColor = FLinearColor(0.0,0.66,1.0));
+
+* Favor functions with lots of return parameters over functions that return structs. Here's a snippet showing how to create multiple output pins on your node:
+
+		UFUNCTION(BlueprintCallable, Category = "Example Nodes")
+		static void MultipleOutputs(int32& OutputInteger, FVector& OutputVector);
+
+* Adding new parameters to an existing function is fine, but if you need to remove or change them, you should deprecate the original and add a new function. Make sure to use the deprecation metadata so the information about the new function will show up in Blueprints:
+
+		UFUNCTION(BlueprintCallable, Category="Collision", meta=(DeprecatedFunction, DeprecationMessage = "Use new CapsuleOverlapActors", WorldContext="WorldContextObject", AutoCreateRefTerm="ActorsToIgnore"))
+		static ENGINE_API bool CapsuleOverlapActors_DEPRECATED(UObject* WorldContextObject, const FVector CapsulePos, float Radius, float HalfHeight, EOverlapFilterOption Filter, UClass* ActorClassFilter, const TArray<AActor*>& ActorsToIgnore, TArray<class AActor*>& OutActors);
+
+* If a function needs to take an enum, consider using the 'expand enum as execs' metadata, as this can make the node easier to use. 
+
+		UFUNCTION(BlueprintCallable, Category = "DataTable", meta = (ExpandEnumAsExecs="OutResult", DataTablePin="CurveTable"))
+		static void EvaluateCurveTableRow(UCurveTable* CurveTable, FName RowName, float InXY, TEnumAsByte<EEvaluateCurveTableResult::Type>& OutResult, float& OutXY);
+
+* Many operations that take time to complete (e.g. move here) should be latent functions.
+
+		/** 
+		 * Perform a latent action with a delay.
+		 * 
+		 * @param WorldContext	World context.
+		 * @param Duration 		length of delay.
+		 * @param LatentInfo 	The latent action.
+		 */
+		UFUNCTION(BlueprintCallable, Category="Utilities|FlowControl", meta=(Latent, WorldContext="WorldContextObject", LatentInfo="LatentInfo", Duration="0.2"))
+		static void	Delay(UObject* WorldContextObject, float Duration, struct FLatentActionInfo LatentInfo );
+
+* Consider putting functions in a shared library if possible. This makes them easy to use across multiple classes, and avoids a 'target' pin.
+
+		class DOCUMENTATIONCODE_API UTestBlueprintFunctionLibrary : public UBlueprintFunctionLibrary
+
+* Remember to mark nodes as pure if possible, as this will avoid needing an execution pin on a node that requires wiring.
+
+		/* Returns a uniformly distributed random number between 0 and Max - 1 */
+		UFUNCTION(BlueprintPure, Category="Math|Random")
+		static int32 RandomInteger(int32 Max);
+
+* Marking a function as `const` will also cause the Blueprint node to not have execution pins: 
+		/**
+		 * Get the actor-to-world transform.
+		 * @return The transform that transforms from actor space to world space.
+		 */
+		UFUNCTION(BlueprintCallable, meta=(DisplayName = "GetActorTransform"), Category="Utilities|Transformation")
+		FTransform GetTransform() const;
diff --git a/Engine/Documentation/Source/Engine/Blueprints/UserGuide/BlueprintCommsUsage/BlueprintCommsUsage.INT.udn b/Engine/Documentation/Source/Engine/Blueprints/UserGuide/BlueprintCommsUsage/BlueprintCommsUsage.INT.udn
--- a/Engine/Documentation/Source/Engine/Blueprints/UserGuide/BlueprintCommsUsage/BlueprintCommsUsage.INT.udn
+++ b/Engine/Documentation/Source/Engine/Blueprints/UserGuide/BlueprintCommsUsage/BlueprintCommsUsage.INT.udn
@@ -56,9 +56,12 @@
 
 ![](Engine\Blueprints\UserGuide\BlueprintCommsUsage\BPComHowTo\2_9.png)
 
-**Direct Blueprint Communication** (or **Blueprint Comms**) is the most common method of communication and is excellent when you have two **Blueprint Actors** that you know you want to have talk to each other at a certain point. 
+**Direct Blueprint Communication** (or **Blueprint Comms**) is the most common method of communication and is excellent when you have two **Blueprint Actors** that you know you want to have talk 
+to each other at a certain point. 
 	
-This type of communication is always one-to-one, meaning one Blueprint (referred to as the **Working Blueprint**) is requesting to access another Blueprint (referred to as the **Target Blueprint**). The simplest way to use **Blueprint Comms** is to get a reference to the **Target Blueprint** through a publicly exposed object variable and then specify which instance of that Blueprint you want to access.
+This type of communication is always one-to-one, meaning one Blueprint (referred to as the **Working Blueprint**) is requesting to access another Blueprint (referred to as the **Target 
+Blueprint**). The simplest way to use **Blueprint Comms** is to get a reference to the **Target Blueprint** through a publicly exposed object variable and then specify which instance of 
+that Blueprint you want to access.
 
 Refer to [](Engine\Blueprints\UserGuide\BlueprintComms) for more detailed information. 
 
@@ -71,46 +74,23 @@
 
 Refer to [Setting Up Direct Blueprint Communications](Engine\Blueprints\UserGuide\BlueprintCommsUsage\BPComHowTo#directblueprintcommunications) for a walk through of how to set one up. 
 
-## Blueprint Casting
 
-![](Engine\Blueprints\UserGuide\BlueprintCommsUsage\BPComHowTo\Event6.png)
-
-The Direct Blueprint Communication method above works well for communicating with two Actors placed inside your level, as you can select the instance in the level that will be the **Target Blueprint**. However, what if you want to communicate between an Actor in a level and a player character? Or, what if you want to access information inside an Actor and the Actor itself is a variable (for instance, affect the Actor that enters a fire you have placed in your level differently depending upon the Actor that enters it). 
-
-Taking the player character example, suppose you want the player to take damage when they enter fire that is placed in your level. Since we want to apply damage to the player, the player character Blueprint would be the **Target Blueprint** we want to access, however we cannot specify it as the "target" as it is not a publicly exposed object variable and it does not exist in the world until we start the game. 
-
-In this situation we could use **Blueprint Casting** to take the playable character and **Cast To** our **Target Blueprint** (Character Blueprint), or in other words, **treat as X** where X is the **Target Blueprint** we want to gain access to. 
-
-### When to Use It...
 	
-Here are a few examples when you could use **Blueprint Casting**: 
-
-[EXCERPT:Castusage]
-*	**You have a character interact with an Actor in your level and in turn causes a change to the character in some way.**
-	*	Player picks up a health pack in the level, restoring health in the player character. 
-	*	Enemy walks in snow or mud, causing their movement speed to slow down. 
-*	**You want to pass variables to a HUD or UI Widget, Save Game, Controller, or Game Mode.** 
-	*	Passing player or enemy health from a Character Blueprint to a HUD or Widget Blueprint to display the value on screen. 
-	*	Passing variables (current health or energy, money, items, or abilities) to a Save Game Blueprint so that they can be restored later. 
-*	**You have multiple Blueprints which are children of another Blueprint and you want to access one of the children.**
-	*	You have Cat, Dog, and Bird Blueprints which share functionality from an Animal Blueprint. Using the Animal Blueprint, you **Cast To** the Bird Blueprint to execute a "Flying" function. 
-*	**You have a Blueprint that is dependent on another Blueprint's Variables, Events, or Functions.** 
-	*	Plugging a battery into a flashlight turns on the light and the intensity of the light is determined by a value in the battery. 
-*	**You want to access multiple Blueprints of the same class and modify them all the same way at once.**
-	*	You have several lights in your level and you want them all to turn on or off when an event occurs. 
-[/EXCERPT:Castusage]
-
-Refer to [Setting Up Blueprint Casts](Engine\Blueprints\UserGuide\BlueprintCommsUsage\BPComHowTo#blueprintcasting) for a walk through of how to set one up. 
-	
 ## Event Dispatchers
 
-![](Engine\Blueprints\UserGuide\BlueprintCommsUsage\BPComHowTo\InterfaceExample2_7.png)
+![](OpenDoorEventDispatch.png)
 
-By binding one or more Events to an **Event Dispatcher**, you can cause all of Events that are bound to the dispatcher to fire once the **Event Dispatcher** is called. 
+**Event Dispatchers** are best suited for telling the world something happened, when you don't care what they do with that information.  The pattern for using these is often (but not always) that 
+something spawns another thing (e.g. the Level Blueprint spawns a boss), and they hook the Event Dispatcher to get a notification when something important happens for that spawned thing.  
 
-For example, inside the player character Blueprint you create an **Event Dispatcher** called "SpawnEnemy" and it is executed when the player presses the **F** key. Inside your level you place several spawn points from which an enemy can spawn, then in the Level Blueprint you bind the SpawnEnemy **Event Dispatcher** from the character Blueprint with an Event in the Level Blueprint called "ExecuteSpawn" that randomly chooses a spawn location and spawns an enemy. 
+For the boss example, perhaps we want to open a door when the player kills the boss.  So, the Level Blueprint would spawn the boss, and then hook its OnDied Event Dispatcher to an event that 
+opened the door.  That allows the boss to stay general, as it doesn't care what level it's placed in, and just provides a "hook" for anyone that might cares it dies.  It doesn't just have to be
+ the locked door though. Maybe all his minions hook his OnDied delegate, and they all run away when the boss dies.  You can have as many things as you want hooked into a single Event Dispatcher.
+and all of the Events that are bound to the dispatcher will fire once the Event Dispatcher is called. 
 
-Whenever the player presses the **F** key, the **Event Dispatcher** is called and since it is bound to the "ExecuteSpawn" Event in the Level Blueprint, that Event is also called resulting in an enemy being spawned. 
+To sum up, Event Dispatchers are used to tell anyone interested that something just happened to you. The events called by an Event Dispatcher can be different for different instances of your Blueprint,
+or at different points in gameplay.
 
+
 Refer to [](Engine/Blueprints/UserGuide/EventDispatcher) for more detailed information. 
 
@@ -132,9 +112,16 @@
 
 ![](Engine\Blueprints\UserGuide\BlueprintCommsUsage\BPComHowTo\InterfaceExample2_5.png)
 
-**Blueprint Interfaces** allow for a common method of interacting with multiple types of objects that all share some specific functionality. This means you can have completely different types of objects, such as a car and a tree, that share one specific thing; they can both be shot by weapon fire and take damage. 
+**Blueprint Interfaces** allow for a common method of interacting with multiple types of objects that all share some specific functionality. This means you can have completely different types of 
+objects, such as a car and a tree, that share one specific thing; they can both be shot by weapon fire and take damage. 
 
-By creating a **Blueprint Interface** that contains an "OnTakeWeaponFire" function and having both the car and the tree implement that **Blueprint Interface**, you can treat the car and the tree as the same and simply call the "OnTakeWeaponFire" function when either of them is shot and have them respond differently.
+By creating a **Blueprint Interface** that contains an **OnTakeWeaponFire** function and having both the car and the tree implement that **Blueprint Interface**, you can treat the car and the tree 
+as the same and simply call the **OnTakeWeaponFire** function when either of them is shot and have them respond differently.
 
+This is still a one-to-one style of communication, just like direct Blueprint communications, but unlike direct Blueprint communications, the Target may or may not receive the message.  An 
+interface is a contract of sorts.  It says "if you implement this interface, you promise to implement these specific functions, and respond to them when I call."  The difference between direct 
+Blueprint communications and interfaces is that with direct communications, you have to know specifically whether or not the Blueprint on the receiving end cares about a function or not.  
+Interface Messages you can send to anything, and they'll respond if they care, and ignore the message if they don't. Going back to the weapon fire example, you may not want your walls to react to damage, so they just wouldn't implement the interface with **OnTakeWeaponFire**. 
+
 Refer to [](Engine/Blueprints/UserGuide/Types/Interface/UsingInterfaces) for more detailed information. 
 
@@ -147,10 +134,40 @@
 	*	Player has a "use" button, pressing the use button opens doors, turns off lights, picks up items, etc. 
 	*	Enemy has a special ability that changes based on player health. 
 
-Refer to [Setting Up Blueprint Interfaces](Engine\Blueprints\UserGuide\BlueprintCommsUsage\BPComHowTo#blueprintinterfaces) for a walk through of how to set one up. 
+Refer to [Setting Up Blueprint Interfaces](Engine\Blueprints\UserGuide\BlueprintCommsUsage\BPComHowTo#blueprintinterfaces) for a walkthrough of how to set one up. 
+
+
+## Casting For Communications
+
+![](Engine\Blueprints\UserGuide\BlueprintCommsUsage\BPComHowTo\Event6.png)
+
+When communicating between Blueprints, you will often need to make sure that your "Target" Blueprint is the one that you expect. For example, what if you want to communicate between an Actor in a level and a player character? Or, what if you want to access information inside an 
+Actor and the Actor itself is a variable (for instance, affect the Actor that enters a fire you have placed in your level differently depending upon the Actor that enters it). 
+
+Taking the player character example, suppose you want the player to take damage when they enter fire that is placed in your level. Since we want to apply damage to the player, the player character Blueprint would be the **Target Blueprint** we want to access, however we cannot specify it as the "target" as it is not a publicly exposed object variable and it does not exist in the world until we start the game. 
+
+In this situation we could use **Blueprint Casting** to take the playable character and **Cast To** our **Target Blueprint** (Character Blueprint), or in other words, **treat as X** where X is the **Target Blueprint** we want to gain access to. 
+
+### When to Use It...
+	
+Here are a few examples when you could use **Blueprint Casting**: 
+
+[EXCERPT:Castusage]
+*	**You have a character interact with an Actor in your level and in turn causes a change to the character in some way.**
+	*	Player picks up a health pack in the level, restoring health in the player character. 
+	*	Enemy walks in snow or mud, causing their movement speed to slow down. 
+*	**You want to pass variables to a HUD or UI Widget, Save Game, Controller, or Game Mode.** 
+	*	Passing player or enemy health from a Character Blueprint to a HUD or Widget Blueprint to display the value on screen. 
+	*	Passing variables (current health or energy, money, items, or abilities) to a Save Game Blueprint so that they can be restored later. 
+*	**You have multiple Blueprints which are children of another Blueprint and you want to access one of the children.**
+	*	You have Cat, Dog, and Bird Blueprints which share functionality from an Animal Blueprint. Using the Animal Blueprint, you **Cast To** the Bird Blueprint to execute a "Flying" function. 
+*	**You have a Blueprint that is dependent on another Blueprint's Variables, Events, or Functions.** 
+	*	Plugging a battery into a flashlight turns on the light and the intensity of the light is determined by a value in the battery. 
+*	**You want to access multiple Blueprints of the same class and modify them all the same way at once.**
+	*	You have several lights in your level and you want them all to turn on or off when an event occurs. 
+[/EXCERPT:Castusage]
+
+Refer to [Setting Up Blueprint Casts](Engine\Blueprints\UserGuide\BlueprintCommsUsage\BPComHowTo#blueprintcasting) for a walk through of how to set one up. 
 
 
 
-
-
-
