TypoCorrection

diff --git a/Engine/Documentation/Source/Engine/Paper2D/Sprites/Paper2DSprites.INT.udn b/Engine/Documentation/Source/Engine/Paper2D/Sprites/Paper2DSprites.INT.udn
--- a/Engine/Documentation/Source/Engine/Paper2D/Sprites/Paper2DSprites.INT.udn
+++ b/Engine/Documentation/Source/Engine/Paper2D/Sprites/Paper2DSprites.INT.udn
@@ -1,7 +1,7 @@
+﻿Availability:Public
 Title:Paper 2D Sprites
 Crumbs: %ROOT%, Engine, Engine/Paper2D
 Description:Description of Sprites in Unreal Engine 4 and how to create them. 
-Availability:Public
 Related: Engine\Paper2D\Sprites\Editor
 Related: Engine\Paper2D\Sprites\Collision
 Related: Engine\Paper2D\Sprites\RenderGeometry
@@ -130,7 +130,7 @@ See the [](Engine\Paper2D\Sprites\Editor) documentation for information on worki
 
 ### Importing from a JSON file
 
-**Paper 2D** includes an importer for JSON formatted sprite sheet descriptions, which can be exported from tools like **Adobe Flash CS6** or **Texture Packer**. It will import any referenced textures and create sprite assets for each sprite. The importer also assumes that all of the sprites are frames of an animation, so it will always create a Flipbook in addition to the individual sprites. You can delete the generated flipbook if it isn't needed.
+**Paper 2D** includes an importer for JSON formatted sprite sheet descriptions, which can be exported from tools like **Adobe Flash CS6** or **Texture Packer**. It will import any referenced textures and create sprite assets for each sprite. The importer also assumes that all of the sprites are frames of an animation, so it will always create a Flipbook in addition to the individual sprites. You can delete the generated flipbook if it is not needed.
 
 [REGION:note]
 See [](Engine\Paper2D\Importing) for more information on Importing Options. 
