TypoCorrection

diff --git a/Engine/Documentation/Source/Engine/Paper2D/Sprites/Editor/SpriteEditor.INT.udn b/Engine/Documentation/Source/Engine/Paper2D/Sprites/Editor/SpriteEditor.INT.udn
--- a/Engine/Documentation/Source/Engine/Paper2D/Sprites/Editor/SpriteEditor.INT.udn
+++ b/Engine/Documentation/Source/Engine/Paper2D/Sprites/Editor/SpriteEditor.INT.udn
@@ -1,7 +1,7 @@
+﻿Availability:Public
 Title:Sprite Editor Reference
 Crumbs: %ROOT%, Engine, Engine/Paper2D
 Description:Tool for previewing and editing UPaperSprite assets
-Availability:Public
 Related: Engine\Paper2D\Sprites\Collision
 Related: Engine\Paper2D\Sprites\RenderGeometry
 Related: Engine\Paper2D\Flipbooks
@@ -89,11 +89,11 @@ This shows all of the editable properties of the opened Sprite asset.
 | **Source Texture** | The source texture that contains the image to be used for the sprite. |
 | **Source Dimension** | The size of the source region (in pixels) within the Source Texture. |
 | **Source UV** | The offset of the source region (in pixels) within the Source Texture. |
-| **Default Material** | This is the default material that will be used to render the sprite if a particular component doesn't have a custom material override. |
+| **Default Material** | This is the default material that will be used to render the sprite if a particular component does not have a custom material override. |
 | **Pixels per unit** | Adjusts the scaling factor between pixels in the texture and Unreal Units (centimeters). For example, a value of 1.28 will size a 128 pixel wide sprite to be 1 meter. This value can be adjusted per sprite, but you can also set a global default value for newly created sprites in the project settings (Paper2D plugin page). Adjusting your scale to be in a realistic range helps greatly with the stability of physics simulations, among other things. |
 | **Pivot Mode** | Controls the way that the sprite pivot is calculated. It defaults to being computed as the center of the source region, but can be set to any of the corners or mid-points on the edges. You can also define a custom pivot point by setting the mode to Fully Custom and typing in a value for the Custom Pivot Point property. |
 | **Sockets** | Sprites can also define sockets, which can be used as attachment points for other components or reference points to spawn effects or do queries. These can be created by adding to the **Sockets** array in the Details panel and edited in the viewport using the standard gizmos. |
-| **Sprite Collision Domain** | Controls whether or not the sprite has collision shapes, and which domain they'll be created in. See the _Editing sprite collision_ section for more details. When this is set to _Use 3D Physics_, Collision Thickness controls the thickness of 3D collision shapes. |
+| **Sprite Collision Domain** | Controls whether or not the sprite has collision shapes, and which domain they will be created in. See the _Editing sprite collision_ section for more details. When this is set to _Use 3D Physics_, Collision Thickness controls the thickness of 3D collision shapes. |
 
 ### Sprite List
 
