TypoCorrection

diff --git a/Engine/Documentation/Source/Engine/Paper2D/Sprites/Collision/SpriteCollision.INT.udn b/Engine/Documentation/Source/Engine/Paper2D/Sprites/Collision/SpriteCollision.INT.udn
--- a/Engine/Documentation/Source/Engine/Paper2D/Sprites/Collision/SpriteCollision.INT.udn
+++ b/Engine/Documentation/Source/Engine/Paper2D/Sprites/Collision/SpriteCollision.INT.udn
@@ -1,7 +1,7 @@
+﻿Availability:Public
 Title:Sprite Collision
 Crumbs: %ROOT%, Engine, Engine/Paper2D, Engine/Paper2D/Sprites
 Description:Displays and allows editing of the Sprite Collision shapes. 
-Availability:Public
 Related: Engine\Paper2D\Sprites\Collision
 Related: Engine\Paper2D\Flipbooks
 Related: Engine/Rendering/Materials
@@ -57,7 +57,7 @@ like point traces must be separately enabled in the Physics project settings (bE
 
 The integration supports collision detection and response (including Unreal collision channels/filtering), rigid body 
 simulation, and ray casts. Non-zero extent queries, sweep tests, and overlap tests are not implemented yet. Things like 
-CharacterMovementComponent and MoveComponent with bSweep=true rely on these sorts of queries and don't work correctly yet.
+CharacterMovementComponent and MoveComponent with bSweep=true rely on these sorts of queries and do not work correctly yet.
 
 The long term integration strategy is to make it a first class citizen, where the same techniques and knowledge used in 
 3D scenes will directly apply to 2D scenes; for example, there'll only be one Overlap event, not a separate 2D and 3D 
