#UE4doc #docs: Changed "EndFPSChart" to "StopFPSChart"

diff --git a/Engine/Documentation/Source/Engine/Performance/PerformanceProfiling.INT.udn b/Engine/Documentation/Source/Engine/Performance/PerformanceProfiling.INT.udn
--- a/Engine/Documentation/Source/Engine/Performance/PerformanceProfiling.INT.udn
+++ b/Engine/Documentation/Source/Engine/Performance/PerformanceProfiling.INT.udn
@@ -228,7 +228,7 @@ You can use the tool by typing **SynthBenchmark** into the console.
 
 It can be very useful to get the stat unit times over a longer period of time (e.g. in-game cutscene or a camera path set up to test many cases).
 
-The following chart was generated on an Android device (**Streaming**) from a cutscene. Before and after the cutscene, the console commands **StartFPSChart** and **EndFPSChart** were entered.
+The following chart was generated on an Android device (**Streaming**) from a cutscene. Before and after the cutscene, the console commands **StartFPSChart** and **StopFPSChart** were entered.
 The resulting .csv file (in [ProjectFolder]\Saved\Cooked\Android_ES31\SubwayPatrol\Saved\Profiling\FPSChartStats) was opened in Microsoft Excel.
 For this example, we removed the first 4 header lines, selected all, and inserted a **Scatter With Straight Lines** plot.
 
