#UE4doc #docs: corrected incidents of mispelled "separate"

diff --git a/Engine/Documentation/Source/Shared/Editors/Common/EditorMenuItems/EditorMenuItems.INT.udn b/Engine/Documentation/Source/Shared/Editors/Common/EditorMenuItems/EditorMenuItems.INT.udn
--- a/Engine/Documentation/Source/Shared/Editors/Common/EditorMenuItems/EditorMenuItems.INT.udn
+++ b/Engine/Documentation/Source/Shared/Editors/Common/EditorMenuItems/EditorMenuItems.INT.udn
@@ -66,10 +66,10 @@ Crumbs:
 * **Developer Tools** - Brings up the selected **Developer Tool** in a separate window (Camera Debugger, Collision Analyzer, Debug Tools, Modules, Widget Reflector).
 * **Blueprint Debugger** - Brings up the **Blueprint Debugger** in a separate window.
 * **Class Viewer** - Brings up the **Class Viewer** in a separate window.
-* **Device Profile Editor** - Brings up the **Device Profile Editor** in a seperate window.
+* **Device Profile Editor** - Brings up the **Device Profile Editor** in a separate window.
 * **Message Log** - Brings up the **Message Log** in a separate window.
 * **Output Log** - Brings up the **Output Log** in a separate window.
-* **Undo History** - Brings up the **Undo History** panel in a seperate window.
+* **Undo History** - Brings up the **Undo History** panel in a separate window.
 * **Plugins** - Brings up the **Plugins** window.
 * **Device Manager** - Brings up **Device Manager** in a separate window.
 * **Session Frontend** - Brings up **Unreal Frontend** in a separate window.
