#UE4doc #docs: Made minor edits.

diff --git a/Engine/Documentation/Source/Engine/Foliage/Foliage.INT.udn b/Engine/Documentation/Source/Engine/Foliage/Foliage.INT.udn
--- a/Engine/Documentation/Source/Engine/Foliage/Foliage.INT.udn
+++ b/Engine/Documentation/Source/Engine/Foliage/Foliage.INT.udn
@@ -249,7 +249,7 @@ Once instances are selected, a number of actions are available:
 
 ![Lasso Tool](ToolPaintSelect.png)
 
-The Lasso tool allows you to select a number of instances simultaneously by using the sphere brush that is also used with the Paint tool. The **Meshes** list selection and filter settings apply, just as they do for regular mesh painting. Holding down **Shift** while painting deselects instances. Once you have a selection, you can switch to the regular Select tool and perform some of the the actions listed above.
+The Lasso tool allows you to select a number of instances simultaneously by using the sphere brush that is also used with the Paint tool. The **Meshes** list selection and filter settings apply, just as they do for regular mesh painting. Holding down **Shift** while painting deselects instances. Once you have a selection, you can switch to the regular Select tool and perform some of the actions listed above.
 
 ![Lasso Tool](Foliage_Lasso.png)(w:600)
 
@@ -332,7 +332,7 @@ Hardware instancing on consoles such as Xbox360 and PS3 works a little different
 
 For example, if you had a foliage mesh with 100 triangles, normally this would need 300 vertex indices (3 for each triangle) to render as a regular Static Mesh. These are 2 bytes each, so 600 bytes of data are stored for the Static Mesh. 
 
-If you wanted to be able to render a maximum of 20 instances of this mesh as foliage in a single draw call, we could set the ConsolePreallocateInstanceCount to 20, which would store duplicate values of those 300 indices 20 times. That would mean the total memory requirement for the indices of that Static Mesh would increase to 12,000 bytes (20 * 300 * 2).
+If you wanted to be able to render a maximum of 20 instances of this mesh as foliage in a single draw call, we could set the ConsolePreallocateInstanceCount to 20, which would store duplicate values of those 300 indices 20 times. That would mean the total memory requirement for the indices of that Static Mesh would increase to 12,000 bytes (20 x 300 x 2).
 
 This overhead is incurred only once for the Static Mesh regardless of how many instances will be placed in the level, but simply preallocating that data allows us to render up to 20 instances of this Static Mesh in a single draw call any time we want, rendering up to 2000 triangles
 
