#UE4doc #docs: deleted duplicate word, "now" fixed "let's"

diff --git a/Engine/Documentation/Source/Engine/Landscape/Custom/CreatingCustomHeightmapsAndLayers.INT.udn b/Engine/Documentation/Source/Engine/Landscape/Custom/CreatingCustomHeightmapsAndLayers.INT.udn
--- a/Engine/Documentation/Source/Engine/Landscape/Custom/CreatingCustomHeightmapsAndLayers.INT.udn
+++ b/Engine/Documentation/Source/Engine/Landscape/Custom/CreatingCustomHeightmapsAndLayers.INT.udn
@@ -60,7 +60,7 @@ Importing Layers that were made in an external application is a simple process b
 ![](Landscape_Simple_Shader.png)
 1. Once the material is created, apply it to the Landscape Actor so that you have something that looks like this. 
 ![](Landscape_Applied_Material.png)
-1. Now with the material applied, it is now now time to add some  **Layer Info** to our Landscape Actor. Add a **Layer Info** object for each layer you are going to have. 
+1. Now with the material applied, it is now time to add some  **Layer Info** to our Landscape Actor. Add a **Layer Info** object for each layer you are going to have. 
 For this example, we are going to have three. To read more about **Layer Info** objects, check out the [Layer Info Objects](Engine/Landscape/Editing/PaintMode#LayerInfOobjects) section of the [Paint Mode](Engine/Landscape/Editing/PaintMode) document. When completed, your **Target Layers** section should look like this. 
 ![](Landscape_Target_Layers.png)
 1. Now it is time to import our layer. To do this, all you need to do is **Right-Click** on the **Target Layer** you want to import a layer for and then select the layer that you want to import from the dialogue box that pops up. 
@@ -104,7 +104,7 @@ To select a Heightmap to import, press the button next the box that says **Pleas
 
 ![](Engine\Landscape\Creation\Landscape_Import_Section_Button.png)
 
-Now lets try importing a custom height map. For this example, we will be using the image below so make sure that you do load it.  
+Now let's try importing a custom height map. For this example, we will be using the image below so make sure that you do load it.  
 
 [REGION:caption]
 Make sure you right-click Save as on the image to download it to your PC.
