#UE4doc #docs: fixed "mod,e"

diff --git a/Engine/Documentation/Source/Engine/Landscape/QuickStart/1/LSQS_1.INT.udn b/Engine/Documentation/Source/Engine/Landscape/QuickStart/1/LSQS_1.INT.udn
--- a/Engine/Documentation/Source/Engine/Landscape/QuickStart/1/LSQS_1.INT.udn
+++ b/Engine/Documentation/Source/Engine/Landscape/QuickStart/1/LSQS_1.INT.udn
@@ -46,7 +46,7 @@ While each of the three modes for the Landscape tools allows you to interact wit
 | **Common Controls** | **Operation** |
 | --- | --- |
 | ** Ctrl**| Will allow you to select Landscape components. |
-| **Left Mouse Button** | Heightens or increases the heightmap or selected layer's weight. For example, in Sculpting mode, this will raise the Landscape heighmap. In Paint mod,e this will apply the selected material to the Landscape. |
+| **Left Mouse Button** | Heightens or increases the heightmap or selected layer's weight. For example, in Sculpting mode, this will raise the Landscape heighmap. In Paint mode this will apply the selected material to the Landscape. |
 | **Shift + Left Mouse Button**| Lowers or decreases the heightmap or selected layer's weight. For example, in Sculpting mode, this will lower the Landscape heightmap. In Paint mode, this will erase the selected material that was applied to a particular section of the Landscape. |
 | **Ctrl + Z** | Undoes last action. |
 | **Ctrl + Y** | Redoes last undone action. |  
