#UE4doc #docs: made plural: connections

diff --git a/Engine/Documentation/Source/Engine/Landscape/QuickStart/4/LSQS_4.INT.udn b/Engine/Documentation/Source/Engine/Landscape/QuickStart/4/LSQS_4.INT.udn
--- a/Engine/Documentation/Source/Engine/Landscape/QuickStart/4/LSQS_4.INT.udn
+++ b/Engine/Documentation/Source/Engine/Landscape/QuickStart/4/LSQS_4.INT.udn
@@ -151,7 +151,7 @@ Once the **Layer Blend** nodes have been set up, it is time to connect the Textu
 [/REGION]
 
 [REGION:note]
-The material connection were colored in Photoshop to help you better understand were everything needed to be connected. Currently there is no way to change the color of lines connecting Material nodes inside of Unreal Engine 4. 
+The material connections were colored in Photoshop to help you better understand were everything needed to be connected. Currently there is no way to change the color of lines connecting Material nodes inside of Unreal Engine 4. 
 [/REGION]
 
 
