#UE4doc #docs: fixed sentence fragment

diff --git a/Engine/Documentation/Source/Engine/Landscape/QuickStart/5/LSQS_5.INT.udn b/Engine/Documentation/Source/Engine/Landscape/QuickStart/5/LSQS_5.INT.udn
--- a/Engine/Documentation/Source/Engine/Landscape/QuickStart/5/LSQS_5.INT.udn
+++ b/Engine/Documentation/Source/Engine/Landscape/QuickStart/5/LSQS_5.INT.udn
@@ -103,7 +103,7 @@ When completed with the painting, you should have something that looks like this
 
 ## Possible Issues and Workarounds
 
-When you first start painting on your Landscape you might run into an issue where the base Material disappears or turns black. Like in the following picture. 
+When you first start painting on your Landscape you might run into an issue where the base Material disappears or turns black, like in the following picture: 
 [REGION:fullwidth] 
 ![](T_First_Paint_Issues.png)
 [/REGION]
