#UE4doc #docs: corrected punctuation & "...sizes allow" instead of "allows"

diff --git a/Engine/Documentation/Source/Engine/Landscape/TechnicalGuide/LansdscapeTechnicalGuide.INT.udn b/Engine/Documentation/Source/Engine/Landscape/TechnicalGuide/LansdscapeTechnicalGuide.INT.udn
--- a/Engine/Documentation/Source/Engine/Landscape/TechnicalGuide/LansdscapeTechnicalGuide.INT.udn
+++ b/Engine/Documentation/Source/Engine/Landscape/TechnicalGuide/LansdscapeTechnicalGuide.INT.udn
@@ -34,7 +34,7 @@ To get the best performance possible from the Landscape system, there are certai
 ## Technical Details
 
 
-Valid dimensions for Landscape heightmaps are not always immediately obvious when getting started. Being able to determine what dimensions for heightmaps are valid, and beyond that which are optimal, requires a thorough understanding of the underlying architecture of the Landscape. In order to create a system that allows for huge terrains while still being efficient in terms of memory and performance, the architecture implicitly applies restrictions on the dimensions of the heightmap, meaning there are certain dimensions that are valid and others that are not. In previous terrain systems in Unreal Engine, there were either no restrictions (i.e., any dimensions were valid and would work) or the restrictions were fairly simple (i.e., only square power of two heightmaps were allowed). The restrictions on heightmaps for Landscapes are much more complex and rigid.
+Valid dimensions for Landscape heightmaps are not always immediately obvious when getting started. Being able to determine what dimensions for heightmaps are valid and, beyond that, which are optimal requires a thorough understanding of the underlying architecture of the Landscape. In order to create a system that allows for huge terrains while still being efficient in terms of memory and performance, the architecture implicitly applies restrictions on the dimensions of the heightmap, meaning there are certain dimensions that are valid and others that are not. In previous terrain systems in Unreal Engine, there were either no restrictions (i.e., any dimensions were valid and would work) or the restrictions were fairly simple (i.e., only square power of two heightmaps were allowed). The restrictions on heightmaps for Landscapes are much more complex and rigid.
 
 A Landscape Actor is color coded so that it is easier to tell what each section does. The edges of the Landscape are highlighted in **Yellow**, the edge of each component is in **Light Green**, section edges (if set to 2x2 sections) in a **Medium Green**, and the individual Landscape quads in a **Dark Green**.
 
@@ -93,9 +93,9 @@ These examples focus on square Landscapes. However, you can create Landscapes th
 
 ## Performance Considerations
 
-The choice of component size vs the total number of components is a performance trade-off. Smaller component sizes allows quicker LOD transitions and also allows for the occlusion of more terrain, but the smaller size necessitates more components.
+The choice of component size vs the total number of components is a performance trade-off. Smaller component sizes allow quicker LOD transitions and also allows for the occlusion of more terrain, but the smaller size necessitates more components.
 
-Each component has a render-thread CPU processing cost and each section is a draw call, so try to keep these number to a minimum. For the largest Landscapes, Epic recommends a maximum of 1024 components.
+Each component has a render-thread CPU processing cost and each section is a draw call, so try to keep these numbers to a minimum. For the largest Landscapes, Epic recommends a maximum of 1024 components.
 
 
 ## Recommended Landscape Sizes
