#UE4doc #docs: *corrected spelling "FATALSCRIPTWARNINGS"

diff --git a/Engine/Documentation/Source/Programming/Basics/CommandLineArguments/CommandLineArguments.INT.udn b/Engine/Documentation/Source/Programming/Basics/CommandLineArguments/CommandLineArguments.INT.udn
--- a/Engine/Documentation/Source/Programming/Basics/CommandLineArguments/CommandLineArguments.INT.udn
+++ b/Engine/Documentation/Source/Programming/Basics/CommandLineArguments/CommandLineArguments.INT.udn
@@ -98,7 +98,7 @@ These arguments can be passed to either the game or the editor, depending on the
 * `DEVCON`: Disable secure connections for developers. (uses unencrypted sockets)
 * `DUMPFILEIOSTATS`: Track and log File IO statistics.
 * `DUMPUDKSURVEY`: Simply dump UDK hardware survey. No upload.
-* `FATASCRIPTWARNINGS`: Treat script warnings from debugging as fatal.
+* `FATALSCRIPTWARNINGS`: Treat script warnings from debugging as fatal.
 * `FINAL_RELEASE`: Set the FINAL_RELEASE macro for script compiling.
 * `FIXEDSEED`: Initialize the random number generator with a fixed value, 0.
 * `FIXUPTANGENTS`: Fix legacy tangents in distributions automatically.
