#UE4doc #docs: changed "can not" to "cannot" "tranformations" to "transformations" capitalized "Realtime" "WORLD_MAX contant"  to "*constant" "an options" to "an option"

diff --git a/Engine/Documentation/Source/Engine/LevelStreaming/WorldBrowser/WorldBrowser.INT.udn b/Engine/Documentation/Source/Engine/LevelStreaming/WorldBrowser/WorldBrowser.INT.udn
--- a/Engine/Documentation/Source/Engine/LevelStreaming/WorldBrowser/WorldBrowser.INT.udn
+++ b/Engine/Documentation/Source/Engine/LevelStreaming/WorldBrowser/WorldBrowser.INT.udn
@@ -62,7 +62,7 @@ controlled using Blueprints.
 You can create new layers by clicking the **+** button to the right of the existing layer names.
 Each level can only be assigned to one layer at a time.  
 
-Layers can not be deleted or edited, so if you would like to change streaming distance settings, create a new layer with your desired settings and assign levels to it. 
+Layers cannot be deleted or edited, so if you would like to change streaming distance settings, create a new layer with your desired settings and assign levels to it. 
 Layers which have no assigned levels will be deleted automatically the next time you open a level. 
 
 You can filter your levels by layer by clicking on the layer name above the minimap. You can toggle multiple layers by using **Control+Mouse Click**.
@@ -80,7 +80,7 @@ level tiles, your levels must have integer bounds sizes.
 Levels based on Landscape Actors have snapping based on the base landscape component size. This is done to enable seamless 
 editing between Landscape Actors in different levels.
 
-All Landscape Actors will have tranformations disabled in the Editor Viewport to ensure that the landscape can not be moved by an arbitrary amount and 
+All Landscape Actors will have transformations disabled in the Editor Viewport to ensure that the landscape cannot be moved by an arbitrary amount and 
 lose its ability to be aligned with other landscape levels.
 [/REGION]
 
@@ -96,10 +96,10 @@ automatic level bounds calculations and set a fixed size for the Level Bounds Ac
 ![](transform.png)
 
 [REGION:note]
-Currently, realtime must be enabled in the viewports for the Level Bounds Actor to function. 
+Currently, Realtime must be enabled in the viewports for the Level Bounds Actor to function. 
 [/REGION]
 
-The minimap has a yellow square indicating the safe edit area. The safe edit area's size is equal to the `WORLD_MAX` contant in the engine.
+The minimap has a yellow square indicating the safe edit area. The safe edit area's size is equal to the `WORLD_MAX` constant in the engine.
 Levels outside of this area will be hidden automatically, and their level images will reflect this.  **World Composition** tries to 
 keep levels you are working on close to the Editor world origin. This feature is called "focusing" and may shift the current 
 world origin to a new position if you begin editing a level outside of the safe edit area.  Shifted levels still preserve their Actors' original 
@@ -172,7 +172,7 @@ Dedicated servers will load all distance dependent levels, and each connected cl
 
 ## Creating Landscape Levels
 
-**Right-clicking** on a level tile, to summon the context menu, which has a Landscape Actor inside will bring an options for adding adjacent levels with landscape proxy Actors in them. This requires the adjacent levels landscapes have size equal to source Landscape size.
+**Right-clicking** on a level tile to summon the context menu which has a Landscape Actor inside will bring an option for adding adjacent levels with landscape proxy Actors in them. This requires the adjacent levels landscapes have size equal to source Landscape size.
 
 ### Tiled Landscape Import
 
@@ -184,7 +184,7 @@ In the **Levels** panel toolbar you can find an option for importing tiled lands
 | Setting | Description |
 | ------- | ----------- |
 | Select Heightmap Tiles... | Allows you to select multiple heightmap tiles you want to import. |
-| Tile Coordinates Offset | Offset in tiles from origin. If set to zero tile with coordinates x0_y0 will be created at world origin. |
+| Tile Coordinates Offset | Offset in tiles from origin. If set to zero, tile with coordinates x0_y0 will be created at world origin. |
 | Import Configuration | Generates all possible landscape configurations depending on resolution of selected tiles. |
 | Landscape Scale | How much to scale the tiled landscape by.|
 | Material | Material to use. When a material has be chosen, you will be able to assign weightmap tiles for each landscape layer found in the material and how each layer is blended. |
