#UE4doc #docs: changed sentence to "Note: The step above will ensure that the door has collision applied to it."

diff --git a/Engine/Documentation/Source/Engine/Matinee/HowTo/MHT_1/MHT_1.INT.udn b/Engine/Documentation/Source/Engine/Matinee/HowTo/MHT_1/MHT_1.INT.udn
--- a/Engine/Documentation/Source/Engine/Matinee/HowTo/MHT_1/MHT_1.INT.udn
+++ b/Engine/Documentation/Source/Engine/Matinee/HowTo/MHT_1/MHT_1.INT.udn
@@ -57,7 +57,7 @@ In the following steps, we are going to create a simple door that opens/closes w
 
 	![](MHT1_DoorCollision.png)
 
-	_Note: This door did not have collision applied to it, the step above will ensure that it has collision on it._
+	_Note: The step above will ensure that the door has collision applied to it. _
 
 1. With the **SM_Door** selected in the **Level Viewport**, find the **Details** panel. Under **Transform** and the **Rotation** setting, set the **Z** value to **-180**. 
 
@@ -140,7 +140,7 @@ In the following steps, we are going to create a simple door that opens/closes w
 
 	![](MHT1_AddKey.png)(w:780)
 
-1. Minimize the **Matinee** window and return to the **Level Viewport**. The **SM_Door** should still be selected, press the **E** key to enter **Rotation** mode. 
+1. Minimize the **Matinee** window and return to the **Level Viewport**. With the **SM_Door** still selected, press the **E** key to enter **Rotation** mode. 
 
 	![](MHT1_RotateWidget.png)
 
