#UE4doc #docs: changed, "...will appear, click the **Apply Now** button." to "will appear. Click the **Apply Now** button."

diff --git a/Engine/Documentation/Source/Engine/Matinee/HowTo/MHT_5/MHT_5.INT.udn b/Engine/Documentation/Source/Engine/Matinee/HowTo/MHT_5/MHT_5.INT.udn
--- a/Engine/Documentation/Source/Engine/Matinee/HowTo/MHT_5/MHT_5.INT.udn
+++ b/Engine/Documentation/Source/Engine/Matinee/HowTo/MHT_5/MHT_5.INT.udn
@@ -187,7 +187,7 @@ In the following steps, we will create a light that moves, changes colors and br
 
 	![](MHT5_ApplyLighting.png)
 
-	_will appear, click the **Apply Now** button._
+	_will appear. Click the **Apply Now** button."_
 
 1. From the main toolbar, click the **Play** icon to play in the editor. 
 
