#UE4doc #docs: changed "its" to "it is" and "Curve Tension" defintion to "Controls the tightness of the curve when using automatic tangents"

diff --git a/Engine/Documentation/Source/Engine/Matinee/TrackReference/TrackReference.INT.udn b/Engine/Documentation/Source/Engine/Matinee/TrackReference/TrackReference.INT.udn
--- a/Engine/Documentation/Source/Engine/Matinee/TrackReference/TrackReference.INT.udn
+++ b/Engine/Documentation/Source/Engine/Matinee/TrackReference/TrackReference.INT.udn
@@ -30,7 +30,7 @@ When a keyframe is selected, you will see green handles in the level viewports w
 ![interp_mode_menu.png](interp_mode_menu.png)
 
 [REGION:note]
-The color of the keyframe triangle in the _Matinee_ window will indicate the current interpolation mode its using. 
+The color of the keyframe triangle in the _Matinee_ window will indicate the current interpolation mode it is using. 
 [/REGION]
 
 Below are the _Properties_ that can be adjusted for this _Track_ from the **Details** panel. 
@@ -39,7 +39,7 @@ Below are the _Properties_ that can be adjusted for this _Track_ from the **Deta
 | ------- | ------------ |
 | Active Condition | Sets when the track should be enabled; always, when gore is enabled, or when gore is disabled. |
 | Ang Curve Tension | Controls the tightness of the curve for the rotation path. |
-| Disable Movement | Allows you to disable this _Movement Track_ in preview. Useful when keyframing _Actors_ relative to another _Actor_. |
+| Disable Movement | Allows you to disable this _Movement Track_ in preview; useful when keyframing _Actors_ relative to another _Actor_. |
 | Hide 3D Track | Toggles display of the track's 3D movement path in the viewport. |
 | Lin Curve Tension | Controls the tightness of the curve for the translation path. |
 | Rot Mode | Allows you to choose the method for determining the rotation of the controlled _Actor_. |
@@ -180,7 +180,7 @@ Below are the _Properties_ that can be adjusted for this _Track_ from the **Deta
 |Property| Description |
 | ------- | ------------ |
 | Active Condition | Sets when the track should be enabled; always, when gore is enabled, or when gore is disabled. |
-| Curve Tension | Tension of curve, used for keypoints using automatic tangents. |
+| Curve Tension | Controls the tightness of the curve when using automatic tangents. |
 | Persist Fade | Enables and Disables the _Fade Track._ |
 
 ### Slomo Track
@@ -193,7 +193,7 @@ Below are the _Properties_ that can be adjusted for this _Track_ from the **Deta
 |Property| Description |
 | ------- | ------------ |
 | Active Condition | Sets when the track should be enabled; always, when gore is enabled, or when gore is disabled. |
-| Curve Tension | Tension of curve, used for keypoints using automatic tangents. |
+| Curve Tension | Controls the tightness of the curve when using automatic tangents. |
 
 
 ### Audio Master Track
@@ -204,7 +204,7 @@ Below are the _Properties_ that can be adjusted for this _Track_ from the **Deta
 |Property| Description |
 | ------- | ------------ |
 | Active Condition | Sets when the track should be enabled; always, when gore is enabled, or when gore is disabled. |
-| Curve Tension | Tension of curve, used for keypoints using automatic tangents. |
+| Curve Tension | Controls the tightness of the curve when using automatic tangents. |
 
 
 ### Color Scale Track
@@ -215,7 +215,7 @@ Below are the _Properties_ that can be adjusted for this _Track_ from the **Deta
 |Property| Description |
 | ------- | ------------ |
 | Active Condition | Sets when the track should be enabled; always, when gore is enabled, or when gore is disabled. |
-| Curve Tension | Tension of curve, used for keypoints using automatic tangents. |
+| Curve Tension | Controls the tightness of the curve when using automatic tangents. |
 
 
 ## Property Tracks
@@ -255,7 +255,7 @@ Below are the _Properties_ that can be adjusted for this _Track_ from the **Deta
 |Property| Description |
 | ------- | ------------ |
 | Active Condition | Sets when the track should be enabled; always, when gore is enabled, or when gore is disabled. |
-| Curve Tension | Tension of curve, used for keypoints using automatic tangents. |
+| Curve Tension | Controls the tightness of the curve when using automatic tangents. |
 | Property Name | Read-only variable indicating the name of the property in the connected _Actor_ which this track will modify over time. |
 
 
@@ -267,7 +267,7 @@ Below are the _Properties_ that can be adjusted for this _Track_ from the **Deta
 |Property| Description |
 | ------- | ------------ |
 | Active Condition | Sets when the track should be enabled; always, when gore is enabled, or when gore is disabled. |
-| Curve Tension | Tension of curve, used for keypoints using automatic tangents. |
+| Curve Tension | Controls the tightness of the curve when using automatic tangents. |
 | Property Name | Read-only variable indicating the name of the property in the connected _Actor_ which this track will modify over time.|
 
 
@@ -281,7 +281,7 @@ Below are the _Properties_ that can be adjusted for this _Track_ from the **Deta
 |Property| Description |
 | ------- | ------------ |
 | Active Condition | Sets when the track should be enabled; always, when gore is enabled, or when gore is disabled. |
-| Curve Tension | Tension of curve, used for keypoints using automatic tangents. |
+| Curve Tension | Controls the tightness of the curve when using automatic tangents. |
 | Property Name | Read-only variable indicating the name of the property in the connected _Actor_ which this track will modify over time. |
 
 
@@ -294,7 +294,7 @@ Below are the _Properties_ that can be adjusted for this _Track_ from the **Deta
 |Property| Description |
 | ------- | ------------ |
 | Active Condition | Sets when the track should be enabled; always, when gore is enabled, or when gore is disabled. |
-| Curve Tension | Tension of curve, used for keypoints using automatic tangents. |
+| Curve Tension | Controls the tightness of the curve when using automatic tangents. |
 | Property Name | Read-only variable indicating the name of the property in the connected _Actor_ which this track will modify over time. |
 
 
@@ -311,7 +311,7 @@ Below are the _Properties_ that can be adjusted for this _Track_ from the **Deta
 |Property| Description |
 | ------- | ------------ |
 | Active Condition | Sets when the track should be enabled; always, when gore is enabled, or when gore is disabled. |
-| Curve Tension | Tension of curve, used for keypoints using automatic tangents. |
+| Curve Tension | Controls the tightness of the curve when using automatic tangents. |
 | Target Materials | List of materials whose parameters will be modified and references to those materials that need to be given _MaterialInstanceConstants_ in the level, compiled at save time.  |
 | Param Name | Name of the parameter in the materials to modify. |
 
@@ -323,7 +323,7 @@ Below are the _Properties_ that can be adjusted for this _Track_ from the **Deta
 |Property| Description |
 | ------- | ------------ |
 | Active Condition | Sets when the track should be enabled; always, when gore is enabled, or when gore is disabled. |
-| Curve Tension | Tension of curve, used for keypoints using automatic tangents. |
+| Curve Tension | Controls the tightness of the curve when using automatic tangents. |
 | Materials | List of materials whose parameters will be modified and references to those materials that need to be given _MaterialInstanceConstants_ in the level, compiled at save time. |
 | Param Name | Name of the parameter in the materials to modify. |
 
@@ -339,7 +339,7 @@ Below are the _Properties_ that can be adjusted for this _Track_ from the **Deta
 |Property| Description |
 | ------- | ------------ |
 | Active Condition | Sets when the track should be enabled; always, when gore is enabled, or when gore is disabled. |
-| Curve Tension | Tension of curve, used for keypoints using automatic tangents. |
+| Curve Tension | Controls the tightness of the curve when using automatic tangents. |
 | Param Name | Name of the parameter in the particle system to modify. |
 
 
@@ -402,7 +402,7 @@ Below are the _Properties_ that can be adjusted for this _Track_ from the **Deta
 |Property| Description |
 | ------- | ------------ |
 | Active Condition | Sets when the track should be enabled; always, when gore is enabled, or when gore is disabled. |
-| Curve Tension | Tension of curve, used for keypoints using automatic tangents. |
+| Curve Tension | Controls the tightness of the curve when using automatic tangents. |
 | Slot Name | Name of slot to use when playing animation. Passed to _Actor_. When multiple tracks use the same slot name, they are each given a different _ChannelIndex_ when _SetAnimPosition_ is called. |
 | Skip Anim Notifiers | Skips all anim notifiers. |
 
@@ -440,7 +440,7 @@ Below are the _Properties_ that can be adjusted for this _Track_ from the **Deta
 |Property| Description |
 | ------- | ------------ |
 | Active Condition | Sets when the track should be enabled; always, when gore is enabled, or when gore is disabled. |
-| Curve Tension | Tension of curve, used for keypoints using automatic tangents. |
+| Curve Tension | Controls the tightness of the curve when using automatic tangents. |
 | Continue Sound On Matinee End | If _true_, the sounds being played on this track will not be cut off if their duration continues beyond the end of the _Matinee_ sequence. They will continue to play until the sound ends. |
 | Play On Reverse | If _true_, the sounds being played on this track will only play when the _Matinee_ sequence is played in reverse. They will not play when the sequence is playing forwards. An example of using this feature would be to have different opening and closing sounds for movers. |
 | Suppress Subtitles | If _true_, any subtitles associated with the sounds being played by this track will not be displayed. |
@@ -490,7 +490,7 @@ Below are the _Properties_ that can be adjusted for this _Track_ from the **Deta
 |Property| Description |
 | ------- | ------------ |
 | Active Condition | Sets when the track should be enabled; always, when gore is enabled, or when gore is disabled. |
-| Curve Tension | Tension of curve, used for keypoints using automatic tangents. |
+| Curve Tension | Controls the tightness of the curve when using automatic tangents. |
 | Property Name | Read-only variable indicating the name of the property in the connected _Actor_ which this track will modify over time.|
 
 
