#UE4doc #docs: sp "existing" sp "Emissive"

diff --git a/Engine/Documentation/Source/Engine/MediaFramework/MediaFramework.INT.udn b/Engine/Documentation/Source/Engine/MediaFramework/MediaFramework.INT.udn
--- a/Engine/Documentation/Source/Engine/MediaFramework/MediaFramework.INT.udn
+++ b/Engine/Documentation/Source/Engine/MediaFramework/MediaFramework.INT.udn
@@ -77,7 +77,7 @@ The Media Framework is largely a collection of C++ interfaces, a couple of helpe
 [REGION:note]
 Programmers: 
 * The code for the Media Framework is in the following location: _/Engine/Source/Runtime/Media/_
-* The exisitng player plug-ins can be found in the following location: _/Engine/Plugins/Media/_
+* The existing player plug-ins can be found in the following location: _/Engine/Plugins/Media/_
 [/REGION:note]
 
 Currently a player plug-in for Windows is being used which uses the Windows Media Foundation API under the hood. There is also a plug-in for MacOS, which uses Apple's AV Foundation, however it is currently quite limited but should work on iOS. Work has also begun on a player plug-in for Android. 
@@ -276,7 +276,7 @@ We have also added a MediaPlayer variable and assigned the MediaPlayer asset we
 
 ![](WidgetBlueprint2.png)(w:840)
 
-In order for the Material to be applied to our brush, we have to make one small change and set it to **Emmissive**. 
+In order for the Material to be applied to our brush, we have to make one small change and set it to **Emissive**. 
 
 ![](Material.png)
 
