#UE4doc #docs: sp* achieved, asterisk

diff --git a/Engine/Documentation/Source/Engine/Paper2D/Flipbooks/Paper2DFlipbooks.INT.udn b/Engine/Documentation/Source/Engine/Paper2D/Flipbooks/Paper2DFlipbooks.INT.udn
--- a/Engine/Documentation/Source/Engine/Paper2D/Flipbooks/Paper2DFlipbooks.INT.udn
+++ b/Engine/Documentation/Source/Engine/Paper2D/Flipbooks/Paper2DFlipbooks.INT.udn
@@ -71,7 +71,7 @@ You can import a JSON formatted sprite sheet description, which will import asso
 
 ### Blank Flipbooks
 
-Creating a blank Flipbook can be acheived by following the steps below. 
+Creating a blank Flipbook can be achieved by following the steps below. 
 
 **From the Content Browser**: 
 
@@ -91,7 +91,7 @@ Creating a blank Flipbook can be acheived by following the steps below.
 
 	![](FlipbookCreate3.png)
 
-	The asterik symbol in the lower-left corner indicates that the asset has not been saved yet and will go away when you save it. 
+	The asterisk symbol in the lower-left corner indicates that the asset has not been saved yet and will go away when you save it. 
 
 ### Auto-Generated Flipbooks
 
