#UE4doc #docs: *sp particular, collision, specified  -affect to effect-

diff --git a/Engine/Documentation/Source/Engine/Paper2D/TileMaps/Paper2DTilesMaps.INT.udn b/Engine/Documentation/Source/Engine/Paper2D/TileMaps/Paper2DTilesMaps.INT.udn
--- a/Engine/Documentation/Source/Engine/Paper2D/TileMaps/Paper2DTilesMaps.INT.udn
+++ b/Engine/Documentation/Source/Engine/Paper2D/TileMaps/Paper2DTilesMaps.INT.udn
@@ -117,7 +117,7 @@ In addition to setting Tile Width and Tile Height, you can use the following set
 
 ### Per Tile Settings / Applying Tile Collision
 
-Also inside the **Details** panel, you can specify per tile settings (such as applying collision to a particluar tile). 
+Also inside the **Details** panel, you can specify per tile settings (such as applying collision to a particular tile). 
 
 **To apply collision to a tile**: 
 
@@ -129,7 +129,7 @@ Also inside the **Details** panel, you can specify per tile settings (such as ap
 
 	![](Collision_2.png)
 
-	Note the tile number you are editing (green box above), you will need this to apply settings to this particluar tile. 
+	Note the tile number you are editing (green box above), you will need this to apply settings to this particular tile. 
 
 1. In the **Details** panel under the **Sprite** section, find the selected tile and change to the desired **Geometry Type**. 
 
@@ -139,7 +139,7 @@ Also inside the **Details** panel, you can specify per tile settings (such as ap
 	For information on sprite collision types, see the [](Engine\Paper2D\Sprites\Collision) documentation. 
 	[/REGION]
 	
-This will set this tile up for collsion based on the specificed geometry type and shape used. For the collision to actually take affect in-game, when creating your **Tile Map** you will need to set the **Sprite Collision Domain** (see below for information on creating/using Tile Maps). 
+This will set this tile up for collision based on the specified geometry type and shape used. For the collision to actually take effect in-game, when creating your **Tile Map** you will need to set the **Sprite Collision Domain** (see below for information on creating/using Tile Maps). 
 
 
 ## Creating Tile Maps
