#UE4doc #docs: sp re-enabling, tessellation, added: allows [you] to disable

diff --git a/Engine/Documentation/Source/Engine/Performance/PerformanceProfiling.INT.udn b/Engine/Documentation/Source/Engine/Performance/PerformanceProfiling.INT.udn
--- a/Engine/Documentation/Source/Engine/Performance/PerformanceProfiling.INT.udn
+++ b/Engine/Documentation/Source/Engine/Performance/PerformanceProfiling.INT.udn
@@ -122,11 +122,11 @@ viewports, the editor UI needs to be used. You can use the console variables (e.
 
 [REGION:note]
 Some features still consume performance even if the feature is not rendering any more, e.g. **show particles** hides the particles but simulation time is still
-required to support reenabling them at a later time. The console variable **fx.freezeparticlesimulation** allows to disable the update part as well (editor UI also exists).
-Another one: **show tesselation** disables triangle amplification but still uses the tesselation shaders.
+required to support re-enabling them at a later time. The console variable **fx.freezeparticlesimulation** allows to you disable the update part as well (editor UI also exists).
+Another one: **show tessellation** disables triangle amplification but still uses the tessellation shaders.
 [/REGION:note]
 
-A good profiling starting point is with high-level features, such as **show StaticMeshes** or **show tesselation**.
+A good profiling starting point is with high-level features, such as **show StaticMeshes** or **show tessellation**.
 
 All show flags are also exposed as console variables, e.g. console **show Bloom**, **showflag.Bloom 0** or in configuration files: **showflag.Bloom = 0**
 Console variables require more typing but they also override the editor UI showflags and they can be put into configuration files like other console variables.
@@ -137,21 +137,21 @@ The most useful ones for profiling:
 | -------- | ----------- |
 | **ScreenSpaceReflections** | Toggles screen space reflections, can cost a lot of performance, only affects pixels up to a certain roughness (adjusted with **r.SSR.MaxRoughness** or in postprocess settings).
 | **AmbientOcclusion** | Screen Space Ambient Occlusion (Some scenes only benfit very little, for static objects you bake AO in lightmass).
-| **AntiAliasing** | Toggle any anti aliasing (TemporalAA and FXAA), use **TemporalAA** to switch to FXAA (faster, lower quality).
+| **AntiAliasing** | Toggle any antialiasing (TemporalAA and FXAA), use **TemporalAA** to switch to FXAA (faster, lower quality).
 | **Bloom** | Affects image based lens flares and the bloom feature.
 | **DeferredLighting** | Toggle all deferred lighting passes.
 | **DirectionalLights** **PointLights** **SpotLights** | Toggles the different light types (useful to see what kind of light type is having what effect and performance impact.
 | **DynamicShadows** | Toggles all dynamic shadows (shadowmap rendering and shadow filtering/projection).
 | **GlobalIllumination** | Toggles baked and dynamic indirect lighting (LPV).
 | **LightFunctions** | Toggles the light functions rendering.
-| **PostProcessing** | Toggles all postprocessing passes.
+| **PostProcessing** | Toggles all post-processing passes.
 | **ReflectionEnvironment** | Toggles reflection environment passes.
 | **Refraction** | Toggles the refraction pass.
 | **Rendering** | Toggles rendering altogether.
 | **Decals** | Toggle decal rendering.
 | **Landscape** **Brushes** **StaticMeshes** **SkeletalMeshes** **Landscape** | Toggles what geometry is rendered.
 | **Translucency** | Toggles translucency rendering.
-| **Tessellation** | Toggles tesslation (still runs tessellation shaders but spawning more triangles).
+| **Tessellation** | Toggles tessellation (still runs tessellation shaders but spawning more triangles).
 | **IndirectLightingCache** | Toggles if dynamic objects or static objects that have invalidated lightmaps are using the Indirect Lighting Cache.
 | **Bounds** | Show the bounding volume of the objects selected in the editor.
 | **Visualize SSR** | Renders all pixels affected by screen space reflections is bright orange (slow), see below.
