#UE4doc #docs: added "to" sp: re-render, split screen, aggressive

diff --git a/Engine/Documentation/Source/Engine/Performance/CPU/CPU.INT.udn b/Engine/Documentation/Source/Engine/Performance/CPU/CPU.INT.udn
--- a/Engine/Documentation/Source/Engine/Performance/CPU/CPU.INT.udn
+++ b/Engine/Documentation/Source/Engine/Performance/CPU/CPU.INT.udn
@@ -9,7 +9,7 @@ If you are CPU bound in the render thread, it is likely because of too many draw
 have to combine draw calls to reduce the cost for that (e.g. combine multiple walls into one mesh). The actual cost is in many areas:
 * The render thread needs to process each object (culling, material setup, lighting setup, collision, update cost, etc.).
   More complex materials result in a higher setup cost.
-* The render thread needs to prepare the GPU commands to set up the state for each draw call (constant buffers, textures, instance properties, shaders) and do the actual API call.
+* The render thread needs to prepare the GPU commands to set up the state for each draw call (constant buffers, textures, instance properties, shaders) and to do the actual API call.
   Base pass draw calls are usually more costly than depth only draw calls.
 * DirectX validates some data and passes the information to the graphics card driver.
 * The driver (e.g. NVIDIA, AMD, Intel, ...) validates further and creates a command buffer for the hardware. Sometimes this part is split in another thread.
@@ -18,9 +18,9 @@ have to combine draw calls to reduce the cost for that (e.g. combine multiple wa
 * Reduce object count (static/dynamic meshes, mesh particles)
 * Reduce view distance (e.g. on the Scene Capture Actor or per object)
 * Adjust the view (more zoomed out view, moving objects to not be in the same view)
-* Avoid SceneCaptureActor (needs to rerender the scene, set fps to be low or update only when needed)
-* Avoid splitscreen (Always more CPU bound than single view, needs custom scalability settings or content to be more agressive)
-* Elements per draw calls (combine materials accepting a more complex pixel shaders or simply have less materials, combine textures to fewer larger textures - only if that reduces material count, use LOD models with fewer elements)
+* Avoid SceneCaptureActor (needs to re-render the scene, set fps to be low or update only when needed)
+* Avoid split screen (Always more CPU bound than single view, needs custom scalability settings or content to be more aggressive)
+* Elements per draw calls (combine materials accepting more complex pixel shaders or simply have less materials, combine textures to fewer larger textures - only if that reduces material count, use LOD models with fewer elements)
 * Disable features on the mesh like custom depth or shadow casting
 * Changing light sources to not shadow cast or having a tighter bounding volume (view cone, attenuation radius)
 
