#UE4doc #docs: sp: visibility

diff --git a/Engine/Documentation/Source/Engine/Performance/Guidelines/PerformanceGuidelines.INT.udn b/Engine/Documentation/Source/Engine/Performance/Guidelines/PerformanceGuidelines.INT.udn
--- a/Engine/Documentation/Source/Engine/Performance/Guidelines/PerformanceGuidelines.INT.udn
+++ b/Engine/Documentation/Source/Engine/Performance/Guidelines/PerformanceGuidelines.INT.udn
@@ -36,7 +36,7 @@ Here are some general guidelines for content creation and level design to improv
 * Light functions cost extra (actual cost depends on the material) and can prevent the light being rendered as a tiled light.
 * IES profiles cost extra (less than light functions) and can prevent the light being rendered as a tiled light. Do not use IES where spotlight cone angles could create the same effect.
 * Billboards, imposter meshes, or skybox textures can be used to efficiently fake detailed geometry.
-* Good level design takes occlusion culling into account (add visbility blockers for better performance). Use **r.VisualizeOccludedPrimitives** to check this.
+* Good level design takes occlusion culling into account (add visibility blockers for better performance). Use **r.VisualizeOccludedPrimitives** to check this.
 * Avoid Light Propagation Volumes or limit its cost by using the GIReplace material expression or disabling it on most objects. 
 * Disable shadow casting where possible, either per object or per light.
 * Use ProfileGPU (**Ctrl + Shift + ,**) in the editor to get quick information on what is slow.
