#UE4doc #docs: capitalized: Tonemapper, SceneColor, non-Windows

diff --git a/Engine/Documentation/Source/Engine/Performance/Options/Options.INT.udn b/Engine/Documentation/Source/Engine/Performance/Options/Options.INT.udn
--- a/Engine/Documentation/Source/Engine/Performance/Options/Options.INT.udn
+++ b/Engine/Documentation/Source/Engine/Performance/Options/Options.INT.udn
@@ -26,7 +26,7 @@ The most useful ones for profiling:
 | **Slomo** | Allows to slow down or accelerate time. |
 | **r.VisualizeOccludedPrimitives** | Displays the bounding boxes of occluded primitives.|
 | **StartFPSChart** **StopFPSChart** | See below.|
-| **r.SeparateTranslucency** | The feature fixes transluceny with depth of field, can be disabled if not needed, has other implications (reading SceneColor).|
+| **r.SeparateTranslucency** | The feature fixes translucency with depth of field, can be disabled if not needed, has other implications (reading SceneColor).|
 | **r.TonemapperQuality** | To disable the noise we add in the tonemapper to avoid color banding because of 8bit quantization, minor quality improvement, not needed when we output in 10:10:10.|
 | **r.SceneColorFormat** | Allows to use a different scenecolor format (default is 64bit for best quality and to support screenspace subsurface scattering).|
 | **FX.AllowGPUSorting** | To disable particle sorting (pays off when used for a lot of particles).|
