#UE4doc #docs: updated preview link pic

diff --git a/Engine/Documentation/Source/Engine/Performance/Scalability/ScalabilityReference/ScalabilityReference.INT.udn b/Engine/Documentation/Source/Engine/Performance/Scalability/ScalabilityReference/ScalabilityReference.INT.udn
--- a/Engine/Documentation/Source/Engine/Performance/Scalability/ScalabilityReference/ScalabilityReference.INT.udn
+++ b/Engine/Documentation/Source/Engine/Performance/Scalability/ScalabilityReference/ScalabilityReference.INT.udn
@@ -6,7 +6,7 @@ Description:Scalability options, properties, and console variables.
 [VAR:Topic]
 [OBJECT:Topic]
 	[PARAM:image]
-		![%Engine/Performance/Scalability/ScalabilityReference:title%](Engine/Rendering/rendering_topic.png)
+		![%Engine/Performance/Scalability/ScalabilityReference:title%](Engine/Rendering/scalability_topic.png)
 	[/PARAM]
 	[PARAM:icon]
 		![](%ROOT%/light_icon.png)(convert:false)
@@ -25,7 +25,7 @@ Description:Scalability options, properties, and console variables.
 [VAR:TopicCompact]
 [OBJECT:TopicCompact]
 	[PARAM:image]
-		![%Engine/Performance/Scalability/ScalabilityReference:title%](Engine/Rendering/rendering_topic.png)
+		![%Engine/Performance/Scalability/ScalabilityReference:title%](Engine/Rendering/scalability_topic.png)
 	[/PARAM]
 	[PARAM:icon]
 		![](%ROOT%/light_icon.png)(convert:false)
