#UE4doc #docs: *fixed: modern notebook[s] have

diff --git a/Engine/Documentation/Source/Engine/Performance/Scalability/ScalabilityAndYou/ScalabilityAndYou.INT.udn b/Engine/Documentation/Source/Engine/Performance/Scalability/ScalabilityAndYou/ScalabilityAndYou.INT.udn
--- a/Engine/Documentation/Source/Engine/Performance/Scalability/ScalabilityAndYou/ScalabilityAndYou.INT.udn
+++ b/Engine/Documentation/Source/Engine/Performance/Scalability/ScalabilityAndYou/ScalabilityAndYou.INT.udn
@@ -62,7 +62,7 @@ Different machines can have different performance/memory/quality characteristics
 * Main memory size/latency/bandwidth, running 32bit instead of 64bit can limit this further
 * CPU (one or more main processors, hyper-threading,  SSE features, branch prediction)
 * GPU (memory, Multi GPU, bandwidth, features, shared memory)
-* Resolution (modern notebook have high resolution displays but slow GPUs)
+* Resolution (modern notebooks have high resolution displays but slow GPUs)
 
 
 Graphics is generally the easiest to scale without affecting gameplay much. The players generally accept the concept of scalability but only very few actively go to the options menu and adjust them properly. Option screens are often overwhelming in quantity of options and full of intimidating names.  Still, there are options to give the player control over a lot of decisions that are subjective and have tradeoffs:
