#UE4doc #docs: updated preview link pic

diff --git a/Engine/Documentation/Source/Engine/Performance/Scalability/Scalability.INT.udn b/Engine/Documentation/Source/Engine/Performance/Scalability/Scalability.INT.udn
--- a/Engine/Documentation/Source/Engine/Performance/Scalability/Scalability.INT.udn
+++ b/Engine/Documentation/Source/Engine/Performance/Scalability/Scalability.INT.udn
@@ -25,7 +25,7 @@ Description:Options for adjusting quality and performance.
 [VAR:TopicCompact]
 [OBJECT:TopicCompact]
 	[PARAM:image]
-		![%Engine/Performance/Scalability:title%](Engine/Rendering/rendering_topic.png)
+		![%Engine/Performance/Scalability/ScalabilityReference:title%](Engine/Performance/Scalability/ScalabilityReference/scalability_topic.png)
 	[/PARAM]
 	[PARAM:icon]
 		![](%ROOT%/light_icon.png)(convert:false)
@@ -55,5 +55,5 @@ The following sections will describe the features Unreal Engine 4 implements for
 
 [REGION:topics third]
 %Engine/Performance/Scalability/ScalabilityAndYou:topiccompact%
-%Engine/Performance/Scalability/ScalabilityReference:topiccompact%
+%Engine/Performance/Scalability/ScalabilityReference:TopicCompact%
 [/REGION]
