#UE4doc #docs: edited code to fix preview link on page "Scalability Index" /deleted duplicate "the"/ sp: from/capitalized Perlin/

diff --git a/Engine/Documentation/Source/Engine/Performance/Scalability/ScalabilityReference/ScalabilityReference.INT.udn b/Engine/Documentation/Source/Engine/Performance/Scalability/ScalabilityReference/ScalabilityReference.INT.udn
--- a/Engine/Documentation/Source/Engine/Performance/Scalability/ScalabilityReference/ScalabilityReference.INT.udn
+++ b/Engine/Documentation/Source/Engine/Performance/Scalability/ScalabilityReference/ScalabilityReference.INT.udn
@@ -6,7 +6,7 @@ Description:Scalability options, properties, and console variables.
 [VAR:Topic]
 [OBJECT:Topic]
 	[PARAM:image]
-		![%Engine/Performance/Scalability/ScalabilityReference:title%](Engine/Rendering/scalability_topic.png)
+		![%Engine/Performance/Scalability/ScalabilityReference:title%](Engine/Performance/Scalability/ScalabilityReference/scalability_topic.png)
 	[/PARAM]
 	[PARAM:icon]
 		![](%ROOT%/light_icon.png)(convert:false)
@@ -25,7 +25,7 @@ Description:Scalability options, properties, and console variables.
 [VAR:TopicCompact]
 [OBJECT:TopicCompact]
 	[PARAM:image]
-		![%Engine/Performance/Scalability/ScalabilityReference:title%](Engine/Rendering/scalability_topic.png)
+		![%Engine/Performance/Scalability/ScalabilityReference:title%](Engine/Performance/Scalability/ScalabilityReference/scalability_topic.png)
 	[/PARAM]
 	[PARAM:icon]
 		![](%ROOT%/light_icon.png)(convert:false)
@@ -54,7 +54,7 @@ The Scalability settings allow you to adjust the quality of various features in
 
 ## Scalability Settings
 
-Accessing the Scalability settings in the editor is done through the **Quick Settings** menu in the toolbar ![](%ROOT%/Shared/Icons/icon_view_40x.png). The **Engine Scalability Settings** contains short cuts to the most commonly used setting defined in the the `BaseScalability.ini` file, while **Material Quality Level** sets the global material quality setting.
+Accessing the Scalability settings in the editor is done through the **Quick Settings** menu in the toolbar ![](%ROOT%/Shared/Icons/icon_view_40x.png). The **Engine Scalability Settings** contains short cuts to the most commonly used setting defined in the `BaseScalability.ini` file, while **Material Quality Level** sets the global material quality setting.
 
 [REGION:lightbox]
 [![](ScaleR_ViewQual.png)(h:490)](ScaleR_ViewQual.png) [![](ScaleR_ViewMatQual.png)(h:490)](ScaleR_ViewMatQual.png)
@@ -255,7 +255,7 @@ The **View** -> **Engine Scalability Settings** -> **Shadows** option adjusts th
 
 ## Textures - sg.TextureQuality
 
-A modern rendering engine requires a lot more GPU memory (textures, meshes, GBuffer, Depth Buffer, Shadow maps). Some of those scale with the screen resolution (e.g. GBuffer), others with specific quality settings (e.g. Shadow maps).  Another large amount comes form the used textures (usually compressed and streamed). You can instruct the streaming system to be more aggressive in management (smaller pool size, culling unused textures) or to have less or more detail in the mip level computation. This can have effects on the image quality, how much you can notice texture streaming artifacts and how smooth the game runs (updates require expensive memory transfers). The results can vary, depending on the media (e.g. faster/lower hard drive /SSD). Streaming from a DVD/Blu-Ray adds much more latency so you should try to avoid that.
+A modern rendering engine requires a lot more GPU memory (textures, meshes, GBuffer, Depth Buffer, Shadow maps). Some of those scale with the screen resolution (e.g. GBuffer), others with specific quality settings (e.g. Shadow maps).  Another large amount comes from the used textures (usually compressed and streamed). You can instruct the streaming system to be more aggressive in management (smaller pool size, culling unused textures) or to have less or more detail in the mip level computation. This can have effects on the image quality, how much you can notice texture streaming artifacts and how smooth the game runs (updates require expensive memory transfers). The results can vary, depending on the media (e.g. faster/lower hard drive /SSD). Streaming from a DVD/Blu-Ray adds much more latency so you should try to avoid that.
 
 The texture quality also affects the texture filtering mode (r.MaxAnisotropy). Limiting the anisotropic sample count reduces texture bandwidth, but does not save texture memory. 
 
@@ -389,7 +389,7 @@ Material can use the "Quality Switch" material expression to disable some expens
 [![](SC_SwitchNode.png)(w:900)](SC_SwitchNode.png)
 [/REGION]
 
-Whatever the **Material Quality Level** is set to, low or high, will determine which expressions get evaluated for this material (the low or high pin). The default pin will fill in either high or low (or both) if they have no inputs. This material contains 2 reasonably high cost perlin noise operations when set to high:
+Whatever the **Material Quality Level** is set to, low or high, will determine which expressions get evaluated for this material (the low or high pin). The default pin will fill in either high or low (or both) if they have no inputs. This material contains 2 reasonably high cost Perlin noise operations when set to high:
 
 [REGION:note]
 This example is to simply show off the **Material Switch Node**. The **Noise** nodes are really expensive making them really good for this example, but they should be used sparingly as there are cheaper ways to get this effect.
@@ -416,7 +416,7 @@ This feature is not for distance LOD as you cannot have two quality levels at th
 
 * Shader computations (e.g. disable fuzz layer).
 	* Use **Shader Complexity** (**Alt+8** in the editor) mode and instruction count to verify your optimization.
-* Texture lookups (e.g. no detail bumpmap).
+* Texture lookups (e.g. no detail bump map).
 * Memory bandwidth (e.g. using less textures).
 	* You would need to profile on actual hardware to verify this.
 
