changed "can simulate" to "can be simulated" fixed line to: "it's collision properties, mass, if it is sleeping, " "when ever" to "whenever,"

diff --git a/Engine/Documentation/Source/Engine/Physics/Collision/Collision.INT.udn b/Engine/Documentation/Source/Engine/Physics/Collision/Collision.INT.udn
--- a/Engine/Documentation/Source/Engine/Physics/Collision/Collision.INT.udn
+++ b/Engine/Documentation/Source/Engine/Physics/Collision/Collision.INT.udn
@@ -58,12 +58,12 @@ In the properties list, there are a large number of physics related key words.
 | ^ | _WorldStatic_ - Volumes and world geometry should be _WorldStatic_. | 
 | ^ | _WorldDynamic_ - Moving Actors outside of _Pawns_, _PhysicsBodies_, _Vehicles_, and _Destructible_ Actors. Like an elevator. | 
 | ^ | _Pawn_ - Characters. | 
-| ^ | _PhysicsBody_ - Any physics object that is or can simulate in the world. | 
+| ^ | _PhysicsBody_ - Any physics object that is or can be simulated in the world. | 
 | ^ | _Vehicle_ - A good response to have so you can have _Pawns_ jump into them. | 
 | ^ | _Destructible_ - Destructible Actors. |
 | [REGION:tablesection] &nbsp; [/REGION]||
 | **Use CCD** | Whether or not to use **Continuous Collision Detection** for this component. Increases accuracy of collision detection. |
-| **Always Create Physics State** | Indicates if you would like to create a **physics state** (it collision properties, mass, if it is sleepting, etc... all go into a **physics state**) all the time (for collision and simulation). If set to _false_, the state will still be generated if collision or simulation are activated on this Actor, say through a Blueprint or code, but there will be some overhead as it has to work the new object into the full physics simulation. Setting this to **true** can improve performance by removing the overhead for calculating the physics state of the object once it actually needs one (on collision or simulation). | 
+| **Always Create Physics State** | Indicates if you would like to create a **physics state** (it's collision properties, mass, if it is sleeping, etc... all go into a **physics state**) all the time (for collision and simulation). If set to _false_, the state will still be generated if collision or simulation are activated on this Actor, say through a Blueprint or code, but there will be some overhead as it has to work the new object into the full physics simulation. Setting this to **true** can improve performance by removing the overhead for calculating the physics state of the object once it actually needs one (on collision or simulation). | 
 | **Multi Body Overlap** | If **true**, this component will generate individual overlaps for each overlapping physics body if it is a multi-body component. Think about skeletal physics assets, with all their individual collision shapes. With this enabled, a character's hand would generate _its own_ overlap event. Thus giving more control over what is reported and what to do in that specific case. |
 | **Check Async Scene on Move** | If this is flagged as **true**, the component will look for collisions in both of the physics scenes (synchronous and asynchronous). The asynchronous scene is primarily used by the broken pieces of destructible Actors. |
 | **Trace Complex on Move** | If this is set to **true**, component sweeps with this component will trace against complex collision during movement. **Complex Collision** is, in short, per face collision, where as **Simple Collision** are your spheres, capsules, boxes, and generated convex shapes. |
