#UE4doc #docs: changed "can simulate" to "can be simulated" fixed line to: "it's collision properties, mass, if it is sleeping, " "when ever" to "whenever,"

diff --git a/Engine/Documentation/Source/Engine/Physics/Collision/Collision.INT.udn b/Engine/Documentation/Source/Engine/Physics/Collision/Collision.INT.udn
--- a/Engine/Documentation/Source/Engine/Physics/Collision/Collision.INT.udn
+++ b/Engine/Documentation/Source/Engine/Physics/Collision/Collision.INT.udn
@@ -79,7 +79,7 @@ Some components will have additional or fewer options in the collision category,
 | ^ | _Use Simple Collision as Complex_ - Will use convex collision for everything. |
 | ^ | _Use Complex Collision as Simple_ - Will use per-poly collision for even simple collision needs. |
 | [REGION:tablesection] **Physics Body in Physics Asset** [/REGION] ||
-| **Colision Response** | Whether or not a specific physics body in a Physics Asset will collide at all. |
+| **Collision Response** | Whether or not a specific physics body in a Physics Asset will collide at all. |
 
 ## Interactions
 
@@ -146,11 +146,11 @@ Just collision is useful and in general, the bare minimum for physics interactio
 | Sphere Collision Setup| Wall Collision Setup |
 | --- | ---- |
 | ![](COL_collideEvent_Sphere.png)(w:300) %spacer%|![](COL_collideNoEvent_Box.png)(w:301) %spacer%|
-| As in the example above, the sphere is a **PhysicsBody** and it is set to `block` **WorldDynamic** (which is what the wall is). However, the sphere has also enabled **Simulation Generates Hit Events** so it will trigger an event for itself when ever it collides with something. | The wall is a **WorldDynamic** and is set to `block` **PhysicsBody** Actors (which is what the sphere is). Since the wall is not set to **Simulation Generates Hit Events**, it will not generate an event for itself. |
+| As in the example above, the sphere is a **PhysicsBody** and it is set to `block` **WorldDynamic** (which is what the wall is). However, the sphere has also enabled **Simulation Generates Hit Events** so it will trigger an event for itself whenever it collides with something. | The wall is a **WorldDynamic** and is set to `block` **PhysicsBody** Actors (which is what the sphere is). Since the wall is not set to **Simulation Generates Hit Events**, it will not generate an event for itself. |
 
 With the sphere set to **Simulation Generates Hit Events**, the sphere will tell itself that it has had a collision. It will fire off events such as **ReceiveHit** or **OnComponentHit** in the sphere's blueprint. Now if the box had an event for collision, it would not fire because it will never notify itself it has happened.
 
-Further, an object that is reporting rigid collisions, will report them all and spam reports when it is just sitting on something, so it is best to be careful to filter what it is colliding within its blueprint or in code.
+Further, an object that is reporting rigid collisions will report them all and spam reports when it is just sitting on something, so it is best to be careful to filter what it is colliding within its blueprint or in code.
 
 ### Overlap and Ignore
 
@@ -164,14 +164,14 @@ For all intents and purposes, **Overlap** and **Ignore** work exactly the same *
 | Sphere Collision Setup| Wall Collision Setup |
 | --- | ---- |
 | ![](COL_OverlapNoEvent_Sphere.png)(w:300) %spacer% |![](COL_collideNoEvent_Box.png)(w:303) %spacer% |
-| Here the sphere is set to `overlap` **WorldDynamic** Actors (like our wall), but it does not have **Generate Overlap Events** enabled. As far as the sphere is concerned, it has not collided or overlapped anything, effectively it has ignored the wall.  | The wall is a **WorldDynamic** and is set to `block` **PhysicsBody** Actors (which is what the sphere is). As stated above, both Actors need to be set to block each others respective object types. If they do not, they will not collide. |
+| Here the sphere is set to `overlap` **WorldDynamic** Actors (like our wall), but it does not have **Generate Overlap Events** enabled. As far as the sphere is concerned, it has not collided or overlapped anything, effectively it has ignored the wall.  | The wall is a **WorldDynamic** and is set to `block` **PhysicsBody** Actors (which is what the sphere is). As stated above, both Actors need to be set to block each other's respective object types. If they do not, they will not collide. |
 
 Or:
 
 | Sphere Collision Setup| Wall Collision Setup |
 | --- | ---- |
 | ![](COL_ignore_sphere.png)(w:300) %spacer% |![](COL_collideNoEvent_Box.png)(w:303) %spacer% |
-| Here the sphere is set to `ignore` **WorldDynamic** Actors (like our wall), and it will pass through the wall. | The wall is a **WorldDynamic** and is set to `block` **PhysicsBody** Actors (which is what the sphere is). As stated above, both Actors need to be set to block each others respective object types. If they do not, they will not collide. |
+| Here the sphere is set to `ignore` **WorldDynamic** Actors (like our wall), and it will pass through the wall. | The wall is a **WorldDynamic** and is set to `block` **PhysicsBody** Actors (which is what the sphere is). As stated above, both Actors need to be set to block each other's respective object types. If they do not, they will not collide. |
 
 
 
@@ -184,7 +184,7 @@ Unlike collisions that can fire every frame, the overlap events are **ReceiveBeg
 | Sphere Collision Setup| Wall Collision Setup |
 | --- | ---- |
 | ![](COL_OverlapEvent_Sphere.png)(w:300) %spacer% | ![](COL_collideNoEvent_Box.png)(w:303) %spacer% |
-| Here the sphere is set to `overlap` **WorldDynamic** Actors (like our wall), and it will generate an event for itself when it does overlap something. | The wall is a **WorldDynamic** and is set to `block` **PhysicsBody** Actors (which is what the sphere is). As stated above, both Actors need to be set to block each others respective object types. If they do not, they will not collide. But, an **Overlap** does occur here, and an event for the sphere is fired. |
+| Here the sphere is set to `overlap` **WorldDynamic** Actors (like our wall), and it will generate an event for itself when it does overlap something. | The wall is a **WorldDynamic** and is set to `block` **PhysicsBody** Actors (which is what the sphere is). As stated above, both Actors need to be set to block each other's respective object types. If they do not, they will not collide. But, an **Overlap** does occur here, and an event for the sphere is fired. |
 
  
 ## Custom Object Channels and Trace Channels
