#UE4doc #docs: changed "affect" to "effect"

diff --git a/Engine/Documentation/Source/Engine/Physics/Constraints/DampingAndFriction/DampingAndFriction.INT.udn b/Engine/Documentation/Source/Engine/Physics/Constraints/DampingAndFriction/DampingAndFriction.INT.udn
--- a/Engine/Documentation/Source/Engine/Physics/Constraints/DampingAndFriction/DampingAndFriction.INT.udn
+++ b/Engine/Documentation/Source/Engine/Physics/Constraints/DampingAndFriction/DampingAndFriction.INT.udn
@@ -32,7 +32,7 @@ The Damping for Physics Constraints work a bit differently than with Physics Bod
 ![](Phys_ConstraintDampProps.png)(convert:false)
 [/REGION]
 
-_The two Angular Swing Motions share Damping values, the Angular Twist receives its own Damping values, and all the Linear axes use the same Damping values. While they may have separate properties, the properties all have the same affect on the Motions they relate to._
+_The two Angular Swing Motions share Damping values, the Angular Twist receives its own Damping values, and all the Linear axes use the same Damping values. While they may have separate properties, the properties all have the same effect on the Motions they relate to._
 
 | Property | Description |
 | -------- | ----------- |
