#UE4doc #docs: changed "rayCast" to "Raycast"

diff --git a/Engine/Documentation/Source/Engine/Physics/Destructibles/DestructibleProperties/DestructibleProperties.INT.udn b/Engine/Documentation/Source/Engine/Physics/Destructibles/DestructibleProperties/DestructibleProperties.INT.udn
--- a/Engine/Documentation/Source/Engine/Physics/Destructibles/DestructibleProperties/DestructibleProperties.INT.udn
+++ b/Engine/Documentation/Source/Engine/Physics/Destructibles/DestructibleProperties/DestructibleProperties.INT.udn
@@ -64,7 +64,7 @@ Description:Descriptions of all the properties on Destructible Actors in the Des
 | **Debris Timeout** | Whether or not chunks at or deeper than the "debris" depth will time out.  | 
 | **Debris Max Separation** | Whether or not chunks at or deeper than the "debris" depth will be removed if they separate too far from their origins.
 | **Crumble Smallest Chunks** | If set, the smallest chunks may be further broken down, either by fluid crumbles (if a crumble particle system is specified in the NxDestructibleActorDesc), or by simply removing the chunk if no crumble particle system is specified. |
-| **Accurate Raycasts** | If set, the rayCast function will search within the nearest visible chunk hit for collisions with child chunks.  This is used to get a better raycast position and normal, in case the parent collision volume does not tightly fit the graphics mesh.  The returned chunk index will always be that of the visible parent that is intersected, however. |
+| **Accurate Raycasts** | If set, the Raycast function will search within the nearest visible chunk hit for collisions with child chunks.  This is used to get a better raycast position and normal, in case the parent collision volume does not tightly fit the graphics mesh.  The returned chunk index will always be that of the visible parent that is intersected, however. |
 | **Use Valid Bounds** | If set, the validBounds field of NxDestructibleParameters will be used.  These bounds are translated	(but not scaled or rotated) to the origin of the destructible Actor.  If a chunk or chunk island moves outside of those bounds, it is destroyed. |
 | **Form Extended Structures** | If this is set and the destructible is initially static, it will become part of an extended support structure if it is in contact with another static destructible that also has this flag set. |
 | [REGION:tablesection] _Fracture Effects - Per Depth Level_ [/REGION] ||
