#UE4doc #docs: sp animation

diff --git a/Engine/Documentation/Source/Engine/Animation/AnimationRetargeting/AnimationRetargeting.INT.udn b/Engine/Documentation/Source/Engine/Animation/AnimationRetargeting/AnimationRetargeting.INT.udn
--- a/Engine/Documentation/Source/Engine/Animation/AnimationRetargeting/AnimationRetargeting.INT.udn
+++ b/Engine/Documentation/Source/Engine/Animation/AnimationRetargeting/AnimationRetargeting.INT.udn
@@ -26,7 +26,7 @@ Please see [Animation Retargeting (Different Skeletons)](Engine\Animation\Retarg
 
 ## Why Use Retargeting?
 
-Consider three characters that you would like to share aniamtion: A base character, a short stocky version, and a tall skinny version.
+Consider three characters that you would like to share animation: A base character, a short stocky version, and a tall skinny version.
 
 [VAR:Chars]w:300[/VAR]
 
