changed "a" to "an" before Animation

diff --git a/Engine/Documentation/Source/Engine/Animation/AnimBlueprints/VimBlueprints.INT.udn b/Engine/Documentation/Source/Engine/Animation/AnimBlueprints/VimBlueprints.INT.udn
--- a/Engine/Documentation/Source/Engine/Animation/AnimBlueprints/VimBlueprints.INT.udn
+++ b/Engine/Documentation/Source/Engine/Animation/AnimBlueprints/VimBlueprints.INT.udn
@@ -54,7 +54,7 @@ A simple example of an **Animation Blueprint** is demonstrated on the [Animation
 ## Working with Animation Blueprints
 
 [EXCERPT:Graphs]
-There are two main components in a _Animation Blueprint_ that work in conjunction to create the final animation 
+There are two main components in an _Animation Blueprint_ that work in conjunction to create the final animation 
 for each frame. The **EventGraph** can perform updates to values which can then be used in the **AnimGraph** 
 to drive _State Machines_, _BlendSpaces_, or other nodes allowing blending between multiple _AnimationSequences_ or poses 
 that can fire off notifications to other systems to enable dynamic animation-driven effects to occur.
