#UE4doc #docs: sp: Constraints

diff --git a/Engine/Documentation/Source/Engine/Physics/PhAT/UserGuide/UserGuide.INT.udn b/Engine/Documentation/Source/Engine/Physics/PhAT/UserGuide/UserGuide.INT.udn
--- a/Engine/Documentation/Source/Engine/Physics/PhAT/UserGuide/UserGuide.INT.udn
+++ b/Engine/Documentation/Source/Engine/Physics/PhAT/UserGuide/UserGuide.INT.udn
@@ -202,7 +202,7 @@ The use of Physics Constraints is covered in the [Physics Constraints User Guide
 1. Enable Constraint Editing Mode ![](Shared\Icons\icon_PHatMode_Joint_40x.png).
 1. Select a Physics Constraint in the Viewport or in the **Hierarchy** Pane.
 1. Move and rotate the Physics Constraint using the Translation and Rotation to create the rotational point of the "joint" the Physics Constraint will form.
-1. Edit the Physics Constraints's properties in the **Details** Pane.
+1. Edit the Physics Constraint's properties in the **Details** Pane.
 
 	[REGION:tip]
 	You can switch Swing1, Swing2, and Twist from Limited to Locked quickly by using the "**1**", "**2**", and "**3**" keys respectively. The "**4**" key can be used cycle through them setting one to limited and the other two to locked.
