#UE4doc #docs: sp: Furthermore

diff --git a/Engine/Documentation/Source/Engine/Physics/Vehicles/DoubleWishboneVehicle/DoubleWishboneVehicle.INT.udn b/Engine/Documentation/Source/Engine/Physics/Vehicles/DoubleWishboneVehicle/DoubleWishboneVehicle.INT.udn
--- a/Engine/Documentation/Source/Engine/Physics/Vehicles/DoubleWishboneVehicle/DoubleWishboneVehicle.INT.udn
+++ b/Engine/Documentation/Source/Engine/Physics/Vehicles/DoubleWishboneVehicle/DoubleWishboneVehicle.INT.udn
@@ -26,13 +26,13 @@ Once you have watched those videos and have successfully setup a base vehicle, y
 
 When we set up the vehicle movement component and assign wheels to it, the vertical movement of the wheels is typically updated by a **WheelHandler** node inside the vehicle's animation Blueprint. On each frame, the handler will move the wheels up and down along the local Z axis to simulate suspension, making sure that the wheels touch the ground whenever that is possible.
 
-Further more, the **WheelHandler** also applies rotation to make the wheels spin around their Y axis, driven by the rotation speed of the simulated motor and gearbox. In addition to this, it will turn the wheels around their Z axis, driven by the current steering angle.
+Furthermore, the **WheelHandler** also applies rotation to make the wheels spin around their Y axis, driven by the rotation speed of the simulated motor and gearbox. In addition to this, it will turn the wheels around their Z axis, driven by the current steering angle.
 
 The default vertical suspension movement is clamped by the Suspension Max Raise and Drop settings in your VehicleWheel-derived classes. This movement is entirely linear, as illustrated below:
 
 ![](simwheeldefaultmovement.gif)(convert:false)
 
-_The standard vertical movement produced by the **WheelHandler**._
+_The standard vertical movement produced by the **WheelHandler**_
 
 While not entirely realistic, this type of suspension simulation is in fact perfectly fine for most vehicles, because you cannot usually see the suspension arms or the other components such as the springs and shocks. Any car model with a full body can get typically away with this.
 
@@ -42,13 +42,13 @@ To achieve a more realistic result, we really need a type of motion that can be
 
 ![](realwheelmovement.gif)(convert:false)
 
-_Ideal vertical and lateral movement._
+_Ideal vertical and lateral movement_
 
 To solve this little conundrum, we can start by assuming that wheel management falls into two different categories: one being the simulated wheel nodes controlled by the **WheelHandler**, the other being the wheels that we can actually see being rendered when we drive the car. Then we could feed data supplied to us by the **WheelHandler** to the rest of our suspension setup to get the desired results.
 
 In short, the visible wheel and the simulated wheel could in fact be two entirely different objects, and the simulated wheel would not even have to be visible at all.
 
-The wheel parameters we define in the VehicleWheel class can specify the collision mesh, radius and width explicitly. Those specifications do not necessarily have to correspond to any real geometry in the vehicle's skeletal mesh, so there is really no need for the simulated wheel nodes in the mesh to have any real geometry attached to them at all. Further more, we can specify additional wheel offsets in the wheel setup section of the vehicle Blueprint, so the pivot of our wheels do not have to be aligned with the geometrical center of the visible wheel mesh.
+The wheel parameters we define in the VehicleWheel class can specify the collision mesh, radius and width explicitly. Those specifications do not necessarily have to correspond to any real geometry in the vehicle's skeletal mesh, so there is really no need for the simulated wheel nodes in the mesh to have any real geometry attached to them at all. Furthermore, we can specify additional wheel offsets in the wheel setup section of the vehicle Blueprint, so the pivot of our wheels do not have to be aligned with the geometrical center of the visible wheel mesh.
 
 These are the basic insights we need to build a suspension that looks like it could realistically work from a mechanical point of view.
 
