#UE4doc #docs: sp: constraints

diff --git a/Engine/Documentation/Source/Engine/Physics/Vehicles/VehcileContentCreation/Vehicle_Content.INT.udn b/Engine/Documentation/Source/Engine/Physics/Vehicles/VehcileContentCreation/Vehicle_Content.INT.udn
--- a/Engine/Documentation/Source/Engine/Physics/Vehicles/VehcileContentCreation/Vehicle_Content.INT.udn
+++ b/Engine/Documentation/Source/Engine/Physics/Vehicles/VehcileContentCreation/Vehicle_Content.INT.udn
@@ -77,7 +77,7 @@ The **Physics Asset** is incredibly important for vehicles, and cannot be overlo
 
 ![](badPA.png)
 
-That is because the **Physics Asset Tool** (**PhAT**) attempts to wrap the vertices skinned to a joint the best it can. And honestly, you are going to delete and re-create them all. Why? The **PhAT** does not currently have a way to gracefully handle the re-creation of the contraints that hold all the Physics Bodies together. So if you delete an intermediate Physics Body, it does not create a new constraint up the hierarchy. However, if you delete all the Physics Bodies and start building them from the root joint out, all your constraints will get created correctly.
+That is because the **Physics Asset Tool** (**PhAT**) attempts to wrap the vertices skinned to a joint the best it can. And honestly, you are going to delete and re-create them all. Why? The **PhAT** does not currently have a way to gracefully handle the re-creation of the constraints that hold all the Physics Bodies together. So if you delete an intermediate Physics Body, it does not create a new constraint up the hierarchy. However, if you delete all the Physics Bodies and start building them from the root joint out, all your constraints will get created correctly.
 
 In the **Hierarchy** pane, **Shift** Select everything, and then press the **Delete** key. This will remove all Physics Bodies from the asset.
 
