#UE4doc #docs: sp: bearings

diff --git a/Engine/Documentation/Source/Engine/Physics/Vehicles/VehicleUserGuide/VehicleIntro.INT.udn b/Engine/Documentation/Source/Engine/Physics/Vehicles/VehicleUserGuide/VehicleIntro.INT.udn
--- a/Engine/Documentation/Source/Engine/Physics/Vehicles/VehicleUserGuide/VehicleIntro.INT.udn
+++ b/Engine/Documentation/Source/Engine/Physics/Vehicles/VehicleUserGuide/VehicleIntro.INT.udn
@@ -214,7 +214,7 @@ This is really where the vehicle starts to come together. If your art is ready,
 
 		_The **Bone Name** in the above image is specific to this skeletal mesh._
 
-	The order you assign wheels has no baring on if it is a front or a rear wheel, only the **Bone Name** and **Wheel Class** have any affect.
+	The order you assign wheels has no bearing on if it is a front or a rear wheel, only the **Bone Name** and **Wheel Class** have any affect.
 
 	[REGION:note]
 	If you need more than 4 wheels, you can click the "**+**" icon next to the **Wheel Setups** property to add more.
