#UE4doc #docs: sp: Perspective

diff --git a/Engine/Documentation/Source/Engine/UI/LevelEditor/Viewports/ViewportControls/ViewportControls.INT.udn b/Engine/Documentation/Source/Engine/UI/LevelEditor/Viewports/ViewportControls/ViewportControls.INT.udn
--- a/Engine/Documentation/Source/Engine/UI/LevelEditor/Viewports/ViewportControls/ViewportControls.INT.udn
+++ b/Engine/Documentation/Source/Engine/UI/LevelEditor/Viewports/ViewportControls/ViewportControls.INT.udn
@@ -78,7 +78,7 @@ buttons pressed. These are also the only controls that can be used to navigate t
 
 | Control | Action |
 | ------- | ------- |
-| [REGION:tablesection]Persepctive[/REGION] |||
+| [REGION:tablesection]Perspective[/REGION] |||
 | **LMB + Drag** | Moves the camera forward and backward and rotates left and right. |
 | **RMB + Drag** | Rotates the viewport camera. |
 | **LMB + RMB + Drag** | Moves up and down. |
