TypoCorrection

diff --git a/Engine/Documentation/Source/Engine/Paper2D/TileMaps/Paper2DTilesMaps.INT.udn b/Engine/Documentation/Source/Engine/Paper2D/TileMaps/Paper2DTilesMaps.INT.udn
--- a/Engine/Documentation/Source/Engine/Paper2D/TileMaps/Paper2DTilesMaps.INT.udn
+++ b/Engine/Documentation/Source/Engine/Paper2D/TileMaps/Paper2DTilesMaps.INT.udn
@@ -1,7 +1,7 @@
+﻿Availability:Public
 Title:Paper 2D Tile Sets / Tile Maps
 Crumbs: %ROOT%, Engine, Engine/Paper2D
 Description:Overview of how to create Tile Sets and Tile Maps for use within Paper 2D. 
-Availability:Public
 Related: Engine\Paper2D\Sprites\Collision
 Related: Engine\Paper2D\Flipbooks
 Related: Engine/Rendering/Materials
