TypoCorrection

diff --git a/Engine/Documentation/Source/Platforms/GearVR/QuickStart/GearVRQuickStart.INT.udn b/Engine/Documentation/Source/Platforms/GearVR/QuickStart/GearVRQuickStart.INT.udn
--- a/Engine/Documentation/Source/Platforms/GearVR/QuickStart/GearVRQuickStart.INT.udn
+++ b/Engine/Documentation/Source/Platforms/GearVR/QuickStart/GearVRQuickStart.INT.udn
@@ -1,7 +1,7 @@
+﻿Availability:Public
 Title:Samsung Gear VR UE4 Quick Start
 Crumbs:%ROOT%, Platforms/GearVR
 Description:Instructions on how to quickly setup and deploy projects for use with the Samsung Gear VR Headset.
-Availability:Public
 Navigation:platform
 Platform:Gear VR
 
