TypoCorrection

diff --git a/Engine/Documentation/Source/Platforms/GearVR/QuickStart/3/GearVRQuickStart_3.INT.udn b/Engine/Documentation/Source/Platforms/GearVR/QuickStart/3/GearVRQuickStart_3.INT.udn
--- a/Engine/Documentation/Source/Platforms/GearVR/QuickStart/3/GearVRQuickStart_3.INT.udn
+++ b/Engine/Documentation/Source/Platforms/GearVR/QuickStart/3/GearVRQuickStart_3.INT.udn
@@ -1,7 +1,7 @@
+﻿Availability:Public
 Title:3. Gear VR Project Packaging & Deployment
 Crumbs:%ROOT%, Platforms/GearVR
 Description:Instructions on how to deploy projects form UE4 to use with the Samsung Gear VR Headset.
-Availability:Public
 Navigation:platform
 Platform:Gear VR
 
