TypoCorrection

diff --git a/Engine/Documentation/Source/Platforms/GearVR/BestPractices/GearVRBestPractices.INT.udn b/Engine/Documentation/Source/Platforms/GearVR/BestPractices/GearVRBestPractices.INT.udn
--- a/Engine/Documentation/Source/Platforms/GearVR/BestPractices/GearVRBestPractices.INT.udn
+++ b/Engine/Documentation/Source/Platforms/GearVR/BestPractices/GearVRBestPractices.INT.udn
@@ -1,7 +1,7 @@
+﻿Availability:Public
 Title:Samsung Gear VR Best Practices
 Crumbs:%ROOT%, Platforms
 Description:Information on how to get the most out of your Gear VR headset.
-Availability:Public
 Navigation:platform
 Platform:Gear VR
 Related: Platforms/Mobile/Performance
@@ -23,8 +23,8 @@ using the Samsung Gear Virtual Reality(VR) Headset and Samsung Galaxy Note 4 Sma
 ## Gear VR Rendering Options
 While combining the Gear VR Headset with the Samsung Galaxy Note 4 Smartphone can produce some very immersive VR experiences, 
 there are some rendering features that UE4 can not perform while using the Gear VR Headset due to hardware limitations. 
-Many of the restrictions that you will encounter are similar to restrictions for [mobile devices](https://docs.unrealengine.com/latest/INT/Platforms/Mobile/Performance/index.html).
-However because of the unique rendering requirements for the Gear VR Headset there are few things extra things that must be considered.
+Many of the restrictions that you will encounter are similar to restrictions for [mobile devices](Platforms/Mobile/Performance).
+However because of the unique rendering requirements for the Gear VR Headset there are few extra things that must be considered.
 In the following sections we will cover what these requirements are and how they can effect your Gear VR project. 
 
 ### Mobile HDR
@@ -38,7 +38,7 @@ Below you will find a list of all the Post Process features that will work with
 
 |Post Processing Section |Feature Name| Description |
 | -------- | ---------- |------|
-| Auto Exposure|Exposure Bias |B Can be used to increase or decrease the scenes brightness. |
+| Auto Exposure|Exposure Bias |Can be used to increase or decrease the scenes brightness. |
 
 ### Materials
 Creating and interacting with Materials for Gear VR projects is very similar to creating and interacting with Materials for other projects. 
@@ -58,8 +58,8 @@ For example only the following Material inputs will work with the Gear VR Headse
 | **Opacity Mask** |Opacity Mask helps to define if a certain pixel is see through or not. |
 
 [REGION:note]
-Please keep in mind that the both **Translucent** and **Masked** Materials can be very expensive to render, especially on hardware like Smartphone's.
-Whenever possible you should use [Opaque Materials](https://docs.unrealengine.com/latest/INT/Engine/Rendering/Materials/MaterialProperties/BlendModes/index.html#opaque) as they are much cheaper to render.
+Please keep in mind that the both **Translucent** and **Masked** Materials can be very expensive to render, especially on hardware like Smartphones.
+Whenever possible you should use [Opaque Materials](Engine/Rendering/Materials/MaterialProperties/BlendModes) as they are much cheaper to render.
 [/region]
 
 ### Material Shading Models
@@ -98,7 +98,7 @@ The settings contained here are for helping to lessen Material rendering costs o
 | bUseLightmapDirectionality | When disabled, lighting from light maps will be flat but cheaper. |
 
 ### Customized UVs
-[Customized UV's](https://docs.unrealengine.com/latest/INT/Engine/Rendering/Materials/CustomizedUVs/index.html) allow you to perform certain shader operations in vertex shaders instead of pixel shaders to help avoid dependent texture fetches. 
+[Customized UVs](Engine/Rendering/Materials/CustomizedUVs) allow you to perform certain shader operations in vertex shaders instead of pixel shaders to help avoid dependent texture fetches. 
 This offers an excellent way to claw back much needed performance when doing simple things like UV manipulation to make a texture title more than it currently does. 
 Whenever possible make sure to do all UV manipulation using Customized UVs especially when developing content for the Gear VR Headset.
 
@@ -107,7 +107,7 @@ Customized UVs are currently only implemented for certain component types such a
 [/region]
 
 ### Reflections
-The Gear VR does not support Screen Space Reflections(SSR) which means you will only be able generate reflections that come from [Reflection Capture Actors](https://docs.unrealengine.com/latest/INT/Resources/ContentExamples/Reflections/1_4/) that are placed in levels. 
+The Gear VR does not support Screen Space Reflections(SSR) which means you will only be able generate reflections that come from [Reflection Capture Actors](Resources/ContentExamples/Reflections/1_4/) that are placed in levels. 
 Also keep in mind that the reflections that are generated from the reflections actors will not be of very high fidelity due to hardware limitations. 
 
 ## Gear VR & Content Setup
@@ -124,11 +124,11 @@ Failing to follow these restrictions will result in your project now working.
 ### Mesh Considerations
 When creating Meshes to use with your Gear VR Project the following information needs to be kept in mind to ensure that what you are creating will work with the Gear VR Headset. 
 
-* Use as few [Material ID's](https://docs.unrealengine.com/latest/INT/Engine/Content/FBX/Materials/index.html) on 3D meshes as possible to save on draw calls.
-* Make sure to make use of [Static Mesh LOD's](https://docs.unrealengine.com/latest/INT/Engine/Content/FBX/StaticMeshes/index.html#lodsetup) whenever possible.
+* Use as few [Material IDs](Engine/Content/FBX/Materials) on 3D meshes as possible to save on draw calls.
+* Make sure to make use of [Static Mesh LODs](Engine/Content/FBX/StaticMeshes) whenever possible.
 * Keep an eye on your 3D meshes polygon count and try and keep it as [Low Poly](https://en.wikipedia.org/wiki/Low_poly) as possible.
-* Make sure that all static models have a 2nd UV set so they can make use of [Light Mapping](https://docs.unrealengine.com/latest/INT/Engine/Content/Types/StaticMeshes/LightmapUnwrapping/index.html).
-* Whenever possible try and make use if [Imposter Sprites](https://docs.unrealengine.com/latest/INT/Engine/Content/Tools/RenderToTextureTools/3/index.html) to replace far away Static Meshes as they are much cheaper to render.
+* Make sure that all static models have a 2nd UV set so they can make use of [Light Mapping](Engine/Content/Types/StaticMeshes/LightmapUnwrapping).
+* Whenever possible try and make use if [Imposter Sprites](Engine/Content/Tools/RenderToTextureTools/3) to replace far away Static Meshes as they are much cheaper to render.
 
 ### Texture Considerations
 When creating Textures to use with your Gear VR Project the following information needs to be kept in mind to ensure that what you are creating will not be detrimental to your projects overall performance. 
@@ -139,20 +139,20 @@ When creating Textures to use with your Gear VR Project the following informatio
 * Loading up to 512 MB of Textures into memory is possible with the Gear VR however it is recommended that you do not load more that 128 MB of Textures into memory.
 * Use square Textures whenever possible as them make a more efficient use of memory. 
 * Wherever possible try and use [Texture Atlasing](http://en.wikipedia.org/wiki/Texture_atlas) to reduce the number of Textures needed in memory.
-* [RGB Texture Masking / Packing](https://docs.unrealengine.com/latest/INT/Engine/Rendering/Materials/HowTo/Masking/index.html) is a great way to cut down on unique Textures in memory.
+* [RGB Texture Masking / Packing](Engine/Rendering/Materials/HowTo/Masking) is a great way to cut down on unique Textures in memory.
 * Only **TC_Default** and **TC_NormalMap** compression settings should be used as they are the only two that will work on mobile platforms. 
 
 ### Level Considerations
 Getting the maximum performance out of your levels on the Gear VR is a delicate balancing act that will require time to get just right. 
 The following suggestions / things to look out for can help you quickly identify areas of your levels that might cause performance issues when viewed on the Gear VR Headset.
 
-* Always make sure that the level lighting has been [built](https://docs.unrealengine.com/latest/INT/Engine/Rendering/LightingAndShadows/Lightmass/index.html) to ensure you get the best performance possible. If you see a red **Lighting needs to be rebuilt** error message on the screen, this means that your lighting needs to be built.
-* Make sure that all the lights in you level are set to [Static](https://docs.unrealengine.com/latest/INT/Engine/Rendering/LightingAndShadows/LightMobility/StaticLights/index.html) as the other lighting options can be very expensive to render.
-* Check to see if your level is making use of [Precomputed Visibility](https://docs.unrealengine.com/latest/INT/Engine/Actors/Volumes/index.html) and if not make sure to add one to the level and re-build lighting to enable it in the level's World Settings. 
-* Draw calls for the entire scene should not be greater than 100 for any view. This can be seen with [stat openglrhi](https://docs.unrealengine.com/latest/INT/Engine/Performance/StatCommands/index.html) on the device, [stat d3d11rhi](https://docs.unrealengine.com/latest/INT/Engine/Performance/StatCommands/index.html) in ES2 previewer on PC.
-* The total triangle count of the level should not be greater than 100,000 for any view. This can be seen with [stat openglrhi](https://docs.unrealengine.com/latest/INT/Engine/Performance/StatCommands/index.html) on the device, [stat d3d11rhi](https://docs.unrealengine.com/latest/INT/Engine/Performance/StatCommands/index.html) in ES2 previewer on PC.
+* Always make sure that the level lighting has been [built](Engine/Rendering/LightingAndShadows/Lightmass) to ensure you get the best performance possible. If you see a red **Lighting needs to be rebuilt** error message on the screen, this means that your lighting needs to be built.
+* Make sure that all the lights in you level are set to [Static](Engine/Rendering/LightingAndShadows/LightMobility/StaticLights) as the other lighting options can be very expensive to render.
+* Check to see if your level is making use of [Precomputed Visibility](Engine/Actors/Volumes) and if not make sure to add one to the level and re-build lighting to enable it in the level's World Settings. 
+* Draw calls for the entire scene should not be greater than 100 for any view. This can be seen with [stat openglrhi](Engine/Performance/StatCommands) on the device, [stat d3d11rhi](Engine/Performance/StatCommands) in ES2 previewer on PC.
+* The total triangle count of the level should not be greater than 100,000 for any view. This can be seen with [stat openglrhi](Engine/Performance/StatCommands) on the device, [stat d3d11rhi](Engine/Performance/StatCommands) in ES2 previewer on PC.
 * Be very careful in how Translucent and Opacity Masked Materials are used in the level as they are very expensive to render.
-* Make sure to [Performance Profile](https://docs.unrealengine.com/latest/INT/Engine/Performance/index.html) your level often to make sure you catch any performance issues sooner rather than later.
+* Make sure to [Performance Profile](Engine/Performance) your level often to make sure you catch any performance issues sooner rather than later.
 
 
 
