TypoCorrection

diff --git a/Engine/Documentation/Source/Platforms/GearVR/Debugging/GearVRDebugging.INT.udn b/Engine/Documentation/Source/Platforms/GearVR/Debugging/GearVRDebugging.INT.udn
--- a/Engine/Documentation/Source/Platforms/GearVR/Debugging/GearVRDebugging.INT.udn
+++ b/Engine/Documentation/Source/Platforms/GearVR/Debugging/GearVRDebugging.INT.udn
@@ -1,7 +1,7 @@
+﻿Availability:Public
 Title:Samsung Gear VR Debugging
 Crumbs:%ROOT%, Platforms
 Description:Information on how to debug applications for the Gear VR headset.
-Availability:Public
 Navigation:platform
 Platform:Gear VR
 Related: Platforms/Mobile/Performance
@@ -14,8 +14,8 @@ Related: Engine/Rendering/Materials/CustomizedUVs
 
 [TOC(start:2 end:3)]
 
-Being able to debug your project on it's target device is critical to tracking down performance issues and other possible project breaking errors. 
-In the following section you will learn how what tools are available to assist you in debugging your Gear VR projects and how to use them.
+Being able to debug your project on its target device is critical to tracking down performance issues and other possible project breaking errors. 
+In the following section you will learn what tools are available to assist you in debugging your Gear VR projects and how to use them.
 
 ### Enabling Gear VR Developer Mode
 Having the ability to view your Gear VR Project without having to use the Gear VR Headset can greatly speed up the time it takes to debug an issue by allowing you quick and easy access to things like Unreal's console. 
@@ -113,7 +113,7 @@ Check out the official Android documentation for much more detailed information
 	![](T_Windows_ADB_Devices.png)
 
 	[region:tip]
-	If you enter **ADB Devices** but do not see anything show up, make sure that you have gone through the [Gear VR Headset Prerequisites](Platforms/GearVR) setup section and installed / setup everything that is listed there and then try this step again. 
+	If you enter **ADB Devices** but do not see anything show up, make sure that you have gone through the [](Platforms/GearVR) setup section and installed / setup everything that is listed there and then try this step again. 
 	[/region]
 
 1. Now that you have verified that you can see your Smartphone in ADB Devices it is now time to get the log information from your phone. 
