TypoCorrection

diff --git a/Engine/Documentation/Source/Programming/Assertions/Assertions.INT.udn b/Engine/Documentation/Source/Programming/Assertions/Assertions.INT.udn
--- a/Engine/Documentation/Source/Programming/Assertions/Assertions.INT.udn
+++ b/Engine/Documentation/Source/Programming/Assertions/Assertions.INT.udn
@@ -1,7 +1,7 @@
+﻿Availability:Public
 Title: Assertions
 Crumbs:
 Description:
-Availability:Public
 Version:4.7
 
 Assertions are tool to verify assumptions that a given piece of code relies on. This can be as simple as verifying that a pointer is not NULL or as complicated as verifying a particular function is not reentered. UE4 provides a series of macros to perform these types of validations. They are macros so that they can compile out in certain build configurations, either for performance reasons or because they aren't needed in a final build. If you want to see the macros for yourself, you can find them in:
