#UE4doc #docs: sp: dragging/ grammer: "an" before Atmospheric

diff --git a/Engine/Documentation/Source/Engine/QuickStart/3/LDQS_3.INT.udn b/Engine/Documentation/Source/Engine/QuickStart/3/LDQS_3.INT.udn
--- a/Engine/Documentation/Source/Engine/QuickStart/3/LDQS_3.INT.udn
+++ b/Engine/Documentation/Source/Engine/QuickStart/3/LDQS_3.INT.udn
@@ -34,7 +34,7 @@ For more information on **Actors**, see: [**Actors & Geometry**](Engine\Actors).
 First you will need a floor to place everything on, for this, use a Box Geometry Actor.
 
 1. In the **Modes Panel**, with **Place** mode enabled, click the **BSP** Category.
-1. Place a Box Geometry Actor from the **Modes Panel** by clicking and draging it into the level.
+1. Place a Box Geometry Actor from the **Modes Panel** by clicking and dragging it into the level.
 
 	![](Engine/QuickStart/4/placeBox.png)
 
@@ -64,7 +64,7 @@ First you will need a floor to place everything on, for this, use a Box Geometry
 
 1. Switch to the **Visual Effects** category in the **Modes Panel**.
 
-1. Place a **Atmospheric Fog** into the level.
+1. Place an **Atmospheric Fog** into the level.
 	
 	![](Engine/QuickStart/4/placeFog.png)
 	
