#UE4doc #docs: sp: substantially, standalone

diff --git a/Engine/Documentation/Source/Engine/QuickStart/5/LDQS_5.INT.udn b/Engine/Documentation/Source/Engine/QuickStart/5/LDQS_5.INT.udn
--- a/Engine/Documentation/Source/Engine/QuickStart/5/LDQS_5.INT.udn
+++ b/Engine/Documentation/Source/Engine/QuickStart/5/LDQS_5.INT.udn
@@ -38,7 +38,7 @@ By now you have noticed the "Preview" labels in the shadows and the light leakin
 1. Set the **Quality Level** to **Production**.
 
 	[REGION:note]
-	Before beginning a **Build**, remember, it can take a while on larger levels. Our small hut may not take too long, but a scene with hundreds of lights that covers many square kilometers will take substantailly longer.
+	Before beginning a **Build**, remember, it can take a while on larger levels. Our small hut may not take too long, but a scene with hundreds of lights that covers many square kilometers will take substantially longer.
 	[/REGION]
 
 1. Click the ![**Build**](Engine\QuickStart\6\build_button.png) button.
