#UE4doc #docs: changed "let us" to "let's"

diff --git a/Engine/Documentation/Source/Engine/QuickStart/2/LDQS_2.INT.udn b/Engine/Documentation/Source/Engine/QuickStart/2/LDQS_2.INT.udn
--- a/Engine/Documentation/Source/Engine/QuickStart/2/LDQS_2.INT.udn
+++ b/Engine/Documentation/Source/Engine/QuickStart/2/LDQS_2.INT.udn
@@ -26,7 +26,7 @@ Next, create a new level to build your game environment in. When creating a new
 Unreal Engine enables you to select from _templates_. By default, there is a _Default_ template - which has a very simple environment - 
 and an _Empty Level_ template - which is completely blank - to choose from.
 
-Let us get started creating a new level:
+Let's get started creating a new level:
 
 1. From the File menu of the **Level Editor** window, choose ![](Engine/QuickStart/3/new_button.png) **New Level...**. This will display the **New Level** dialog:
 
