#UE4doc #docs: Updated Scalability link file

diff --git a/Engine/Documentation/Source/Engine/Performance/Scalability/Scalability.INT.udn b/Engine/Documentation/Source/Engine/Performance/Scalability/Scalability.INT.udn
--- a/Engine/Documentation/Source/Engine/Performance/Scalability/Scalability.INT.udn
+++ b/Engine/Documentation/Source/Engine/Performance/Scalability/Scalability.INT.udn
@@ -6,7 +6,7 @@ Description:Options for adjusting quality and performance.
 [VAR:Topic]
 [OBJECT:Topic]
 	[PARAM:image]
-		![%Engine/Performance/Scalability:title%](Engine/Rendering/rendering_topic.png)
+		![%Engine/Performance/Scalability/ScalabilityReference:title%](Engine/Performance/Scalability/ScalabilityReference/scalability_topic.png)
 	[/PARAM]
 	[PARAM:icon]
 		![](%ROOT%/light_icon.png)(convert:false)
