#UE4doc #docs: grammar: Split sentence into two

diff --git a/Engine/Documentation/Source/Engine/Rendering/LightingAndShadows/IndirectLightingCache/IndirectLightingCache.INT.udn b/Engine/Documentation/Source/Engine/Rendering/LightingAndShadows/IndirectLightingCache/IndirectLightingCache.INT.udn
--- a/Engine/Documentation/Source/Engine/Rendering/LightingAndShadows/IndirectLightingCache/IndirectLightingCache.INT.udn
+++ b/Engine/Documentation/Source/Engine/Rendering/LightingAndShadows/IndirectLightingCache/IndirectLightingCache.INT.udn
@@ -89,7 +89,7 @@ Here the torus is a movable Static Mesh and the 5x5x5 interpolation positions ar
 [/REGION]
 
 This interpolation only happens when the object moves enough to invalidate it, hence the cache naming.  Note that the positions 
-are quite far from the object's bounds, that is to provide continuous and stable lighting as the object moves through the world. 
+are quite far from the object's bounds. This is to provide continuous and stable lighting as the object moves through the world. 
 Most objects will effectively get 3x3x3 interpolation resolution.  The volume lighting samples store directional indirect lighting which allows normal maps to work.   
 
 ## Previewing Unbuilt Lighting
