#UE4doc #docs: sp: Fraction

diff --git a/Engine/Documentation/Source/Engine/Rendering/LightingAndShadows/Lightmass/Lightmass.INT.udn b/Engine/Documentation/Source/Engine/Rendering/LightingAndShadows/Lightmass/Lightmass.INT.udn
--- a/Engine/Documentation/Source/Engine/Rendering/LightingAndShadows/Lightmass/Lightmass.INT.udn
+++ b/Engine/Documentation/Source/Engine/Rendering/LightingAndShadows/Lightmass/Lightmass.INT.udn
@@ -117,7 +117,7 @@ In the first image is a scene with indirect lighting but no ambient occlusion.
 | ------- | ------------ |
 | **Visualize Ambient Occlusion** | Overrides lightmaps with just the occlusion factor when lighting is built.  This is useful for seeing exactly what the occlusion factor is, and comparing the effects of different settings. |
 | **Max Occlusion Distance** | Maximum distance for an object to cause occlusion on another object. |
-| **Fully Occluded Samples Franction** | Fraction of samples taken that must be occluded in order to reach full occlusion.  Note that there is also a per-primitive FullyOccludedSamplesFraction, which allows controlling how much occlusion an object causes on other objects. |
+| **Fully Occluded Samples Fraction** | Fraction of samples taken that must be occluded in order to reach full occlusion.  Note that there is also a per-primitive FullyOccludedSamplesFraction, which allows controlling how much occlusion an object causes on other objects. |
 | **Occlusion Exponent** | Higher exponents increase contrast. |
 
 [REGION:imagetable]
@@ -191,7 +191,7 @@ The amount of light passing through the material is called **Transmission**, and
 	* BLEND_Translucent and BLEND_Additive: Transmission = Lerp(White, Emissive, Opacity)
 	* BLEND_Modulate: Transmission = Emissive
 
-This means that at an opacity of 0, the material will not filter out incident lighting and there will be no translucent shadow.  At an opacity of 1, incident lighting will be filtered by the material's emissive or BaseColor (depending on whether it is lit).  Note that indirect lighting will sometimes wash out the translucent shadows and make them less saturated than the emissive or diffuse of the translucent material.
+This means that at an opacity of 0, the material will not filter out incident lighting and there will be no translucent shadow.  At an opacity of 1, incident lighting will be filtered by the material's emissive or Base Color (depending on whether it is lit).  Note that indirect lighting will sometimes wash out the translucent shadows and make them less saturated than the emissive or diffuse of the translucent material.
 
 #### Translucent shadow sharpness
 
@@ -240,7 +240,7 @@ Compare the lighting clarity between the scene in the first image built with mid
 | | |
 [/REGION]
 
-The **Unlit** view mode is useful for viewing the **Diffuse** term.  The scene in the first image looks flatter and more monotone in the **Unlit** view mode, which means the lighting is doing all the work and variations in final pixel color are mostly due to lighting differences.  **To get good lighting, your scene should look flat and boring in the Unlit viewmode.**  Baking lighting and macro features into the diffuse textures will counteract the lighting.
+The **Unlit** view mode is useful for viewing the **Diffuse** term.  The scene in the first image looks flatter and more monotone in the **Unlit** view mode, which means the lighting is doing all the work and variations in final pixel color are mostly due to lighting differences.  **To get good lighting, your scene should look flat and boring in the Unlit view mode.**  Baking lighting and macro features into the diffuse textures will counteract the lighting.
 
 [REGION:imagetable]
 ![94MidToneDiffuseUnlit.jpg](94MidToneDiffuseUnlit.jpg)(w:465) | ![94DarkDiffuseUnlit.jpg](94DarkDiffuseUnlit.jpg)(w:465) |
@@ -339,7 +339,7 @@ You can access **World Settings** by clicking the ![](WorldSettings_Icon.png) ic
 | **Indirect Lighting Quality** | Scales the sample counts used by the Lightmass GI solver.  Higher settings result in fewer solver artifacts (noise, splotchiness) at much increased build times.  Note that this will not affect artifacts due to using lightmaps (texture seams, compression artifacts, texel shapes).  |
 | **Indirect Lighting Smoothness** | Higher values cause the indirect lighting to be smoothed more, which can hide solver noise, but also causes detailed indirect shadows and Ambient Occlusion to be lost.  It can be useful to lower this somewhat (.66 or .75) when increasing *Indirect Lighting Quality* for highest quality. |
 | **Environment Color** | Color that rays which miss the scene will pick up.  The environment can be visualized as a sphere surrounding the level, emitting this color of light in each direction. |
-| **Environment Intensity** | Scales EnvironmentColor to allow an HDR environment color. |
+| **Environment Intensity** | Scales Environment Color to allow an HDR environment color. |
 | **Diffuse Boost** | Scales the diffuse contribution of all materials in the scene.  Increasing DiffuseBoost is an effective way to increase the intensity of the indirect lighting in a scene. The diffuse term is clamped to 1.0 in brightness after DiffuseBoost is applied, in order to keep the material energy conserving (meaning light must decrease on each bounce, not increase).  If raising DiffuseBoost does not result in brighter indirect lighting, the diffuse term is being clamped and the Light's IndirectLightingScale should be used to increase indirect lighting instead. |
 | **Use Ambient Occlusion** | Enables static Ambient Occlusion to be calculated by Lightmass and built into your lightmaps. |
 | **Direct Illumination Occlusion Fraction** | How much of the AO to apply to direct lighting. |
