Paper 2D Updates
- Tile Sets / Tile Map page.
- Updates to Flipbook Editor pages.
- HowTo Section (Animation State Machines, Physics, Sockets, TopDown Setup).

diff --git a/Engine/Documentation/Source/Engine/Paper2D/HowTo/Animation/2DCharacterAnimation.INT.udn b/Engine/Documentation/Source/Engine/Paper2D/HowTo/Animation/2DCharacterAnimation.INT.udn
new file mode 100644
--- /dev/null
+++ b/Engine/Documentation/Source/Engine/Paper2D/HowTo/Animation/2DCharacterAnimation.INT.udn
@@ -0,0 +1,190 @@
+Title:Setting up an Animation State Machine
+Crumbs: %ROOT%, Engine, Engine/Paper2D
+Description:How to set up and play different animations for a 2D character based on different states. 
+Availability:Public
+Related: Engine\Paper2D\HowTo\TopDown
+Related: Engine\Paper2D\Sprites
+Related: Engine\Paper2D\Flipbooks
+
+On this page you will learn how to set up an Animation State machine for a 2D character which will allow the character to cycle between different Flipbook animations based on your defined criteria. 
+
+[REGION:warning]
+This tutorial assumes you have a Paper 2D character that is already set up for movement and have some Flipbook animations that you want the character to switch between. If you do not have a character or Flipbook animations, see the [Paper 2D Top Down Tutorial](Engine\Paper2D\HowTo\TopDown) for information on setting up a Paper 2D character and applying movement as well as a link to the sample assets used in this tutorial. 
+[/REGION]
+
+With your project open and Paper 2D character ready to go, follow the steps below. 
+
+## Variable Setup
+
+1.	Inside your project, open your Paper 2D Character Blueprint (**TopDownCharacter** in our example). 
+
+1.	In the **My Blueprint** window, click the **Add Variable** button to create a Bool variable and name it **IsMoving?**. 
+
+	![](Animation1.png)
+
+1.	Create and name another variable based on one of your Flipbook animations (such as Idle) and set it to a **Paper Flipbook** type. 
+
+	![](Animation2.png)
+
+	Depending upon the type of game you are making, you may have different animations. Most character based games, have a state where the character is idle (or resting and not moving). 
+
+1.	Click **Compile** then set the default value of the variable to its corresponding Flipbook asset. 
+
+	![](Animation3.png)
+
+	Above we have assigned our Idle Flipbook animation to our IdleFlipbook variable. 
+
+	(Optional) For the sake of clarity, you may also want to put the variable in a **Category** called **Animations**. 
+
+	![](Animation4.png)
+
+	You can type in the name of the category you would like and will keep things organized for you. 
+
+1.	Repeat the previous two steps for each of the animations that you want your character to be able to play. 
+
+	![](Animation5.png)
+
+	Above we have created another Paper Flipbook variable called **RunFlipbook**, clicked **Compile** then set the **Default Value** for that variable to our corresponding Run Flipbook animation. While our example only has the two animations, you could have one for a death animation, another for reloading a gun (if they have one) or a melee attack, another for picking up or using items, etc. 
+
+	Once you have added variables for each of our Flipbook animation states, proceed to the next section. 
+	
+## Update Animation Setup
+
+Next we will create some script to handle the updating of our Animation state. 
+
+1.	**Right-click** inside the **Event Graph** and search for and add a **Custom Event**. 
+
+	![](Animation6.png)
+
+	Call the Custom Event something to the effect of **UpdateAnimation**. 
+
+1.	**Right-click** in the graph again and add a **Get Velocity** node, then drag off the return value and add the **Vector Length** node. 
+
+	![](Animation7.png)
+
+	We can use the Vector Length of Velocity to determine if our character is moving or not. 
+
+1.	Off the return value of the Vector Length, add a **>** node then drag-and-drop the **IsMoving?** variable onto the Bool return to set it. 
+
+	![](Animation8.png)
+
+## Animation State Machine Function
+
+Here we will create a function that will handle the checking of which state we are in and output a corresponding Flipbook animation to use. 
+
+1.	In the **My Blueprint** window, create a new **Function** called **Anim State Machine**. 
+
+	![](Animation9.png)
+
+1.	In the **Details** panel of the function, add an **Input** Boolean called **IsMoving?** and an **Output** Paper Flipbook called **Output Flipbook**. 
+
+	![](Animation10.png)
+
+	Our function will take in the value of our variable also called IsMoving? to check if we are moving and use that as a test case to determine what state we are in (moving or not). Our function will then output an _Output Flipbook_ which we will set our sprite to use in a moment. 
+
+1.	Connect a **Branch** node off the input executable pin of the function. 
+
+	![](Animation11.png)
+
+	The first thing that occurs when this function is called is run a True/False check against our **IsMoving?** variable. 
+
+1.	Back inside the **Content Browser**, **Right-click** and under **Blueprints** create an **Enumeration** and call it **Animation State**. 
+
+	![](Animation12.png)
+
+	We are using an Enumeration to create a list of all our animation states that our character can be in which we can then select from. 
+
+1.	Open the Enumeration, then add a state for each of the animations that your character will enter.  
+
+	![](Animation13.png)
+
+	For example, we only have Idle and Run however you may have many more depending upon the animations you plan to use. 
+
+1.	Return to your Paper 2D Character Blueprint and create a new variable called **AnimationState** set to the Enum _Animation State_.
+
+	![](Animation14.png)
+
+1.	Inside the **Anim State Machine** function, hold **Alt** and drag in the **AnimationState** variable to set it. 
+
+1.	Drag in another **AnimationState** variable and connect each to the **True** and **False** pins of the **Branch** node. 
+
+	![](Animation15.png)
+
+	For the Enum connected to True, set the state to **Run** and set the other Enum to **Idle**. We are saying if IsMoving? is true set our Enum state to Run and if we are not moving set it to Idle. While our example only checks whether we are moving or not, you could add different Branch nodes to check different criteria based on your needs. 
+
+	One thing to note is that the order in which you process your checking is important and you'll want to prioritize certain states over others. For example, if we had a death Flipbook animation for our character would probably want to process that first and check if the character is dead and if so, we would want to play the death animation and kill the character. If they are not dead, then we can proceed on to the next series of checks (if the character is moving, attacking, using an item, etc. with the final check being if the character is not doing anything at all resulting in the Idle pose. 
+	
+	[REGION:note]
+	See the [Animation State Machine](https://wiki.unrealengine.com/Videos/Player?series=PLZlv_N0_O1gauJh60307mE_67jqK42twB&video=Af-jsQFlLFU) section of the Paper 2D video tutorial series for more information. 
+	[/REGION]
+
+1.	Hold **Ctrl** and drag in the **IdleFlipbook**, **RunFlipbook** and **Animation State** variables, then **Right-click** and add a **Select** node. 
+
+	![](Animation16.png)
+
+	For this section, you will need to drag in each of your animation variables into the graph. 
+
+1.	Connect each of the nodes so the function is complete and looks similar to below. 
+
+	![](Animation17.png)(w:800)
+
+	Our function now checks if we are moving and if so, sets our Enum state to Run (Idle if we are not moving). We then use a Select node to get our Animation State from our Enum and based on the state assigns the corresponding Flipbook as the _Output Flipbook_ to use. 
+
+## Finishing Update Animation Event
+
+Now that our function is complete, we will call it as part of our Update Animation Event and finish off the remaining script. 
+
+1.	On the **Event Graph** in the **My Blueprint** window, choose to **Show Inherited Variables** then hold **Ctrl** and drag in your **Sprite**. 
+
+	![](Animation18.png)
+
+1.	Drag in the **Anim State Machine** function and connect it to the **Set Is Moving?** node, then off your **Sprite** add a **Set Flipbook** node.
+
+	![](Animation19.png)
+
+1.	Connect all of your nodes together for the **UpdateAnimation** custom event as shown below. 
+
+	![](Animation20.png)(w:820)
+
+1.	Following the **Add Movement Input** node from your MoveForward Event, call the **Update Animation** custom event. 
+
+	![](Animation21.png)
+
+1.	**Compile** and **Save** then play in the editor. 
+
+	![](AnimationExample.gif)(convert:false)
+
+	You will now have a character that cycles between the two states based on your condition of whether or not the character is moving. 
+
+As mentioned, in order to add additional animation states to this setup you would need to: 
+
+* Have Flipbook Animations for each of the animation states you want your character to enter. 
+* Add and set them as Flipbook variables inside your Paper 2D Character Blueprint. 
+* Add each of the states to the **Animation State** Enumeration we created. 
+* Create conditions that determine when the character can enter each state and add it to the **Anim State Machine** function. 
+* Connect each of your Flipbook variables to the **Select** node inside the **Anim State Machine** function. 
+
+
+	
+
+
+
+
+
+
+
+
+
+
+
+
+
+
+
+
+
+
+
+
+
+
diff --git a/Engine/Documentation/Source/Engine/Paper2D/HowTo/Paper2DHowTo.INT.udn b/Engine/Documentation/Source/Engine/Paper2D/HowTo/Paper2DHowTo.INT.udn
new file mode 100644
--- /dev/null
+++ b/Engine/Documentation/Source/Engine/Paper2D/HowTo/Paper2DHowTo.INT.udn
@@ -0,0 +1,19 @@
+Title:Paper 2D - How To...
+Crumbs: %ROOT%, Engine, Engine/Paper2D
+Description:The Paper 2D How To pages provides several short step-by-step guides for working with Paper 2D. 
+Availability:Public
+Related: Engine\Paper2D\Sprites
+Related: Engine\Paper2D\Flipbooks
+
+
+## Overview
+The **Paper 2D - How To** pages provide short step-by-step examples of ways to accomplish certain tasks within Paper 2D in Unreal Engine 4. While there may be other ways to accomplish these tasks, these samples should get you started and give you ideas on how you may want to approach your objectives. 
+
+Refer to the table below for more information and click each link to view its contents. 
+
+|How To... | Description |
+| ------- | ------------ |
+| **[](Engine\Paper2D\HowTo\Physics)** | %Engine\Paper2D\HowTo\Physics:description% |
+| **[](Engine\Paper2D\HowTo\Sockets)** | %Engine\Paper2D\HowTo\Sockets:description% |
+| **[](Engine\Paper2D\HowTo\TopDown)** | %Engine\Paper2D\HowTo\TopDown:description% |
+| **[](Engine\Paper2D\HowTo\Animation)** | %Engine\Paper2D\HowTo\Animation:description% |
diff --git a/Engine/Documentation/Source/Engine/Paper2D/HowTo/Physics/SpritePhysics.INT.udn b/Engine/Documentation/Source/Engine/Paper2D/HowTo/Physics/SpritePhysics.INT.udn
new file mode 100644
--- /dev/null
+++ b/Engine/Documentation/Source/Engine/Paper2D/HowTo/Physics/SpritePhysics.INT.udn
@@ -0,0 +1,52 @@
+Title:Applying 2D Physics
+Crumbs: %ROOT%, Engine, Engine/Paper2D
+Description:Examples of implementing physics for 2D games. 
+Availability:Public
+Related: Engine\Paper2D\Sprites
+Related: Engine\Paper2D\Flipbooks
+
+This page provides an example of applying physics to sprites in 2D games. 
+
+[REGION:note]
+_For this example we are using a 2D Side Scroller Game, however you can apply the same concepts to any 2D game._
+[/REGION]
+
+When using physics in any 2D style game, typically you'll want to apply some constraints to your sprite when physics are applied to prevent the sprite from moving or rotating in such a way that the card becomes visible to the player, or even worse, falls out of your game world as seen in the example below. 
+
+![](GIF_1.gif)(convert:false)
+
+Above, we have applied physics to a sprite but have not applied any constraints to those physics so our object that falls and applies physics actually tumbles outside of our game world and off screen which is not what we want. We want our physics object to react to the world and player interaction, however we want to keep it inside our level and prevent it from rotating on certain axis so that the sprite is always fully visible. 
+
+For our sprite, we can apply constraints from the **Details** panel in addition to applying physics. 
+
+1.	Select the Sprite in your level that you want to apply physics to. 
+
+1.	In the **Details** panel, click the **Simulate Physics** option under _Physics_. 
+
+1.	Expand _Constraints_ and choose the **Lock Position** (typically to the **Y** axis for side scrolling games). 
+
+	![](Constraints1.png)
+
+	While this will apply physics and lock it to only the Y axis, there are additional issues that may occur with just this setting. 
+
+	![](GIF_2.gif)(convert:false)
+
+	Our sprite is now locked to the Y axis, however you can see that it can still rotate freely which may not be what we want. 
+
+1.	Also in the _Constraints_ section, choose a **Lock Rotation** (typically to the **X** axis for side scrolling games). 
+
+	![](Constraints3.png)
+
+	Now when physics is applied to the sprite its position is locked to the Y and its rotation is locked at the X. 
+
+	![](GIF_3.gif)(convert:false)
+
+	You can also use the **Mode** option to constrain movement along a specified axis (in this case **XZPlane** to achieve the same effect). 
+
+
+Depending upon the type of 2D game you are making, varying settings can be used to restrict how physics are applied to your object. 
+
+
+
+
+
diff --git a/Engine/Documentation/Source/Engine/Paper2D/HowTo/Sockets/SpriteSockets.INT.udn b/Engine/Documentation/Source/Engine/Paper2D/HowTo/Sockets/SpriteSockets.INT.udn
new file mode 100644
--- /dev/null
+++ b/Engine/Documentation/Source/Engine/Paper2D/HowTo/Sockets/SpriteSockets.INT.udn
@@ -0,0 +1,98 @@
+Title:Using Sprite Sockets
+Crumbs: %ROOT%, Engine, Engine/Paper2D
+Description:How to attach objects to sockets defined in a Paper 2D Sprite. 
+Availability:Public
+Related: Engine\Paper2D\Sprites
+Related: Engine\Paper2D\Flipbooks
+
+
+On this page you will learn how to add sockets to sprites then access those sockets from within a Blueprint which you can attach an object to. You can spawn particle effects from sprite sockets (such as a character on fire), attach a weapon to a character's hand (like a gun or sword) or any number of different items that you want to spawn and attach to a sprite's socket. 
+
+For this example we will attach a hat to our character's head. 
+
+**To add a socket to a sprite**: 
+
+1.	Locate the sprite you wish to add a socket to and open it. 
+
+1.	In the **Details** panel, click the **+** next to **Sockets** to add a new socket. 
+
+1.	Enter in a **Socket Name** and move the socket where you would like using the **Local Transform** or widget inside the viewport. 
+
+	![](AddingSocket.png)
+
+	To move the socket in the viewport, you may need to zoom in and click on the socket first. 
+
+It is important to note the name you give your socket. Above we have created a new socket called _HatSocket_ to which we want to spawn in a 2D hat for our character to wear. The steps above have only covered one sprite, in our idle animation we are using 8 different sprites so we need to add sockets to each of those sprites as well. 
+
+A quick way to do this would be to use the **Property Matrix**. 
+
+**Add Sockets to Several Sprites**: 
+
+1.	In the **Content Browser**, select each of the sprites in your Flipbook animation that you want to apply sockets to. 
+
+1.	**Right-click** on a sprite, then under **Asset Actions** select **Property Matrix...**
+
+	![](PropteryMatrix.png)
+
+	By opening the **Property Matrix** we can modify each of these sprites at the same time which will save time. 
+
+	[REGION:note]
+	We also selected the sprite which we already added a socket to for reference so all 8 frames are selected. 
+	[/REGION]
+
+1. Inside the **Property Matrix**, select each of your sprites (which do not have sockets) then click the **+** to add sockets to them. 
+
+	![](AddSockets.png)(w:820)
+
+1. Expand element 0, then under **Socket Name** name them the same as your first sprite (in our example **HatSocket**). 
+
+	![](AddSockets2.png)(w:820)
+
+1. Expand **Local Transform** then **Translation** and set the XYZ values to the same as your first sprite. 
+
+	![](AddSockets3.png)(w:820)
+
+1. Open up each of the sprites in your animation to ensure that a socket has been added to each. 
+
+	![](AddSockets4.png)(w:820)
+
+	[REGION:note]
+	You may need to adjustment the socket location per sprite to ensure it is at the desired location for that frame. 
+	[/REGION]
+
+Now that we have a socket attached to each frame of our Flipbook, we can tell our Blueprint to spawn an Actor at the socket location. 
+
+
+**Accessing Sprite Sockets via Blueprint**: 
+
+1.	Inside our Character Blueprint, in the **My Blueprint** window click the eye icon and **Show Inherited Variables**. 
+
+1.	**Left-click** and drag in the **Sprite** component for your character. 
+
+	![](Blueprint1.png)
+
+1.	Off an **Event Begin Play** node, add a **Spawn Actor from Class** node set to the Actor Class you want. 
+
+1.	Use a **Make Transform** node for the **Spawn Transform** and off the **Return Value** use an **Attach To** node. 
+
+1.	Connect your **Sprite** component as the **In Parent** and for the **In Socket Name** enter the Socket you created. 
+
+	![](Blueprint2.png)(w:800)
+
+	Your setup should look similar to above. For example we are just spawning a Blueprint called **Hat_BP** which is a Blueprint containing a Sprite Component of a hat that we are attaching to our character at the **HatSocket** location we've created. The same process we took above for our idle frames, we've done for our run frames so that we end up with a character that wears a hat in our game (as seen below). 
+
+1.	Once you **Compile** and play in the editor, the Actor Class spawned will spawn at the socket location and become attached. 
+
+	![](Example_GIF.gif)(convert:false)
+
+You could add a socket to a sprite and attach any number of different things in that socket, whether or not the sprite is part of a Flipbook and is animated as well as what goes in the socket is entirely up to you. 
+
+
+
+
+
+
+
+
+
+
diff --git a/Engine/Documentation/Source/Engine/Paper2D/HowTo/TopDown/TopDown.INT.udn b/Engine/Documentation/Source/Engine/Paper2D/HowTo/TopDown/TopDown.INT.udn
new file mode 100644
--- /dev/null
+++ b/Engine/Documentation/Source/Engine/Paper2D/HowTo/TopDown/TopDown.INT.udn
@@ -0,0 +1,182 @@
+Title:Top Down 2D Game Setup
+Crumbs: %ROOT%, Engine, Engine/Paper2D
+Description:How to quickly set up a top down 2D game with a playable character. 
+Availability:Public
+Related: Engine\Paper2D\HowTo\Animation
+Related: Engine\Paper2D\TileMaps
+Related: Gameplay/HowTo/UsingCameras/Blueprints
+Related: Engine\Paper2D\Sprites
+Related: Engine\Paper2D\Flipbooks
+
+On this page you will learn how to modify the existing **Third Person Template** project to create a **Top Down 2D** style game. 
+
+You can use your own 2D character with this tutorial or you can download a sample character from the link below: 
+
+* [Sample Assets](SampleAssets.rar)
+
+## Setup Project
+
+Here we will set up our project for this tutorial, you can import your own assets during this section or use the Sample Assets provided. 
+
+1.	Create a new **Third Person Blueprint** template project. 
+
+1.	Inside the project, **Right-click** on the **Content** folder in the **Content Browser** and create a **New Folder**. 
+
+	![](TopDown1.png)
+
+	Give the folder a name such as **TopDown** or whatever you would like.
+
+1.	Locate the downloaded Sample Assets and drag the **SoldierSprites.paper2dsprites** asset into **Content Browser**. 
+
+	![](TopDown2.png)
+
+	This will automatically create a **SolderSprites** Sprite Sheet asset and folders containing **Frames** and **Textures**. 
+
+	![](TopDown3.png)
+
+	[REGION:note]
+	For more information on how to set up assets for importing, see the [](Engine\Paper2D\Importing) documentation. 
+	[/REGION]
+
+1.	**Right-click** on the **SoldierSprites** asset and select **create Flipbooks**. 
+
+	![](TopDown4.png)
+
+	This will automatically create two Flipbooks based on the sprite sheet for us. 
+
+	![](TopDown5.png)
+
+1.	(Optional) Open the **Idle** Flipbook and in the details panel set **Frames Per Second** to **3.0**. 
+
+	![](TopDown6.png)
+
+	Also set the **Run** Flipbook **Frames Per Second** to **8.0**. 
+
+	![](TopDown7.png)
+
+	This will slow down the playback speed of both Flipbook Animations. 
+
+Our project is now set up and we can begin setting up the Top Down Character. 
+
+## Character Setup
+
+Here we will create a new Paper 2D Character Blueprint and set up the character. 
+
+1.	**Right-click** in the **Content Browser** and select to create a new **Blueprint Class**. 
+
+	![](TopDown8.png)
+
+1.	In the **Pick Parent Class** window, click the **All Classes** drop down and search for and add the **PaperCharacter**. 
+
+	![](TopDown9.png)
+
+	Give the new Blueprint a name such as **TopDownCharacter**. 
+
+1.	Inside the **TopDownCharacter** Blueprint, in the **Components** window click the **Sprite** component. 
+
+	![](TopDown10.png)
+
+1.	In the **Details** panel under **Sprite** click the **Source Flipbook** drop down and select the **Idle** Flipbook. 
+
+	![](TopDown11.png)
+
+	The character will then be added to the **Viewport** inside the Capsule Component. 
+
+	![](TopDown12.png)
+
+1.	Back inside the **Details** panel under **Transform**, set the **Rotation** to **-90,0,90** and **Scale** to **0.75** for XYZ. 
+
+	![](TopDown13.png)
+
+	This will rotate the Flipbook into position and scale it down so it fits inside the Collision Capsule. 
+
+1.	In the **Components** window, click the **CapsuleComponent** then in the **Details** panel adjust the **Shape** settings. 
+
+	![](TopDown14.png)
+
+	This will vary based on the character you are using, we have set both **Half Height** and **Radius** to **45.0**. 
+
+1.	In the **Components** window, click **TopDownCharacter(self)** then in the **Details** panel, un-check **Use Controller Rotation Yaw**. 
+
+	![](TopDown15.png)
+
+	For this example we are going to automatically orientate our character to the direction he is moving and not use the controller rotation. 
+
+1.	In the **Components** window, click the **CharacterMovement** component.
+
+1.	In the **Details** panel for **CharacterMovement**, update the **Max Walk Speed** to **400** and check **Orient Rotation to Movement**. 
+
+	![](TopDown16.png)
+
+	This will reduce the characters movement speed and automatically rotate the character based on our movement. 
+
+1.	In the **Components** window, click **Add Component** and search for and add a **Spring Arm** component. 
+
+	![](TopDown17.png)
+
+1.	Click **Add Component** again and add a **Camera** component, then drag the Camera onto the Spring Arm to attach it. 
+
+	![](TopDown18.png)
+
+1.	Select the **Camera** component, then in the **Details** panel, zero out the **Location** setting. 
+
+	![](ZeroCamera.png)
+
+1.	Select the **Spring Arm** then make the following updates in the **Details** panel. 
+
+	![](TopDown19.png)
+
+	Here we are adjusting the **Rotation** of the Spring Arm (and by extension, Camera) to **180, -90, 180** to position the camera directly above the character. We are also adjusting the camera distance away from the character by changing the **Target Arm Length** to **600**. Finally we've unchecked the **Inherit** options for **Pitch, Yaw and Roll** as we do not want to inherit settings from the camera. 
+
+1.	Click **Compile** and **Save** then minimize the Blueprint (we will come back to it in a moment). 
+
+The character is setup, however there is no movement script applied to it which we will grab from the **Third Person Character** Blueprint. 
+
+
+## Finishing Up
+
+There is no need to reinvent the wheel here, since this template comes with a character with scripted functionality for movement, we can copy that into our Paper Character to provide movement for it. 
+
+[REGION:note]
+For more detailed information on setting up character movement and setting up character movement from a blank project, see the [](Gameplay/HowTo/CharacterMovement/Blueprints) documentation. 
+[/REGION]
+
+1.	Inside the **ThirdPersonBP** folder in the **Content Browser**, open the **Blueprints** folder then the **ThirdPersonCharacter** Blueprint. 
+
+1.	Inside the **Event Graph**, copy the **Movement Input** section by **Left Mouse** dragging a selection box then pressing **Ctrl+C** to copy. 
+
+	![](TopDown20.png)(w:800)
+
+	There is other script in here for 3D gameplay such as manipulating the camera, controller input and jumping. For the purposes of this tutorial we are only concerned with general movement. 
+
+1.	Return to your **Paper Character** Blueprint and press **Ctrl+V** in the **Event Graph** to paste it in. 
+
+1.	Also inside the **ThirdPersonBP** and **Blueprints** folder, open the **ThirdPersonGameMode** Blueprint. 
+
+1.	Under **Classes**, change the **Default Pawn Class** to your **Paper Character** (_TopDownCharacter_). 
+
+	![](TopDown21.png)
+
+	This will set our Paper Character as the default character to use during gameplay. 
+
+1.	Inside the level delete the Blue Character so you will not take control of it when launching the game. 
+
+	![](DeleteMe.png)
+
+1.	From the Main Tool bar, click the **Play** button to play in the editor. 
+
+	![](TopDownExample.gif)(convert:false)
+
+	Although the 2D character is idle, you can now move around the default map using the **WASD** keys. 
+
+This tutorial illustrates how to set up the foundation of a top down game by creating a playable 2D character, there are other tutorials you can take a look at to continue to develop this sample such as the **2D Character Animation How To** (to set up different Flipbook Animation states for the character) as well as the section on **Paper 2D Tile Maps** to create a 2D map. Fancy a fixed camera position instead of a follow camera for your top down game? The **How to Use Cameras in Blueprints** page will show you how to set that up and more. 
+
+
+
+
+
+
+
+
+
+
diff --git a/Engine/Documentation/Source/Engine/Paper2D/Paper2D.INT.udn b/Engine/Documentation/Source/Engine/Paper2D/Paper2D.INT.udn
--- a/Engine/Documentation/Source/Engine/Paper2D/Paper2D.INT.udn
+++ b/Engine/Documentation/Source/Engine/Paper2D/Paper2D.INT.udn
@@ -26,6 +26,7 @@ Refer to each section below for more information and the **Paper 2D** map inside
 		* [](Engine\Paper2D\Flipbooks"%Engine\Paper2D\Flipbooks:description%")
 			* [](Engine/Paper2D/Flipbooks/Editor "%Engine/Paper2D/Flipbooks/Editor:description%")
 			* [](Engine\Paper2D\Flipbooks\Components "%Engine\Paper2D\Flipbooks\Components:description%")
+		* [](Engine\Paper2D\HowTo "%Engine\Paper2D\HowTo:description%")
 		* [](Engine\Paper2D\TileMaps"%Engine\Paper2D\TileMaps:description%")
 		* [](Engine\Paper2D\Importing"%Engine\Paper2D\Importing:description%")
 		* [](Resources/ContentExamples/Paper2D "%Resources/ContentExamples/Paper2D:description%")
