#UE4 Docs
Minor update to last paragraph.

diff --git a/Engine/Documentation/Source/Engine/Paper2D/HowTo/TopDown/TopDown.INT.udn b/Engine/Documentation/Source/Engine/Paper2D/HowTo/TopDown/TopDown.INT.udn
--- a/Engine/Documentation/Source/Engine/Paper2D/HowTo/TopDown/TopDown.INT.udn
+++ b/Engine/Documentation/Source/Engine/Paper2D/HowTo/TopDown/TopDown.INT.udn
@@ -169,7 +169,7 @@ For more detailed information on setting up character movement and setting up ch
 
 	Although the 2D character is idle, you can now move around the default map using the **WASD** keys. 
 
-This tutorial illustrates how to set up the foundation of a top down game by creating a playable 2D character, there are other tutorials you can take a look at to continue to develop this sample such as the **2D Character Animation How To** (to set up different Flipbook Animation states for the character) as well as the section on **Paper 2D Tile Maps** to create a 2D map. Fancy a fixed camera position instead of a follow camera for your top down game? The **How to Use Cameras in Blueprints** page will show you how to set that up and more. 
+This tutorial illustrates how to set up the foundation of a top down game by creating a playable 2D character, there are other tutorials you can take a look at to continue to develop this sample such as the **Setting up an Animation State Machine** (to set up different Flipbook Animation states for the character to use) as well as the section on **Paper 2D Tile Maps** to create a 2D map. Want a fixed camera position instead of a follow camera for your top down game? The **How to Use Cameras in Blueprints** page will show you how to set that up and more. 
 
 
 
