#UE4doc #docs: sp: view mode

diff --git a/Engine/Documentation/Source/Engine/Rendering/LightingAndShadows/LightMobility/StationaryLights/StationaryLights.INT.udn b/Engine/Documentation/Source/Engine/Rendering/LightingAndShadows/LightMobility/StationaryLights/StationaryLights.INT.udn
--- a/Engine/Documentation/Source/Engine/Rendering/LightingAndShadows/LightMobility/StationaryLights/StationaryLights.INT.udn
+++ b/Engine/Documentation/Source/Engine/Rendering/LightingAndShadows/LightMobility/StationaryLights/StationaryLights.INT.udn
@@ -101,7 +101,7 @@ to different channels of a shadowmap texture. This is a graph coloring problem,
 than 4 overlapping allowed due to topology. Shadowing cannot affect the overlap test, so **the sunlight 
 typically requires a channel from the entire level it is in, even the underground areas.**  Once the channel 
 limit is reached, additional stationary lights will use whole scene dynamic shadows at a severe performance 
-cost. The **StationaryLightOverlap** viewmode can be used to visualize the overlap, which is updated dynamically 
+cost. The **StationaryLightOverlap** view mode can be used to visualize the overlap, which is updated dynamically 
 as you modify the lights. Light icons are changed to a red X when they are not able to allocate a channel.
 
 ![](OverlapError.png)(w:640)
@@ -131,7 +131,7 @@ lights comes from dynamic objects that it affects. This means the cost can vary
 amount, depending on how many dynamic objects there are. With enough dynamic objects, it is more efficient 
 to use a Movable light instead.
 
-In the scene below, the spheres are all Movable, and they all receive shadows from the static world and 
+In the scene below, the spheres are all movable, and they all receive shadows from the static world and 
 cast their own shadows, which merge with the distance field shadows. The _Per Object_ shadow frustums for each 
 Movable component are also shown.
 
@@ -140,7 +140,7 @@ Movable component are also shown.
 [/REGION]
 
 [REGION:note]
-_Per Object_ shadows used by Movable components apply a shadowmap to the object's bounds, therefore the bounds must be accurate.  For Skeletal meshes this means they should have a physics asset.  For particle systems, any fixed bounding box must be large enough to contain all particles.
+_Per Object_ shadows used by movable components apply a shadowmap to the object's bounds, therefore the bounds must be accurate.  For Skeletal meshes this means they should have a physics asset.  For particle systems, any fixed bounding box must be large enough to contain all particles.
 [/REGION]
 
 [/EXCERPT:StationaryShadows]
