#UE4doc #docs: minor grammar and punctuation edits

diff --git a/Engine/Documentation/Source/Engine/Rendering/LightingAndShadows/LightPropagationVolumes/LightPropagationVolumes.INT.udn b/Engine/Documentation/Source/Engine/Rendering/LightingAndShadows/LightPropagationVolumes/LightPropagationVolumes.INT.udn
--- a/Engine/Documentation/Source/Engine/Rendering/LightingAndShadows/LightPropagationVolumes/LightPropagationVolumes.INT.udn
+++ b/Engine/Documentation/Source/Engine/Rendering/LightingAndShadows/LightPropagationVolumes/LightPropagationVolumes.INT.udn
@@ -69,8 +69,8 @@ At the moment, this console variable cannot be changed at runtime.
 * Add/Adjust geometry and the light direction to give you some shadowed parts and some parts that receive the light (ideally bright colored materials).
 * Make sure you do not have any Lightmass GI enabled to be confused as dynamic GI (both techniques together would result in too bright indirect lighting).
 * Enable **Affect Dynamic Indirect Lighting** on the Directional Light.
-* Enabled Light Propagation Mode viewmode in the showflags to visualize the area that is covered by the Light Propagation Volume. [(see below)](#viewingglobalillumination)
-* Add/Change a post process volume (should be unbound) to tweak the "Size" in the LightPropagationVolume settings to cover the area with the required detail and distance. (tradeoff)
+* Enabled Light Propagation Mode viewmode in the showflags to visualize the area that is covered by the Light Propagation Volume [(See below)](#viewingglobalillumination).
+* Add/Change a post process volume (should be unbound) to tweak the "Size" in the LightPropagationVolume settings to cover the area with the required detail and distance (Tradeoff).
 * Use the GlobalIllumination showflags to enable/disable GI.
                 
 
@@ -85,7 +85,7 @@ The tooltips (hover over with the mouse) give you more details on the settings.
 
 <!-- add setting descriptions -->
 
-Having those settings in a volume allows to have them soft blendable depending on the camera position (e.g. 
+Having those settings in a volume allows them to be soft blendable depending on the camera position (e.g. 
 less GI bounce in an indoor environment or under some trees).
 
 Secondary occlusion should be enabled, and since secondary bounces are very cheap, they should also be enabled. 
