#UE4doc #docs: sp: between#UE4doc #docs: Made minor edits.

diff --git a/Engine/Documentation/Source/Engine/Rendering/LightingAndShadows/LightTypes/Directional/Directional.INT.udn b/Engine/Documentation/Source/Engine/Rendering/LightingAndShadows/LightTypes/Directional/Directional.INT.udn
--- a/Engine/Documentation/Source/Engine/Rendering/LightingAndShadows/LightTypes/Directional/Directional.INT.udn
+++ b/Engine/Documentation/Source/Engine/Rendering/LightingAndShadows/LightTypes/Directional/Directional.INT.udn
@@ -136,7 +136,7 @@ Property | Description |
 ### Light Shafts
 Property | Description |
 | -------- | ----------- |
-| **Enable Light Shaft Occlusion** | Whether to occlude fog and atmosphere inter scattering with the screenspace blurred occlusion from this light. |
+| **Enable Light Shaft Occlusion** | Whether to occlude fog and atmosphere inter scattering with the screen space blurred occlusion from this light. |
 | **Occlusion Mask Darkness** | How dark the occlusion masking is, a value of 1 is no darkening. |
 | **Occlusion Depth Range** | Everything close to the camera than this distance will occlude the light shafts. |
 | **Enable Light Shaft Bloom**  | Whether to render light shaft bloom for this light. |
@@ -167,6 +167,6 @@ Property | Description |
 | **Dynamic Shadow Distance StationaryLight** | How far **Cascaded Shadow Map** dynamic shadows will cover for a stationary light, measured from the camera. |
 | **Num Dynamic Shadow Cascades** | Number of cascades to split the view frustum into for the whole scene.  |
 | **Cascade Distribution Exponent** | Controls whether the cascades are distributed closer to the camera (larget exponent) or further from the camera (smaller exponent). |
-| **Cascade Transition Fraction** | Proportion of the fade region betweeen cascades. |
+| **Cascade Transition Fraction** | Proportion of the fade region between cascades. |
 | **Shadow Distance Fadeout Fraction** | Controls the size of the fade out region at the far extent of the dynamic shadow's influence.|
 | **Use Inset Shadows for Movable Objects** | (Stationary lights only) Whether to use per-object inset shadows for movable components, even though cascaded shadow maps are enabled.|
\ No newline at end of file
