#UE4doc #docs: made all instances "self-shadowing" for consistency

diff --git a/Engine/Documentation/Source/Engine/Rendering/LightingAndShadows/LightTypes/Spot/Spot.INT.udn b/Engine/Documentation/Source/Engine/Rendering/LightingAndShadows/LightTypes/Spot/Spot.INT.udn
--- a/Engine/Documentation/Source/Engine/Rendering/LightingAndShadows/LightTypes/Spot/Spot.INT.udn
+++ b/Engine/Documentation/Source/Engine/Rendering/LightingAndShadows/LightTypes/Spot/Spot.INT.udn
@@ -84,7 +84,7 @@ Like other lights, a **Spot Light** can be set to one of three **Mobility** sett
 ![](spacer.png)
 [/REGION]
 
-Two examples of a **Spot Light** placed inside a level are shown below demonstrating how to determine the light radius and cone effectors: .
+Two examples of a **Spot Light** placed inside a level are shown below demonstrating how to determine the light radius and cone effectors:
 
 ![](spotlight_example.png)
 
