TypoCorrection

diff --git a/Engine/Documentation/Source/Engine/Rendering/LightingAndShadows/LightPropagationVolumes/LightPropagationVolumes.INT.udn b/Engine/Documentation/Source/Engine/Rendering/LightingAndShadows/LightPropagationVolumes/LightPropagationVolumes.INT.udn
--- a/Engine/Documentation/Source/Engine/Rendering/LightingAndShadows/LightPropagationVolumes/LightPropagationVolumes.INT.udn
+++ b/Engine/Documentation/Source/Engine/Rendering/LightingAndShadows/LightPropagationVolumes/LightPropagationVolumes.INT.udn
@@ -1,7 +1,7 @@
+﻿Availability:Public
 Title:Light Propagation Volumes
 Crumbs:%ROOT%, Engine, Engine/Rendering/LightingAndShadows
 Description:System for dynamically generating global illumination.
-Availability:Public
 
 [VAR:Topic]
 [OBJECT:Topic]
