TypoCorrection

diff --git a/Engine/Documentation/Source/Engine/Paper2D/HowTo/Paper2DHowTo.INT.udn b/Engine/Documentation/Source/Engine/Paper2D/HowTo/Paper2DHowTo.INT.udn
--- a/Engine/Documentation/Source/Engine/Paper2D/HowTo/Paper2DHowTo.INT.udn
+++ b/Engine/Documentation/Source/Engine/Paper2D/HowTo/Paper2DHowTo.INT.udn
@@ -1,7 +1,7 @@
+﻿Availability:Public
 Title:Paper 2D - How To...
 Crumbs: %ROOT%, Engine, Engine/Paper2D
 Description:The Paper 2D How To pages provides several short step-by-step guides for working with Paper 2D. 
-Availability:Public
 Related: Engine\Paper2D\Sprites
 Related: Engine\Paper2D\Flipbooks
 
