TypoCorrection

diff --git a/Engine/Documentation/Source/Engine/Paper2D/HowTo/TopDown/TopDown.INT.udn b/Engine/Documentation/Source/Engine/Paper2D/HowTo/TopDown/TopDown.INT.udn
--- a/Engine/Documentation/Source/Engine/Paper2D/HowTo/TopDown/TopDown.INT.udn
+++ b/Engine/Documentation/Source/Engine/Paper2D/HowTo/TopDown/TopDown.INT.udn
@@ -1,7 +1,7 @@
+﻿Availability:Public
 Title:Top Down 2D Game Setup
 Crumbs: %ROOT%, Engine, Engine/Paper2D
 Description:How to quickly set up a top down 2D game with a playable character. 
-Availability:Public
 Related: Engine\Paper2D\HowTo\Animation
 Related: Engine\Paper2D\TileMaps
 Related: Gameplay/HowTo/UsingCameras/Blueprints
