TypoCorrection

diff --git a/Engine/Documentation/Source/Engine/Paper2D/HowTo/Animation/2DCharacterAnimation.INT.udn b/Engine/Documentation/Source/Engine/Paper2D/HowTo/Animation/2DCharacterAnimation.INT.udn
--- a/Engine/Documentation/Source/Engine/Paper2D/HowTo/Animation/2DCharacterAnimation.INT.udn
+++ b/Engine/Documentation/Source/Engine/Paper2D/HowTo/Animation/2DCharacterAnimation.INT.udn
@@ -1,7 +1,7 @@
+﻿Availability:Public
 Title:Setting up an Animation State Machine
 Crumbs: %ROOT%, Engine, Engine/Paper2D
 Description:How to set up and play different animations for a 2D character based on different states. 
-Availability:Public
 Related: Engine\Paper2D\HowTo\TopDown
 Related: Engine\Paper2D\Sprites
 Related: Engine\Paper2D\Flipbooks
