TypoCorrection

diff --git a/Engine/Documentation/Source/Engine/Paper2D/HowTo/Physics/SpritePhysics.INT.udn b/Engine/Documentation/Source/Engine/Paper2D/HowTo/Physics/SpritePhysics.INT.udn
--- a/Engine/Documentation/Source/Engine/Paper2D/HowTo/Physics/SpritePhysics.INT.udn
+++ b/Engine/Documentation/Source/Engine/Paper2D/HowTo/Physics/SpritePhysics.INT.udn
@@ -1,7 +1,7 @@
+﻿Availability:Public
 Title:Applying 2D Physics
 Crumbs: %ROOT%, Engine, Engine/Paper2D
 Description:Examples of implementing physics for 2D games. 
-Availability:Public
 Related: Engine\Paper2D\Sprites
 Related: Engine\Paper2D\Flipbooks
 
