TypoCorrection

diff --git a/Engine/Documentation/Source/Engine/Paper2D/HowTo/Sockets/SpriteSockets.INT.udn b/Engine/Documentation/Source/Engine/Paper2D/HowTo/Sockets/SpriteSockets.INT.udn
--- a/Engine/Documentation/Source/Engine/Paper2D/HowTo/Sockets/SpriteSockets.INT.udn
+++ b/Engine/Documentation/Source/Engine/Paper2D/HowTo/Sockets/SpriteSockets.INT.udn
@@ -1,7 +1,7 @@
+﻿Availability:Public
 Title:Using Sprite Sockets
 Crumbs: %ROOT%, Engine, Engine/Paper2D
 Description:How to attach objects to sockets defined in a Paper 2D Sprite. 
-Availability:Public
 Related: Engine\Paper2D\Sprites
 Related: Engine\Paper2D\Flipbooks
 
