#UE4doc #docs: changed "viewmode" and "realtime" to "view mode" and "real time"  and "a" to "an"

diff --git a/Engine/Documentation/Source/Engine/Rendering/LightingAndShadows/ReflectionEnvironment/ReflectionEnvironment.INT.udn b/Engine/Documentation/Source/Engine/Rendering/LightingAndShadows/ReflectionEnvironment/ReflectionEnvironment.INT.udn
--- a/Engine/Documentation/Source/Engine/Rendering/LightingAndShadows/ReflectionEnvironment/ReflectionEnvironment.INT.udn
+++ b/Engine/Documentation/Source/Engine/Rendering/LightingAndShadows/ReflectionEnvironment/ReflectionEnvironment.INT.udn
@@ -59,7 +59,7 @@ with other reflection techniques, such as Screen Space Reflections.
 To use the Reflection Environment, first setup lighting and build lighting once.  **There must be some indirect 
 diffuse lighting for the Reflection Environment to show up at all.**  Next, drag *_SphereReflectionCapture_* Actors 
 from the _Visual_ tab of the _Place Mode_ of the **Modes** panel into your level and make sure your materials have noticeable specular, and a low roughness.  
-Use the Reflection Override viewmode to visualize what is being captured.
+Use the Reflection Override view mode to visualize what is being captured.
 
 ## Glossy Indirect Specular
 
@@ -83,7 +83,7 @@ is especially important for materials like metal which have no diffuse response.
 [EXCERPT:RefEnvIntro]
 The Reflection Environment works by capturing the static scene at many points and reprojecting it onto 
 simple shapes like spheres in reflections. Artists choose the capture points by placing **_ReflectionCapture_** 
-Actors. Reflections update in realtime during editing to aid in placement but are static at runtime. 
+Actors. Reflections update in real time during editing to aid in placement but are static at runtime. 
 Projecting the captured scene onto simple shapes gives approximate parallax for the reflection. Each pixel 
 blends between multiple cubemaps to get the final result.  Smaller ReflectionCapture Actors override larger ones, so you can refine reflection parallax accuracy in areas as needed.  For 
 example, you might place a capture at the center of a room and then refine the reflection by placing smaller 
@@ -140,7 +140,7 @@ box is projected onto the box shape, creating significant artifacts in many case
 
 ![Box Shape](BoxShape.png)
 
-The box has a orange preview for the projection shape when selected. It only captures the scene within 
+The box has an orange preview for the projection shape when selected. It only captures the scene within 
 **Box Transition Distance** outside this box. The influence of this capture fades in over the transition 
 distance as well, inside the box. 
 
