#UE4doc #docs: changed "to comment your Materials" to "to add comments to your Materials" "realtime" to "real time"

diff --git a/Engine/Documentation/Source/Engine/Rendering/Materials/Editor/MaterialEditor.INT.udn b/Engine/Documentation/Source/Engine/Rendering/Materials/Editor/MaterialEditor.INT.udn
--- a/Engine/Documentation/Source/Engine/Rendering/Materials/Editor/MaterialEditor.INT.udn
+++ b/Engine/Documentation/Source/Engine/Rendering/Materials/Editor/MaterialEditor.INT.udn
@@ -55,8 +55,8 @@ and choosing New Material from the context menu, or clicking the **New** button
 ## Commenting Material Expression Networks
 
 Comments are a great way to document what your Material is doing, making it easier for you and others to understand 
-what is going on within a complicated material graph. There are two ways to comment your Materials: adding text to 
-the **Desc** property of individual nodes, or placing nodes within a Comment object frame.
+what is going on within a complicated material graph. There are two ways to add comments to your Materials: by adding text to 
+the **Desc** property of individual nodes or by placing nodes within a Comment object frame.
 
 ### Using the Desc Property
 
@@ -98,7 +98,7 @@ When making changes to a network, it can be useful to get immediate feedback fro
 offers two features in the Graph to give you instantaneous feedback for your network: Live Nodes and Live Update.
 
 The difference between these two options may not be apparent at first. Live Nodes allows for any constant changes that the nodes 
-make - due to such expressions as Panners - to playback in realtime. 
+make - due to such expressions as Panners - to playback in real time. 
 
 Live Update, on the other hand, allows for each node within the network to have its shader recompiled whenever a change is made. 
 These changes include creation of new nodes, deletion of nodes, node connections and disconnections, and changes in properties. 
