TypoCorrection

diff --git a/Engine/Documentation/Source/Programming/UnrealArchitecture/TArrays/TArrays.INT.udn b/Engine/Documentation/Source/Programming/UnrealArchitecture/TArrays/TArrays.INT.udn
--- a/Engine/Documentation/Source/Programming/UnrealArchitecture/TArrays/TArrays.INT.udn
+++ b/Engine/Documentation/Source/Programming/UnrealArchitecture/TArrays/TArrays.INT.udn
@@ -1,4 +1,4 @@
-Title: TArray: Arrays in Unreal Engine
+﻿Title: TArray: Arrays in Unreal Engine
 Crumbs:
 Description:
 Availability:Public
@@ -78,7 +78,7 @@ SetNum will also remove elements if the new number is less than the current one.
 
 ## Iteration
 
-There are several ways to iterate over the elements of your array, but the recommended way is to use C++’s ranged-for feature:
+There are several ways to iterate over the elements of your array, but the recommended way is to use C++'s ranged-for feature:
 
 	FString JoinedStr;
 	for (auto& Str : StrArr)
@@ -111,14 +111,14 @@ Arrays can be sorted simply by calling the Sort method:
 	StrArr.Sort();
 	// StrArr == ["!","Brave","Hello","of","Tomorrow","World"]
 
-Here, the values are sorted by means of the element type’s operator<.  In FString’s case, this is a case-insensitive lexicographical comparison.  A binary predicate can also be provided to provide different ordering semantics, e.g.:
+Here, the values are sorted by means of the element type's operator<.  In FString's case, this is a case-insensitive lexicographical comparison.  A binary predicate can also be provided to provide different ordering semantics, e.g.:
 
 	StrArr.Sort([](const FString& A, const FString& B) {
 		return A.Len() < B.Len();
 	});
 	// StrArr == ["!","of","Hello","Brave","World","Tomorrow"]
 
-Now the strings are sorted by their lengths.  Note how the three strings with the same length - "Hello", “Brave” and “World” - have changed order relative to their positions in the array beforehand.  This is because Sort is unstable and the relative order of equivalent elements (those strings are equivalent here, because they have the same length) is not guaranteed.  Sort is implemented as a quicksort.
+Now the strings are sorted by their lengths.  Note how the three strings with the same length - "Hello", "Brave" and "World" - have changed order relative to their positions in the array beforehand.  This is because Sort is unstable and the relative order of equivalent elements (those strings are equivalent here, because they have the same length) is not guaranteed.  Sort is implemented as a quicksort.
 
 The HeapSort method, with or without a binary predicate, can be used to perform a heap sort.  Whether or not you choose to use it depends on your particular data and how efficiently it sorts compared to the Sort method.  Like Sort, HeapSort is not stable.  If we had used HeapSort instead of Sort above, this would be the result (the same, in this case):
 
@@ -175,7 +175,7 @@ Passing an invalid index — less than 0 or greater than or equal to Num() — w
 	// bValid5  == true
 	// bValid6  == false
 
-operator[] returns a reference, so it can also be used to mutate the elements inside the array, assuming your array isn’t const:
+operator[] returns a reference, so it can also be used to mutate the elements inside the array, assuming your array isn't const:
 
 	StrArr[3] = StrArr[3].ToUpper();
 	// StrArr == ["!","of","Brave","HELLO","World","Tomorrow"]
@@ -222,7 +222,7 @@ This sets Index to be the index of the first element found.  If there are duplic
 	int32 IndexLast;
 	if (StrArr.FindLast(TEXT("Hello"), IndexLast))
 	{
-		// IndexLast == 3, because there aren’t any duplicates
+		// IndexLast == 3, because there aren't any duplicates
 	}
 
 Both of these functions return a bool to indicate whether or not an element was found, while also writing the index of that element into a variable when it was found.
@@ -284,9 +284,9 @@ You can erase elements from the array using the Remove family of methods.  The R
 	// StrArr == ["!","of","Brave","World","Tomorrow"]
 	
 	StrArr.Remove(TEXT("goodbye"));
-	// StrArr is unchanged, as it doesn’t contain "goodbye"
+	// StrArr is unchanged, as it doesn't contain "goodbye"
 
-Note that "HELLO" was removed even though we asked it to remove "hello".  Equality is tested via the element type’s operator==; remember for FString, this is a case-insensitive comparison.
+Note that "HELLO" was removed even though we asked it to remove "hello".  Equality is tested via the element type's operator==; remember for FString, this is a case-insensitive comparison.
 
 The final element of an array can be removed by the Pop method:
 
@@ -309,7 +309,7 @@ You can also use RemoveSingle to erase the closest element to the front of the a
 	ValArr.RemoveSingle(30);
 	// ValArr == [10,5,10,15,25,30]
 
-We can also remove elements by their index by using the RemoveAt method.  The index must exist or it’ll cause a runtime error.  Remember that indices are zero-based:
+We can also remove elements by their index by using the RemoveAt method.  The index must exist or it'll cause a runtime error.  Remember that indices are zero-based:
 
 	ValArr.RemoveAt(2); // Removes the element at index 2
 	// ValArr == [10,5,15,25,30]
@@ -328,7 +328,7 @@ We can also remove elements which match a predicate by using the RemoveAll metho
 In all of these cases when elements were removed, the elements which followed it were shuffled down into lower indices, as there can never be 'holes' left in the array.
 [/REGION]
 
-The shuffling process has an overhead.  If you don’t really care what order the remaining elements are left in, this overhead can be reduced by using the RemoveSwap, RemoveAtSwap and RemoveAllSwap methods, which work like their non-swapping variants except that they don’t guarantee the order of the remaining elements, which allows them to be implemented more efficiently:
+The shuffling process has an overhead.  If you don't really care what order the remaining elements are left in, this overhead can be reduced by using the RemoveSwap, RemoveAtSwap and RemoveAllSwap methods, which work like their non-swapping variants except that they don't guarantee the order of the remaining elements, which allows them to be implemented more efficiently:
 
 	TArray<int32> ValArr2;
 	for (int32 i = 0; i != 10; ++i)
@@ -353,7 +353,7 @@ Finally, all elements can be removed by using the Empty method:
 
 ##Operators
 
-Arrays are regular value types and as such can be copied via the standard copy constructor or assignment operator.  As arrays strictly own their elements, copying an array is ‘deep’ and so the new array will have its own copy of the elements:
+Arrays are regular value types and as such can be copied via the standard copy constructor or assignment operator.  As arrays strictly own their elements, copying an array is 'deep' and so the new array will have its own copy of the elements:
 
 	TArray<int32> ValArr3;
 	ValArr3.Add(1);
