#UE4 Docs: Adding the documents for the Vertex Animation doccs.

diff --git a/Engine/Documentation/Source/Engine/Animation/Animation.INT.udn b/Engine/Documentation/Source/Engine/Animation/Animation.INT.udn
--- a/Engine/Documentation/Source/Engine/Animation/Animation.INT.udn
+++ b/Engine/Documentation/Source/Engine/Animation/Animation.INT.udn
@@ -107,6 +107,18 @@ Navigation:topic
 	[/PARAM]
 [/OBJECT]
 
+[OBJECT:TopicButtonList]
+	[PARAM:icon]
+		![](%ROOT%/start_icon.png)(convert:false) 
+	[/PARAM]
+	[PARAM:title]
+		Tools
+	[/PARAM]
+	[PARAM:links]
+		* [](Engine\Animation\Tools\VertexAnimationTool "%Engine\Animation\Tools\VertexAnimationTool:description%")
+	[/PARAM]
+[/OBJECT]
+
 [OBJECT:TopicButtonList]
 	[PARAM:icon]
 		![](%ROOT%/related_icon.png)(convert:false)
diff --git a/Engine/Documentation/Source/Engine/Animation/Tools/VertexAnimationTool/VAT_KF_Meshes/VAT_KF_Meshes.INT.udn b/Engine/Documentation/Source/Engine/Animation/Tools/VertexAnimationTool/VAT_KF_Meshes/VAT_KF_Meshes.INT.udn
new file mode 100644
--- /dev/null
+++ b/Engine/Documentation/Source/Engine/Animation/Tools/VertexAnimationTool/VAT_KF_Meshes/VAT_KF_Meshes.INT.udn
@@ -0,0 +1,116 @@
+Availability:Public
+Title:Vertex Animation Tool - Key Framed Meshes
+Crumbs: %ROOT%, Engine, Engine/Animation
+Description: User guide for the 3D Max Vertex Animation tool set.
+
+[TOC(start:2 end:3)]
+
+## Overview
+You can also use the Vertex Animation tool with meshes that do not make use of the Animation Timeline in 3Ds Max. 
+Generally speaking mesh animations that use this method have been exported frame by frame from another 3D package like  Maya or Blender.
+In this section we will go over how to use this method as well as any differences in setup that you need to be aware of.
+
+
+## 3Ds Max Content Creation & Setup for Key Framed Meshes
+
+[region:note]
+Please note that while we are using 3Ds Max for this example the method in which the content is created and setup inside 3Ds Max is no different than how content could be set up inside other 3D Packages such as Maya or Blender.
+[/region]
+
+1. First inside 3s Max create a new scene and then create a Teapot with the following settings and place it at 0,0,0 in the world.
+	* **Radius:** 70.0
+	* **Segments:** 5
+
+1. Next from the **Modifier List** add a **Melt Modifier** to the Teapot mesh. 
+
+	![](VAT_KFM_Teapot_Melt.png)
+
+1. Now make a copy of the mesh by holding down the **Shift** key and moving the Teapot along the X axis with the mouse. 
+Position the Teapot copy so that the orange bounding box of the Melt modifier does not intersect the original mesh. 
+Once the Teapot copy is in position release the mouse and the shift key and the Clone Options will be displayed. 
+Make sure the **Object** is set to **Copy** and then set the **Number of Copies** to **Four** and press the **Ok** button to create the copies.
+
+	![](VAT_KFM_Teapot_Copy.gif)(convert:false)
+
+1. Then select the 2nd Teapot mesh and set the **Melt Amount** in the Melt modifier to 25. 
+Then select the next Teapot in line and set the **Melt Amount** of the Melt modifier to 50 and the next one to 75% and finally set the last Teapot **Melt Amount** to 100.
+
+	![](VAT_KFM_Teapot_Melt_Amount.gif)(convert:false)
+
+1. With all the Teapots now having some type of deformation on them, select all the Teapots and then with the mouse Right-Click  in the 3Ds Max viewport and select **Convert To:** > **Convert to Editable Poly**.
+
+	[region:note]
+	You have to convert any mesh that you want to use with this tool to Editable Poly. The tool **WILL NOT** work with any other mesh type.
+	[/region]
+
+1. Now that the Teapot meshes have all been converted to Editable Poly, select all the meshes and move them to 0,0,0 in the world. 
+Once the meshes are in the center of the world, deselect the Teapot meshes and then select them in the order that the animation is supposed to play. 
+For this example we want to select Teapot001 first  then select Teapot002 and so on till all five are selected.
+
+	![](VAT_KFM_Teapot_Export_Setup.gif)(convert:false)
+
+	[region:note]
+	The Sequence Painter Tool requires that you select the meshes in the order that they should be played in the animation. Failing to do this will result in the animation that gets played in UE4 not syncing up to the bone driven animation that it is supposed to replace.
+	[/region]
+
+1. With the Teapots now in the center of the world and selected in the order that the animation is supposed to be played, it is now time to run the **Sequence Painter** script. 
+To do this all you need to do is press the **Paint Selection Sequence** button that is in the **Sequence Painter** section of the tool. 
+
+	![](VAT_KFM_Paint_Seq.png)
+
+	[region:note]
+	Depending on the amount of vertices your mesh has the time that it takes for the **Paint Selection Sequence** to finish is quite fast. 
+	[/region]
+
+1. Now that position data has been saved to the UVs of the Teapot, it is time to export the Teapot mesh for use in the UE4. When exporting the Teapot mesh use the same setting as you would for any other Static Meshes and make sure that **Smoothing Groups** are enabled. For this example the Teapot mesh was given the name **SM_Teapot_PSS.FBX**.
+
+
+## UE4 Content Setup & Use for Key Framed Meshes
+Once the Teapot mesh has been exported from 3Ds Max, it is now time for the Teapot mesh to be imported and setup so that it will work the same inside UE4.
+
+1. First import the Teapot mesh, **SM_Teapot_PSS.FBX** by dragging it form the location it was exported to into the **Content Browser**, making sure that the following options are enabled / disabled.
+
+	* Disable the **Auto Generate Collision** option.
+	* Disable the **Remove Degenerates** option.
+	* Disable the **Generate Lightmap UV's** option.
+	* Enable the **Combine Meshes** option.
+	* Disable both **Import Materials** & **Import Textures** option. 
+	
+		![](VAT_SM_Import_Options.png)
+
+1. After the Teapot mesh has been imported, open the mesh up by Double - Clicking on it with the Left Mouse button in Content Browser. 
+Once the mesh has been opened enable the **Use Full Precision UVs** option and set the **Distance Field Resolution Scale** to **0.0** then press the **Apply Changes** button..
+
+	![](VAT_KFM_UE4_Setup.png)
+
+1. With the Teapot mesh now set up correctly, create a new Material in the Content Browser and name it **MAT_Sequance_Painter** and then open up the Material. 
+
+	![](VAT_KFM_Create_Material.png)
+
+1. Once the Material is opened up, add Two Constant Parameter nodes to the Material Graph by holding down the one key on the Keyboard and Left - Clicking with the mouse in the Material Graph to place the Constant Parameter nodes.
+Set the value of one of the Constant Parameter node and plug that into the **Base Color** input on the Main Material Node. Then set the value of the other Constant Parameter node to 5.
+
+	![](VAT_KFM_Set_Values.png)
+
+1. Next use the mouse and Right - Click in the Material Graph to expose the search menu and then search for **MS_SequencePainter_SequenceFlipbook** Material Function using MS as the search term.
+When you find the MS_SequencePainter_SequenceFlipbook Material Function click on it with the Left Mouse button to add it to the Material Graph.
+
+	![](VAT_KFM_Adding_MF_Flipbook.gif)(convert:false)
+
+1. With the correct Material Function now in place, connect the output of the **MS_SequencePainter_SequenceFlipbook**,  into the World Position Offset input on the Main Material node and then connect the remaining Scalar parameter to the **Number Of Frames** input on the MS_SequencePainter_SequenceFlipbook Material Function node.
+
+	![](VAT_KFM_Linking_Up_MF.gif)(convert:false)
+
+1. Now that everything has been hooked up press the **Apply** button to compile the Material and then press the **Save** button. 
+After that is completed you can then close down the Material Graph Editor.
+
+	![](VAT_Compile_Save.png)
+
+1. With the Material created, compiled, and saved drag the Teapot mesh from the Content Browser into the level and then drag the Material that was just created for it from the Content Browser on to the mesh that was placed in the level. 
+Once this has been completed you will now see the mesh "playing" the animation.
+
+	![](VAT_KFM_Final_Effect.gif)(convert:false)
+
+
+
+
diff --git a/Engine/Documentation/Source/Engine/Animation/Tools/VertexAnimationTool/VAT_TL_Meshes/VAT_TL_Meshes.INT.udn b/Engine/Documentation/Source/Engine/Animation/Tools/VertexAnimationTool/VAT_TL_Meshes/VAT_TL_Meshes.INT.udn
new file mode 100644
--- /dev/null
+++ b/Engine/Documentation/Source/Engine/Animation/Tools/VertexAnimationTool/VAT_TL_Meshes/VAT_TL_Meshes.INT.udn
@@ -0,0 +1,165 @@
+Availability:Public
+Title:Vertex Animation Tool - Timeline Meshes
+Crumbs: %ROOT%, Engine, Engine/Animation
+Description: User guide for the 3D Max Vertex Animation tool set.
+
+[TOC(start:2 end:3)]
+
+## Overview
+In the following section we will cover how you can use the 3Ds Max Animation Timeline in conjunction with the Vertex Animation tools to create a 2D Texture that stores the position of vertices in any given 3D mesh.
+
+## 3Ds Max Content Creation & Setup for Timeline Animated Meshes
+
+
+1. First from the **Standard Primitives** create a **Teapot** at 0,0,0 in the world with the following options.
+	* **Radius:** 70.0
+	* **Segments:** 5
+	
+	![](VAT_Teapot_Setup_00.png)
+
+1. Next from the Modifier List, apply a **Bend** Modifier to the Teapot and then set the **Direction** to **90**. 
+	
+	![](VAT_Teapot_BendMod_00.png)
+
+1. Now enable **Auto Key** and then move the **Time Slider** to **10** and set the **Angle** on the **Bend Modifier** to **-50** and then move the **Time Slider** to **25** and set the Bend **Angle** to **50**. Finally move the **Time Slider** to **35** and then set the Bend **Angle** to **0**. Once all of that has been completed make sure to turn off **Auto Key** to ensure that you do not place any extra keys by accident. 
+	
+	![](VAT_Max_Bend_Setup.gif)(convert:false)
+
+1. With the Teapot now bending, open up the Vertex Animation Tools, if not already open, and set the **Animation Options** to **Animated Meshes** and the **Ainm End** to **35**.
+	
+	![](VAT_Tool_Settings_Setup.png)
+
+1. Now that the correct tool options have been set, press the **Process Animated Meshes** button to being the Texture creation process. 
+About half way through the process a Windows Dialogue box will be displayed asking for a location to store the Textures that are generated.
+For this example a folder was placed on the Desktop called **Vertex_Anims** and the Textures where saved to that folder.
+When completed the folder that was placed on the Desktop should now have two Texture files inside of it,
+a **.EXR** Texture which contains the vertices position data and a **.BMP** Texture that contains the Teapots Normal information.
+
+	![](VAT_Output_Textures.png)	
+
+1. With the Textures now created go back to 3Ds Max and select the Teapot mesh with, **_MorphExport**, added to the name and export it as a .FBX file to the same folder that the Textures where exported to.
+Name the mesh **SM_Teapot_Bend_00** and in the FBX options make sure that Smoothing Groups are enable and leave all other setting with the default values.
+
+	![](VAT_3DsMax_Teapot_Export.gif)(convert:false)
+
+
+## UE4 Content Setup & Use for Timeline Animated Meshes
+Once all of the content has been exported from 3Ds Max it is now time to import and set that content up so that it will work correctly inside UE4. 
+Please keep in mind that each and every step listed below **MUST BE** completed exactly how it is described below or you run the risk of things not working correctly inside UE4. 
+
+### Texture Import & Setup
+
+[region:tip]
+If you are unfamiliar with the UE4 Texture Import process please read the [Texture Import Guide](https://docs.unrealengine.com/latest/INT/Engine/Content/Types/Textures/Importing/index.html) for more information.
+[/region]
+
+The first thing that needs to be completed is that both Textures that were exported from 3Ds Max need to be imported into UE4. 
+To do select each Texture from it's export location in Windows Explorer and drag them into the Content Browser, releasing the left mouse button to trigger the import process.
+When you do this, you will be prompted with a warning message about the Textures not being a power of 2 resulting in the Texture not working with Streaming or Mip Mapping.
+Press the **Yes** button to dismiss the prompt and continue with the import process.
+
+![](VAT_Texture_Import.gif)(convert:false)
+
+Once each Texture has been imported, they then need to be setup in a very specific manner to ensure that they work correctly.
+The first Texture to setup is the .EXR Texture that contains the vertices position data.
+Open up this Texture by Double- Clicking on it in the Content Browser with the Left Mouse Button and make sure the following options are set / enabled.
+
+* Make sure the Mip Gen Settings is set to **No Mip Maps**.
+* Change the Texture Group to **UI**.
+* Make sure that **SRGB** has been disabled.
+* Set the Compression Settings to **TC_HDR**.
+
+	![](VAT_EXR_Texture_Setup.png)
+
+Next open up the Normal Map Texture and make sure that the following options are set / enabled.
+
+* Make sure the Mip Gen Settings is set to **No Mip Maps**.
+* Change the Texture Group to **UI**.
+* Make sure that **SRGB** has been disabled.
+* Set the Compression Settings to **TC_Vector Displacement Map**.
+	
+	![](VAT_Normal_Texture_Setup.png)
+
+With the Textures now set up correctly, make sure to press the **Save All** button in the Content Browser to save all changes that have been made. 
+
+![](VAT_Save_All.png)
+
+### Static Mesh Import & Setup
+Now that the Textures have been imported and setup, it is now time to import and setup the Teapot Static Mesh. 
+Import the Teapot Static Mesh in the same way the Texture were imported by first selecting the Teapot Static Mesh in Windows Explorer and then dragging it into the Content Browser and releasing the Left Mouse Button to trigger the import process. 
+When you do this the following dialogue box will be displayed. From this dialogue make sure that the following options are set / disabled.
+
+[region:tip]
+If you are unfamiliar with how to import Static Meshes into UE4 please check out the [FBX Static Mesh Pipeline](https://docs.unrealengine.com/latest/INT/Engine/Content/FBX/StaticMeshes/index.html) document for more information about this..
+[/region]
+
+* Disable the **Auto Generate Collision** option.
+* Disable the **Remove Degenerates** option.
+* Disable the **Generate Lightmap UV's** option.
+* Enable the **Combine Meshes** option.
+* Disable both **Import Materials** & **Import Textures** option. 
+	
+	![](VAT_SM_Setup_Options.png)
+
+Once all those options have been set, press the Import button to begin the import process. 
+After the Static Mesh has been imported, open the Static Mesh up and make sure the following options are enabled / disabled.
+
+* Enable the **Use Full Precision UVs** option.
+* Set **Distance Field Resolution Scale** to 0.0.
+
+	![](VAT_SM_Setup_Options.png)
+
+With the Static Mesh now set up correctly, make sure to press the **Apply Changes** button in the Static Mehs editor and then press the **Save All** button in the Content Browser to save all changes that have been made so far. 
+
+
+### UE4 Material Setup & Use
+
+Now that all of the required content has been imported and properly setup, it is now time to create a Material that can use the content.
+
+[region:tip]
+If you are unfamiliar with how Materials work inside UE4 please refer to the [Materials](https://docs.unrealengine.com/latest/INT/Engine/Rendering/Materials/index.html) documentation for more information.
+[/region]
+
+1. First create a new Material by Right-Clicking with the Mouse and then choosing the **Material** option that is under the **Create Basic Asset** section. Name the Material **MAT_Vertex_Animation**.
+
+1. Next open up the **MAT_Vertex_Animation** Material and locate the two Textures that were imported into the Generic Browser in the previous step. Select both Textures and then Drag them from the Content Browser and Drop them into the Material Graph.  
+	
+	![](VAT_Materail_Adding_Textures.gif)(convert:false)
+
+1. Then select the **Main Material Node** and in **Details Panel** search for **Num Customized UVs**. When you located the **Num Customized UVs** input, input a value of **4**. You should now see four new inputs on the Main Material Node with Customized UV and then a number starting at 0 going to 3.
+	
+	![](VAT_Cus_UVs.png)
+
+1. Next while still in the Main Material Node Details Panel, **Un-Check** the **Tangent Space Normal** option. 
+
+	![](VAT_UnCheck_Tan_Normals.png)
+
+1. Then hold down the One key on the Keyboard and then click inside the Material Graph with the Left Mouse Button to place down a Constant Parameter node. Set the value of the Constant to **0.5** and then plug the output into the **BaseColor** input in the Main Material Node.
+	
+	![](VAT_Add_Const.png)
+
+1. Now it is time to add the Material Function that powers the effect. To do this Right-Click with the Mouse inside the Material Graph and then from the menu that is displayed, type **MS** to show the Maxscript Material Functions. Then click on the **MS_VertexAnimationTools_MorphTargets** Material Function with the mouse to add it to the Material Graph. 
+
+	![](VAT_Add_MF.gif)(convert:false)
+
+1. With all the required Textures and Material Function now added to the Material Graph, we can now begin to connect everything together but first, select both Texture Samplers and convert them to **Texture Objects** by Right-Clicking on one of the Textures and selecting the **Convert To Texture Object** option.
+	
+	![](VAT_Convert_To_TexOBJ.gif)(convert:false)
+
+1. Next connect the output from the Normal Texture Object to the **Morph Normal(T2d)** input. Then connect the .EXR Texture Object to the **Morph Texture (T2d)** input. Finally connect the outputs of the **MS_VertexAnimationTools_MorphTargets** Material Function to the related inputs on the Main Material node.
+
+	![](VAT_Wired_Up.png)
+
+1. Now add another Constant Parameter to the Material Graph by holding down the one key on the Keyboard and then clicking in the Material Graph with the Left Mouse Button to place the node in the Material Graph. Once the Constant Node has been added set it's value to 51 and connect it to the **Number of Morph Targets** input on the **MS_VertexAnimationTools_MorphTargets** Material Function.
+
+	![](VAT_Num_Morph_Targets.png)
+
+	[region:note]
+	The number that is used for the **Number of Morph Targets** input comes from the Y value of the .EXR and Normal Textures. Make sure to always use the number that comes from the Texture as using any other number could result in your effect not working the same in UE4 as it did in 3Ds Max. 
+	[/region]
+
+1. Once everything has been connected press the **Apply** button to compile the Material and then press the **Save** button to save your work and once that is completed closed down the Material Graph Editor.
+
+1. From the **Content Browser** select the Teapot mesh that was imported and drag it from the Content Browser into the level. Then once the Teapot mesh is in the level, select the Material that was created, **MAT_Vertex_Animation**, and apply it to the Teapot mesh by dragging it from the Content Browser on to the Teapot mesh that is placed in the level.
+
+	![](VAT_Final_Effect.gif)(convert:false)
\ No newline at end of file
diff --git a/Engine/Documentation/Source/Engine/Animation/Tools/VertexAnimationTool/VertexAnimationTools.INT.udn b/Engine/Documentation/Source/Engine/Animation/Tools/VertexAnimationTool/VertexAnimationTools.INT.udn
new file mode 100644
--- /dev/null
+++ b/Engine/Documentation/Source/Engine/Animation/Tools/VertexAnimationTool/VertexAnimationTools.INT.udn
@@ -0,0 +1,162 @@
+Availability:Public
+Title:Vertex Animation Tool
+Crumbs: %ROOT%, Engine, Engine/Animation
+Description: User guide for the 3D Max Vertex Animation tool set.
+Related: Engine\Animation\Tools\VertexAnimationTool\VAT_KF_Meshes
+Related: Engine\Animation\Tools\VertexAnimationTool\VAT_TL_Meshes
+
+[TOC(start:2 end:3)]
+
+## Overview
+Storing complex animation data into 2D Textures or a meshes UV's is a great way to reduce animation overhead while still maintaining the required look and feel of the animation. 
+In the past you might have used [Morph Targets](https://docs.unrealengine.com/latest/INT/Engine/Content/FBX/MorphTargets/index.html) to do the exact same thing however this new method has a few advantages that Morph Targets do not.
+One of those advantages is allowing you to use complex animation data inside systems that would not otherwise be capable of utilizing Morph Targets like the Cascade Particle Editor.
+In the following documents we will cover how you can use this script to generate content for your UE4 projects.
+
+## Prerequisites
+Before we can get to using the script there are few things that you must be made aware of so you can get the best results possible when using the technique. 
+
+### Script Location
+[region:warning]
+This script is only available in Unreal Engine 4 (UE4) Version **4.8** or later. Earlier versions of UE4 **WILL NOT** have this script or required Material Functions available for use.
+[/region]
+
+The Vertex Animation Script will be in the following location and is called **VertexAnimationTools.ms.**
+
+		4.8\Engine\Extras\3dsMaxScripts\VertexAnimationTools.ms
+
+### Tool Caveats
+While this tool is very good at reducing animation overhead by storing complex animation data into a texture, there are a few drawbacks to using the tool. 
+First of all the tool can only affect a maximum of **8192 vertices** in a single 2D Texture. 
+This is because the maximum Texture size a Texture can be for DirectX 11 is 8192 pixels in either the X or Y direction. 
+The tool generates the data in the Texture using the following formula. 
+
+		Final texture resolution:  X = number of vertex in the mesh, Y = Number of frames captured.
+
+Due to this limitation this tool is best suited for Visual Effects or ambient Static Meshes that require animations but might not justify using a complex animation rig to do so. 
+This tool will also **NOT WORK** with Skeletal Mesh Animations as the bone transforms are not available to the Material Editor.
+This means that if you want to effect the vertices of a Skeletal Mesh in a similar manner you will need to use [Morph Targets](https://docs.unrealengine.com/latest/INT/Engine/Content/FBX/MorphTargets/index.html) instead.
+
+### Vertex Animation Tools Break Down
+Inside of the Vertex Animation Tools you will find two entirely different methods for effecting the vertices of Static Meshes. 
+In the following section we will go over what these two methods are and how they differ from one another.
+
+* **Vertex Animation Tools:** The top section of the Vertex Animation Tools, called Vertex Animation Tools, is for generating a 2D Texture that will store the morph target vertex positions and normal's.
+
+	![](VAT_Tool_Breakdown_01.png)
+
+	|Property Name| Description|
+	|-------------|------------|
+	|**Animation Options:** | Animation Options allows you to pick between using animation that was created using the Timeline in 3Ds Max or use individual Key Frame Meshes that were created in 3Ds Max or another 3D package like Maya or Blender and then exported out of that package frame by frame so the animation could then be re-constructed inside 3Ds Max.|
+	|**Process Animated Meshes:** | This button will process the animated meshes that are in the 3Ds Max scene, creating and then exporting the required Textures. |
+	|**Anim Start:** |This option defines the frame that the animation should start from. |
+	|**Anim Start:** |This option defines the frame that the animation should end on. |
+	|**Frame Skip:** |This option allows you skip frames to try and save on Texture space. |
+	|**Process Selected Meshes:** |This option, only available if Key Frame Meshes is enabled, will process the selected Key Frame Meshes in the order that they where selected in. |
+
+
+* **Sequence Painter:** The Sequence Painter does something similar to the Vertex Animation Tool's but with one key difference, the information for the vertices positions is stored in the meshes UV's instead of a 2D Texture.
+
+	![](VAT_Tool_Breakdown_02.png)
+
+	|Property Name| Description|
+	|-------------|------------|
+	|**Paint Selection Sequence:** |This stores the information about the meshes Vertices inside of the meshes UVs instead of storing that data in a 2D Texture. |
+
+### 3Ds Max Version & Script Installation
+This tool has only be tested using **3Ds Max 2015**. 
+While it might be able to work with other 3Ds Max versions, none of them have been tested so use other version of 3Ds Max at your own risk.
+To install the script just drag it from it's location in **4.8\Engine\Extras\3dsMaxScripts** right into the 3Ds Max viewport and the script will launch it's self.
+
+![](VAT_Max_Install.gif)(convert:false)
+
+[region:tip]
+If you find that you are using this script a lot you can always add the script to one of the tool bars or quad menus. If you are unfamiliar how to do this the [Autodesk site has a very detailed walk through](http://knowledge.autodesk.com/support/3ds-max/learn-explore/caas/CloudHelp/cloudhelp/2015/ENU/3DSMax/files/GUID-A2CF8BAA-7B52-40A8-8C40-803B1AB5FC05-htm.html) that will explain how you go about doing this.
+[/region]
+
+## 3Ds Max Unit Setup
+Before we can begin to use the tool we need to first ensure that the units of measurement 3Ds Max uses is setup to correctly match the units of measurement UE4 uses.
+This way we can ensure that the data the tool exports from 3Ds Max will work in the same manner inside UE4.
+Since UE4 uses Centimeters for its default unit of measurement we will need to ensure that 3Ds Max uses this as well and to change this we must do the following.
+
+1. First open up 3Ds Max 2015 and once loaded from the Main Toolbar select **Customize** > **Unit Setup**.
+1. Next click on **System Unit Setup** and change the **System Unit Setup** from Inches to **Centimeters** and then press the **OK** button.
+1. Finally change the **Display Unit Scale** to **Generic Units** and then press the **OK** button.
+
+	![](VAT_3DsMax_Unit_Setup.gif)(convert:false)
+
+[region:note]
+It is extremely important that you **DO NOT** skip this step. Skipping this step could result in rendering errors when the content is imported into UE4 because UE4 and 3Ds Max's units are not the same.
+[/region]
+
+## Tool Choices
+The Vertex Animation 3Ds Max script offers two different way to store the vertex animation data. 
+One method will store the vertices positions in a 2D Texture while the other method will store the vertex position data in the UV's of the mesh. 
+Below you will find links that show you how to setup and use both methods.
+
+[**Key framed Method**](Engine\Animation\Tools\VertexAnimationTool\VAT_KF_Meshes) - This method uses individual key frames that can be exported from other 3D packages and imported into 3Ds Max. That information is saved to the meshes UVs.    
+
+[**Animation Timeline Method**](Engine\Animation\Tools\VertexAnimationTool\VAT_TL_Meshes) - This method uses the 3Ds Max animation timeline and the results are encoded into a 2D Texture.
+
+ 
+## Tips & Tricks
+To get the most out this technique, there are a few tips and tricks that you can follow which are outlined below. 
+
+### Speeding Up the Animation Play Rate
+If you feel that the rate at which the animations are playing is to slow, you can use the **TimeWithSpeedVariable** Material Function to speed up the play rate.
+To do this just plugging the output of the **TimeWithSpeedVariable** into the **0-1 Animation** input if you are using the **MS_SequencePainter_SequenceFlibook** Material Function.
+If you are using the **MS_VertexAnimationTools_MorphTargets** Material Function plug the output of the **TimeWithSpeedVariable** into the **Morph Animations** input.
+
+![](VAT_SP_Speed_Up_Time.png)
+
+### Multiple Animated Meshes
+You can select multiple animated meshes at a time and the script will bake all of that data down to one mesh and one set of Textures. 
+This is very useful when working with characters that are made up of different parts. 
+Just select the parts that you want to use and the run the script like normal. 
+The script will then combine the pieces you selected and generate a new mesh from them along with the required 2D Textures.
+
+### Frame Skip
+Using the **Frame Skip** option that is under the **Vertex Animation Tools** section of the script will allow you to skip certain frames. 
+This is a very useful option as it allows you to reduce the final Texture size while still maintaining the look and feel of the original animation.
+
+![](VAT_Frame_Skip.png)
+
+If we take a look at the example video below we can see the Frame Skip feature in action. 
+The first Teapot you see in front of the word, Original, shows how the Teapot will animate using the full frame range. 
+The next Teapot, with the number 2 in front of it, shows how the teapot animates skipping every 2nd frame and so on.
+ As you can see the last example, the effect will still maintain its look and feel even if you skip 10 frames. 
+
+[OBJECT:EmbeddedVideo]
+	[PARAMLITERAL:width]
+	640
+	[/PARAMLITERAL]
+	[PARAMLITERAL:height]
+	360
+	[/PARAMLITERAL]
+	[PARAMLITERAL:videoid]
+	ZhOdykiJnEs
+	[/PARAMLITERAL]
+[/OBJECT]
+
+|Mesh Name /Number| Texture Size | Memory Savings|
+|-----------------|--------------|---------------|
+|Original        |175 KB          |N/A           |
+|2               |59 KB           |116 KB        |
+|5               |30 KB           |145 KB        |
+|10              |21 KB           |154 KB        |
+
+
+## Technical Information
+The following section is for technical information about how the Vertex Animation script works. 
+Please note that this is here purely for reference or for people who would like to know more about how the script works incase they want to modify it. 
+
+### Limitations
+Vertex positional morph target information is stored in a 16 bit signed floating point file format. 
+A 32 bit image would be more precise but 16 should suffice for most FX work. 
+With that said, the offset vertex positions will be less precise as they move further away from their resting positions. 
+
+### Memory Usage
+The memory usage for each frame per vert is as follows :
+
+* Vertex Offset Texture: 8 bytes per vertex per frame (each pixel)
+* Normal Texture: 4 bytes per vertex per frame (each pixel)
\ No newline at end of file
