#UE4doc #docs: duplicate word "the" changed "workaround" to "work around" in context of the sentence "a" to "an"

diff --git a/Engine/Documentation/Source/Engine/Rendering/Materials/ExpressionReference/Utility/UtilityExpressions.INT.udn b/Engine/Documentation/Source/Engine/Rendering/Materials/ExpressionReference/Utility/UtilityExpressions.INT.udn
--- a/Engine/Documentation/Source/Engine/Rendering/Materials/ExpressionReference/Utility/UtilityExpressions.INT.udn
+++ b/Engine/Documentation/Source/Engine/Rendering/Materials/ExpressionReference/Utility/UtilityExpressions.INT.udn
@@ -156,10 +156,10 @@ The Fresnel expression calculates a falloff based on the dot product of the surf
 | -------- | ----------- |
 | [REGION:tablesection]Properties[/REGION] ||
 | **Exponent** | Specifies the how quickly the output value falls off. Larger values mean tighter or quicker falloff. |
-| **Base Reflect Fraction** | Specifies the fraction of specular reflection the the surface is viewed from straight on. A value of 1 effectively disables the Fresnel effect. |
+| **Base Reflect Fraction** | Specifies the fraction of specular reflection the surface is viewed from straight on. A value of 1 effectively disables the Fresnel effect. |
 | [REGION:tablesection]Inputs[/REGION] ||
 | **ExponentIn** | Specifies the how quickly the output value falls off. Larger values mean tighter or quicker falloff. If used, value will always supersede the Exponent property value. |
-| **Base Reflect Fraction** | Specifies the fraction of specular reflection the the surface is viewed from straight on. A value of 1 effectively disables the Fresnel effect. If used, value will always supersede the Exponent property value. |
+| **Base Reflect Fraction** | Specifies the fraction of specular reflection the surface is viewed from straight on. A value of 1 effectively disables the Fresnel effect. If used, value will always supersede the Exponent property value. |
 | **Normal** | Takes in a three-channel vector value representing the normal of the surface, in world space. To see the results of a normal map applied to the surface of the Fresnel object, connect the normal map to the Normal input of the material, then connect a (PixelNormalWS)[Engine/Rendering/Materials/ExpressionReference/Vector#PixelNormalWS] expression to this input on the Fresnel. If no normal is specified, the tangent normal of the mesh is used. |
 
 
@@ -190,7 +190,7 @@ GIReplace allows artists to specify a different, usually simpler, expression cha
 
 ## LightmassReplace
 
-The LightmassReplace expression simply passes through the Realtime input when compiling the material for normal rendering purposes, and passes through the Lightmass input when exporting the material to Lightmass for global illumination.  This is useful to workaround material expressions that the exported version cannot handle correctly, for example WorldPosition.
+The LightmassReplace expression simply passes through the Realtime input when compiling the material for normal rendering purposes, and passes through the Lightmass input when exporting the material to Lightmass for global illumination.  This is useful to work around material expressions that the exported version cannot handle correctly, for example WorldPosition.
 
 | Item | Description |
 | -------- | ----------- |
@@ -302,7 +302,7 @@ In this example, the preview object will smoothly fade to black as the camera ex
 
 ## AntialiasedTextureMask
 
-The AntialiasedTextureMask expression allows you to create a material using a soft (anti-aliased) transition mask. The mask can be used to blend between two complex material properties or to fade out a alpha blended material (works well with SoftMasked). Simple specify a texture with the mask specified in one channel (red, green, blue, or alpha), set the used channel in the expression and specify the comparison value. Assuming the channel stores a grayscale value in the range 0=black to 1=white the comparison function defines if the resulting mask should be 0 or 1. This expression is a parameter, allowing the **Texture** property to be overridden by child MaterialInstances.
+The AntialiasedTextureMask expression allows you to create a material using a soft (anti-aliased) transition mask. The mask can be used to blend between two complex material properties or to fade out an alpha blended material (works well with SoftMasked). Simple specify a texture with the mask specified in one channel (red, green, blue, or alpha), set the used channel in the expression and specify the comparison value. Assuming the channel stores a grayscale value in the range 0=black to 1=white the comparison function defines if the resulting mask should be 0 or 1. This expression is a parameter, allowing the **Texture** property to be overridden by child MaterialInstances.
 
 | Item | Description |
 | -------- | ----------- |
