#UE4 Docs: Fixing an image link to point to the correct imgae.

diff --git a/Engine/Documentation/Source/Engine/Animation/Tools/VertexAnimationTool/VAT_TL_Meshes/VAT_TL_Meshes.INT.udn b/Engine/Documentation/Source/Engine/Animation/Tools/VertexAnimationTool/VAT_TL_Meshes/VAT_TL_Meshes.INT.udn
--- a/Engine/Documentation/Source/Engine/Animation/Tools/VertexAnimationTool/VAT_TL_Meshes/VAT_TL_Meshes.INT.udn
+++ b/Engine/Documentation/Source/Engine/Animation/Tools/VertexAnimationTool/VAT_TL_Meshes/VAT_TL_Meshes.INT.udn
@@ -99,7 +99,7 @@ If you are unfamiliar with how to import Static Meshes into UE4 please check out
 * Enable the **Combine Meshes** option.
 * Disable both **Import Materials** & **Import Textures** option. 
 	
-	![](VAT_SM_Setup_Options.png)
+	![](Engine\Animation\Tools\VertexAnimationTool\VAT_KF_Meshes\VAT_SM_Import_Options.png)
 
 Once all those options have been set, press the Import button to begin the import process. 
 After the Static Mesh has been imported, open the Static Mesh up and make sure the following options are enabled / disabled.
