#UE4doc #docs: Multiple edits

diff --git a/Engine/Documentation/Source/Engine/Animation/Tools/VertexAnimationTool/VertexAnimationTools.INT.udn b/Engine/Documentation/Source/Engine/Animation/Tools/VertexAnimationTool/VertexAnimationTools.INT.udn
--- a/Engine/Documentation/Source/Engine/Animation/Tools/VertexAnimationTool/VertexAnimationTools.INT.udn
+++ b/Engine/Documentation/Source/Engine/Animation/Tools/VertexAnimationTool/VertexAnimationTools.INT.udn
@@ -9,12 +9,12 @@ Related: Engine\Animation\Tools\VertexAnimationTool\VAT_TL_Meshes
 
 ## Overview
 Storing complex animation data into 2D Textures or a meshes UV's is a great way to reduce animation overhead while still maintaining the required look and feel of the animation. 
-In the past you might have used [Morph Targets](https://docs.unrealengine.com/latest/INT/Engine/Content/FBX/MorphTargets/index.html) to do the exact same thing however this new method has a few advantages that Morph Targets do not.
+In the past you might have used [Morph Targets](https://docs.unrealengine.com/latest/INT/Engine/Content/FBX/MorphTargets/index.html) to do the exact same thing, however this new method has a few advantages that Morph Targets do not.
 One of those advantages is allowing you to use complex animation data inside systems that would not otherwise be capable of utilizing Morph Targets like the Cascade Particle Editor.
 In the following documents we will cover how you can use this script to generate content for your UE4 projects.
 
 ## Prerequisites
-Before we can get to using the script there are few things that you must be made aware of so you can get the best results possible when using the technique. 
+Before we can begin using the script there are few things to be aware of so you can get the best results possible when using this technique. 
 
 ### Script Location
 [region:warning]
@@ -35,13 +35,13 @@ The tool generates the data in the Texture using the following formula.
 
 Due to this limitation this tool is best suited for Visual Effects or ambient Static Meshes that require animations but might not justify using a complex animation rig to do so. 
 This tool will also **NOT WORK** with Skeletal Mesh Animations as the bone transforms are not available to the Material Editor.
-This means that if you want to effect the vertices of a Skeletal Mesh in a similar manner you will need to use [Morph Targets](https://docs.unrealengine.com/latest/INT/Engine/Content/FBX/MorphTargets/index.html) instead.
+This means that if you want to affect the vertices of a Skeletal Mesh in a similar manner you will need to use [Morph Targets](https://docs.unrealengine.com/latest/INT/Engine/Content/FBX/MorphTargets/index.html) instead.
 
 ### Vertex Animation Tools Break Down
-Inside of the Vertex Animation Tools you will find two entirely different methods for effecting the vertices of Static Meshes. 
+Inside of the Vertex Animation Tools you will find two entirely different methods for affecting the vertices of Static Meshes. 
 In the following section we will go over what these two methods are and how they differ from one another.
 
-* **Vertex Animation Tools:** The top section of the Vertex Animation Tools, called Vertex Animation Tools, is for generating a 2D Texture that will store the morph target vertex positions and normal's.
+* **Vertex Animation Tools:** The top section of the Vertex Animation Tools, called Vertex Animation Tools, is for generating a 2D Texture that will store the morph target vertex positions and normals.
 
 	![](VAT_Tool_Breakdown_01.png)
 
@@ -52,7 +52,7 @@ In the following section we will go over what these two methods are and how they
 	|**Anim Start:** |This option defines the frame that the animation should start from. |
 	|**Anim Start:** |This option defines the frame that the animation should end on. |
 	|**Frame Skip:** |This option allows you skip frames to try and save on Texture space. |
-	|**Process Selected Meshes:** |This option, only available if Key Frame Meshes is enabled, will process the selected Key Frame Meshes in the order that they where selected in. |
+	|**Process Selected Meshes:** |This option, only available if Key Frame Meshes is enabled, will process the selected Key Frame Meshes in the order that they were selected. |
 
 
 * **Sequence Painter:** The Sequence Painter does something similar to the Vertex Animation Tool's but with one key difference, the information for the vertices positions is stored in the meshes UV's instead of a 2D Texture.
@@ -66,7 +66,7 @@ In the following section we will go over what these two methods are and how they
 ### 3Ds Max Version & Script Installation
 This tool has only be tested using **3Ds Max 2015**. 
 While it might be able to work with other 3Ds Max versions, none of them have been tested so use other version of 3Ds Max at your own risk.
-To install the script just drag it from it's location in **4.8\Engine\Extras\3dsMaxScripts** right into the 3Ds Max viewport and the script will launch it's self.
+To install the script just drag it from its location in **4.8\Engine\Extras\3dsMaxScripts** right into the 3Ds Max viewport and the script will launch itself.
 
 ![](VAT_Max_Install.gif)(convert:false)
 
@@ -90,7 +90,7 @@ It is extremely important that you **DO NOT** skip this step. Skipping this step
 [/region]
 
 ## Tool Choices
-The Vertex Animation 3Ds Max script offers two different way to store the vertex animation data. 
+The Vertex Animation 3Ds Max script offers two different ways to store the vertex animation data. 
 One method will store the vertices positions in a 2D Texture while the other method will store the vertex position data in the UV's of the mesh. 
 Below you will find links that show you how to setup and use both methods.
 
@@ -121,7 +121,7 @@ This is a very useful option as it allows you to reduce the final Texture size w
 
 ![](VAT_Frame_Skip.png)
 
-If we take a look at the example video below we can see the Frame Skip feature in action. 
+If we take a look at the example video below, we can see the Frame Skip feature in action. 
 The first Teapot you see in front of the word, Original, shows how the Teapot will animate using the full frame range. 
 The next Teapot, with the number 2 in front of it, shows how the teapot animates skipping every 2nd frame and so on.
  As you can see the last example, the effect will still maintain its look and feel even if you skip 10 frames. 
@@ -148,7 +148,7 @@ The next Teapot, with the number 2 in front of it, shows how the teapot animates
 
 ## Technical Information
 The following section is for technical information about how the Vertex Animation script works. 
-Please note that this is here purely for reference or for people who would like to know more about how the script works incase they want to modify it. 
+Please note that this is here purely for reference or for people who would like to know more about how the script works in case they want to modify it. 
 
 ### Limitations
 Vertex positional morph target information is stored in a 16 bit signed floating point file format. 
