#UE4doc #docs: minor edit

diff --git a/Engine/Documentation/Source/Engine/Animation/Tools/VertexAnimationTool/VAT_KF_Meshes/VAT_KF_Meshes.INT.udn b/Engine/Documentation/Source/Engine/Animation/Tools/VertexAnimationTool/VAT_KF_Meshes/VAT_KF_Meshes.INT.udn
--- a/Engine/Documentation/Source/Engine/Animation/Tools/VertexAnimationTool/VAT_KF_Meshes/VAT_KF_Meshes.INT.udn
+++ b/Engine/Documentation/Source/Engine/Animation/Tools/VertexAnimationTool/VAT_KF_Meshes/VAT_KF_Meshes.INT.udn
@@ -14,7 +14,7 @@ In this section we will go over how to use this method as well as any difference
 ## 3Ds Max Content Creation & Setup for Key Framed Meshes
 
 [region:note]
-Please note that while we are using 3Ds Max for this example the method in which the content is created and setup inside 3Ds Max is no different than how content could be set up inside other 3D Packages such as Maya or Blender.
+Please note that while we are using 3Ds Max for this example, the method in which the content is created and setup inside 3Ds Max is no different than how content could be set up inside other 3D Packages such as Maya or Blender.
 [/region]
 
 1. First inside 3s Max create a new scene and then create a Teapot with the following settings and place it at 0,0,0 in the world.
@@ -93,7 +93,7 @@ Set the value of one of the Constant Parameter node and plug that into the **Bas
 	![](VAT_KFM_Set_Values.png)
 
 1. Next use the mouse and Right - Click in the Material Graph to expose the search menu and then search for **MS_SequencePainter_SequenceFlipbook** Material Function using MS as the search term.
-When you find the MS_SequencePainter_SequenceFlipbook Material Function click on it with the Left Mouse button to add it to the Material Graph.
+When you find the MS_SequencePainter_SequenceFlipbook Material Function, click on it with the Left Mouse button to add it to the Material Graph.
 
 	![](VAT_KFM_Adding_MF_Flipbook.gif)(convert:false)
 
