#UE4doc #docs: changed "Ainm End" to "Anim End" "enable" to "enabled" "Mehs" to "Mesh" "it's" to "its"

diff --git a/Engine/Documentation/Source/Engine/Animation/Tools/VertexAnimationTool/VAT_TL_Meshes/VAT_TL_Meshes.INT.udn b/Engine/Documentation/Source/Engine/Animation/Tools/VertexAnimationTool/VAT_TL_Meshes/VAT_TL_Meshes.INT.udn
--- a/Engine/Documentation/Source/Engine/Animation/Tools/VertexAnimationTool/VAT_TL_Meshes/VAT_TL_Meshes.INT.udn
+++ b/Engine/Documentation/Source/Engine/Animation/Tools/VertexAnimationTool/VAT_TL_Meshes/VAT_TL_Meshes.INT.udn
@@ -25,7 +25,7 @@ In the following section we will cover how you can use the 3Ds Max Animation Tim
 	
 	![](VAT_Max_Bend_Setup.gif)(convert:false)
 
-1. With the Teapot now bending, open up the Vertex Animation Tools, if not already open, and set the **Animation Options** to **Animated Meshes** and the **Ainm End** to **35**.
+1. With the Teapot now bending, open up the Vertex Animation Tools, if not already open, and set the **Animation Options** to **Animated Meshes** and the **Anim End** to **35**.
 	
 	![](VAT_Tool_Settings_Setup.png)
 
@@ -37,8 +37,8 @@ a **.EXR** Texture which contains the vertices position data and a **.BMP** Text
 
 	![](VAT_Output_Textures.png)	
 
-1. With the Textures now created go back to 3Ds Max and select the Teapot mesh with, **_MorphExport**, added to the name and export it as a .FBX file to the same folder that the Textures where exported to.
-Name the mesh **SM_Teapot_Bend_00** and in the FBX options make sure that Smoothing Groups are enable and leave all other setting with the default values.
+1. With the Textures now created go back to 3Ds Max and select the Teapot mesh with "**_MorphExport**" added to the name and export it as a .FBX file to the same folder that the Textures where exported to.
+Name the mesh "**SM_Teapot_Bend_00**" and in the FBX options make sure that Smoothing Groups are enabled and leave all other setting with the default values.
 
 	![](VAT_3DsMax_Teapot_Export.gif)(convert:false)
 
@@ -54,7 +54,7 @@ If you are unfamiliar with the UE4 Texture Import process please read the [Textu
 [/region]
 
 The first thing that needs to be completed is that both Textures that were exported from 3Ds Max need to be imported into UE4. 
-To do select each Texture from it's export location in Windows Explorer and drag them into the Content Browser, releasing the left mouse button to trigger the import process.
+To do select each Texture from its export location in Windows Explorer and drag them into the Content Browser, releasing the left mouse button to trigger the import process.
 When you do this, you will be prompted with a warning message about the Textures not being a power of 2 resulting in the Texture not working with Streaming or Mip Mapping.
 Press the **Yes** button to dismiss the prompt and continue with the import process.
 
@@ -109,7 +109,7 @@ After the Static Mesh has been imported, open the Static Mesh up and make sure t
 
 	![](VAT_SM_Setup_Options.png)
 
-With the Static Mesh now set up correctly, make sure to press the **Apply Changes** button in the Static Mehs editor and then press the **Save All** button in the Content Browser to save all changes that have been made so far. 
+With the Static Mesh now set up correctly, make sure to press the **Apply Changes** button in the Static Mesh editor and then press the **Save All** button in the Content Browser to save all changes that have been made so far. 
 
 
 ### UE4 Material Setup & Use
@@ -150,7 +150,7 @@ If you are unfamiliar with how Materials work inside UE4 please refer to the [Ma
 
 	![](VAT_Wired_Up.png)
 
-1. Now add another Constant Parameter to the Material Graph by holding down the one key on the Keyboard and then clicking in the Material Graph with the Left Mouse Button to place the node in the Material Graph. Once the Constant Node has been added set it's value to 51 and connect it to the **Number of Morph Targets** input on the **MS_VertexAnimationTools_MorphTargets** Material Function.
+1. Now add another Constant Parameter to the Material Graph by holding down the one key on the Keyboard and then clicking in the Material Graph with the Left Mouse Button to place the node in the Material Graph. Once the Constant Node has been added set its value to 51 and connect it to the **Number of Morph Targets** input on the **MS_VertexAnimationTools_MorphTargets** Material Function.
 
 	![](VAT_Num_Morph_Targets.png)
 
