#UE4doc #docs: removed unnecessary commas

diff --git a/Engine/Documentation/Source/Engine/Rendering/Materials/Functions/Creating/CreatingFunctions.INT.udn b/Engine/Documentation/Source/Engine/Rendering/Materials/Functions/Creating/CreatingFunctions.INT.udn
--- a/Engine/Documentation/Source/Engine/Rendering/Materials/Functions/Creating/CreatingFunctions.INT.udn
+++ b/Engine/Documentation/Source/Engine/Rendering/Materials/Functions/Creating/CreatingFunctions.INT.udn
@@ -24,7 +24,7 @@ Your input node is now in place, and you are ready to create any network you lik
 
 ![SampleNetwork.png](SampleNetwork.png)(w:720)
 
-Functions define their interface to any materials they are used in by function inputs and outputs, which are just nodes that you place in the function.  Inputs and outputs can be named, so that people using the function in their materials know what they mean.  There must be at least one output for a function to be valid, and input and output names must be unique.
+Functions define their interface to any materials they are used in by function inputs and outputs, which are just nodes that you place in the function.  Inputs and outputs can be named so that people using the function in their materials know what they mean.  There must be at least one output for a function to be valid and input and output names must be unique.
 
 
 ## Publishing Your New Function
