#UE4doc #docs: sp: propagate

diff --git a/Engine/Documentation/Source/Engine/Rendering/Materials/Functions/Editing/EditingFunctions.INT.udn b/Engine/Documentation/Source/Engine/Rendering/Materials/Functions/Editing/EditingFunctions.INT.udn
--- a/Engine/Documentation/Source/Engine/Rendering/Materials/Functions/Editing/EditingFunctions.INT.udn
+++ b/Engine/Documentation/Source/Engine/Rendering/Materials/Functions/Editing/EditingFunctions.INT.udn
@@ -18,7 +18,7 @@ In either the **Content Browser** or the **Material Editor's** Graph tab, you ca
 
 However, it is important to note that any changes made and saved to a material function will be present in all instances of that material function moving forward. For example, if you made a change to the internal network of a Radial Gradient material function, all existing instances of that function would receive the update, as well as all new instances moving forward. 
 
-For this reason, unless you are _certain_ that your change needs to propogate throughout all other instances of the function, it may be wise to make a copy of an existing function in the **Content Browser** (**right-click** and choose _Make a Copy_ from the context menu) than editing the original material function.
+For this reason, unless you are _certain_ that your change needs to propagate throughout all other instances of the function, it may be wise to make a copy of an existing function in the **Content Browser** (**right-click** and choose _Make a Copy_ from the context menu) than editing the original material function.
 
 Once you make changes to a function you have to click the ![button_apply.png](button_apply.png) button to propagate the changes to the function asset and any materials using the function. Once done, be **sure to save your asset in the Content Browser!**
 
