#UE4doc #docs: capitalized "Photoshop," sp: "Unlabeled," removed unnecessary comma

diff --git a/Engine/Documentation/Source/Engine/Rendering/Materials/Functions/Overview/FunctionsOverview.INT.udn b/Engine/Documentation/Source/Engine/Rendering/Materials/Functions/Overview/FunctionsOverview.INT.udn
--- a/Engine/Documentation/Source/Engine/Rendering/Materials/Functions/Overview/FunctionsOverview.INT.udn
+++ b/Engine/Documentation/Source/Engine/Rendering/Materials/Functions/Overview/FunctionsOverview.INT.udn
@@ -17,7 +17,7 @@ Functions are edited in the material editor like a normal material, but with som
 
 Material functions are also assets that can appear within the **Content Browser**. Their graphs will be different from those of materials in that instead of a main material node, material functions will instead have output nodes that represent the output connections of the final function. 
 
-It may help to think of a material function like an electronics project enclosure. You may add as many inputs and outputs as you deem necessary. The heart of the function is what it does between those inputs and outputs.  This example takes two layers and blends them together like a [photoshop screen blend](http://www.webdesign.org/photoshop/photoshop-basics/using-blend-modes-screen.14859.html).  The function abstracts the details from an artist that might want to use it, so they do not have to actually know the math of a screen blend to use a screen blend operation.  If anyone wants to change the way the screen blend operation works later, they can do so to this function and the change will automatically be propagated to all materials that use it.
+It may help to think of a material function like an electronics project enclosure. You may add as many inputs and outputs as you deem necessary. The heart of the function is what it does between those inputs and outputs.  This example takes two layers and blends them together like a [Photoshop screen blend](http://www.webdesign.org/photoshop/photoshop-basics/using-blend-modes-screen.14859.html).  The function abstracts the details from an artist that might want to use it, so they do not have to actually know the math of a screen blend to use a screen blend operation.  If anyone wants to change the way the screen blend operation works later, they can do so to this function and the change will automatically be propagated to all materials that use it.
 
 ![ScreenBlendFunction.png](ScreenBlendFunction.png)(w:720)
 
@@ -117,7 +117,7 @@ FunctionOutput nodes have the following properties:
 | **Description** | This provides a description of this input that will be visible when a user hovers over the function's output pin. |
 | **Sort Priority** | This number is used to control the order in which input pins will be listed on the function's node. The order is lowest to highest. |
 |[REGION:tablesection]Input Pins[/REGION]||
-| **(Unlabled)** | This provides the input for the  data that has been processed by the function. This data will be sent out of the function for use in the material. |
+| **(Unlabeled)** | This provides the input for the  data that has been processed by the function. This data will be sent out of the function for use in the material. |
 
 
 ### Input Types
@@ -220,7 +220,7 @@ Corresponding tooltip when used in a material:
 
 When you edit a function and click apply changes, the new version is propagated to any loaded materials or functions which reference this function.  Any unloaded materials which reference the function will be updated with the changes when they are next loaded.
 
-**When an input or output is deleted from a function and the changes are propagated, any links to these deleted connectors in materials that use the function will be broken!**  It is important to be aware of this, because the propagation cannot be undone.  The more materials a function is used in, the bigger the potential for breakage, so be careful.
+**When an input or output is deleted from a function and the changes are propagated, any links to these deleted connectors in materials that use the function will be broken!**  It is important to be aware of this, because the propagation cannot be undone.  The more materials a function is used in, the bigger the potential for breakage so be careful.
 
 All loaded materials that use the function will be marked dirty when the function change is propagated, which can be used to see which packages could be resaved to prevent increased load times.  You can find all the loaded materials which use a function by **right clicking** in the **Content Browser** and choosing this option:
 
