#UE4 Docs
Upated images and minor content for Lighting Quick Start Guide.

diff --git a/Engine/Documentation/Source/Engine/Rendering/LightingAndShadows/QuickStart/5/LQS_5.INT.udn b/Engine/Documentation/Source/Engine/Rendering/LightingAndShadows/QuickStart/5/LQS_5.INT.udn
--- a/Engine/Documentation/Source/Engine/Rendering/LightingAndShadows/QuickStart/5/LQS_5.INT.udn
+++ b/Engine/Documentation/Source/Engine/Rendering/LightingAndShadows/QuickStart/5/LQS_5.INT.udn
@@ -62,7 +62,7 @@ First, let's add some Materials to our apartment instead of using the default on
 
 1.	Drag that same Material into the small bathroom. 
 
-1.	For the other floors, drag in the **M_Wood_Floor_Walnut** Material.
+1.	For the other floors, drag in the **M_Wood_Floor_Walnut_Polished** Material.
 
 1.	While we are at it, drag the **M_Metal_Rust** asset onto the pillars across the roof. 
 
