#UE4doc #docs: corrected singularity vs plural conflict in sentence

diff --git a/Engine/Documentation/Source/Engine/Rendering/Materials/Functions/Reference/Blends/Blends.INT.udn b/Engine/Documentation/Source/Engine/Rendering/Materials/Functions/Reference/Blends/Blends.INT.udn
--- a/Engine/Documentation/Source/Engine/Rendering/Materials/Functions/Reference/Blends/Blends.INT.udn
+++ b/Engine/Documentation/Source/Engine/Rendering/Materials/Functions/Reference/Blends/Blends.INT.udn
@@ -7,7 +7,7 @@ Description:Functions designed to blend one color with another, similar to blend
 
 ## Overview
 
-Blends are a type of function that perform mathematical functions on the color information of textures to allow one to blend into another in a particular manner. These are commonly seen in image editing software packages such as [Photoshop's Blend Modes](http://www.webdesign.org/photoshop/photo-editing/photoshop-basics-blending-modes.20793.html).
+A Blend is a type of function that performs mathematical calculations in the color information of a texture so that one texture can blend into another in a particular manner.
 
 Blends will always have a Base and a Blend input, both of which are Vector3. These will each take in a texture and the two will be blended in some way. The manner in which the blend takes place depends on the type of blend function used.
 
