#UE4doc #docs: sp: components

diff --git a/Engine/Documentation/Source/Engine/Rendering/Materials/Functions/Reference/Math/Math.INT.udn b/Engine/Documentation/Source/Engine/Rendering/Materials/Functions/Reference/Math/Math.INT.udn
--- a/Engine/Documentation/Source/Engine/Rendering/Materials/Functions/Reference/Math/Math.INT.udn
+++ b/Engine/Documentation/Source/Engine/Rendering/Materials/Functions/Reference/Math/Math.INT.udn
@@ -32,7 +32,7 @@ This function takes in a Vector2, a Vector3, or a Vector4, combines their compon
 ![Add Components](AddComponents_Demo.png)
 
 [REGION:note]
-	Since each input is calculated individually for its corresponding output, you may use all three inputs on a single node, so long as you also use an output for each one. For instance, you may use one AddComponents function node to combine the components of a Vector2, utilizing the *f2* output, and use the same node to combine the componets of a Vector3, utilizing the *f3* output.
+	Since each input is calculated individually for its corresponding output, you may use all three inputs on a single node, so long as you also use an output for each one. For instance, you may use one AddComponents function node to combine the components of a Vector2, utilizing the *f2* output, and use the same node to combine the components of a Vector3, utilizing the *f3* output.
 [/REGION]
 
 ### Pi
