#UE4 Docs
Jira - UEDOC 1354

Added a section to UMG Menu Anchor page along with an image.

diff --git a/Engine/Documentation/Source/Engine/UMG/UserGuide/WidgetTypeReference/MenuAnchor/MenuAnchor.INT.udn b/Engine/Documentation/Source/Engine/UMG/UserGuide/WidgetTypeReference/MenuAnchor/MenuAnchor.INT.udn
--- a/Engine/Documentation/Source/Engine/UMG/UserGuide/WidgetTypeReference/MenuAnchor/MenuAnchor.INT.udn
+++ b/Engine/Documentation/Source/Engine/UMG/UserGuide/WidgetTypeReference/MenuAnchor/MenuAnchor.INT.udn
@@ -24,9 +24,15 @@ In the **Details** panel for a placed **Menu Anchor**, there are a couple of spe
 | -------- | ----------- |
 | **Menu Class** | The Widget Blueprint (pop-up) to spawn when the menu is called. Creates the widget freshly each time. |
 | **Placement** | Several placement options use to determine the location of the created widget. |
-| **OnGetMenuContent** | Allows for the customization of the pop-up by binding a function or property to the **OnGetMenuContent** Event which is called when the pop-up is summoned. |
+| **OnGetMenuContent** | Allows for the customization of the pop-up by binding a function or property to the **OnGetMenuContent** Event which is called when the pop-up is summoned (see below). |
 
-The **Menu Anchor** widget has some Widget specific nodes that can be called through script which are described below. 
+For the **OnGetMenuContent** Event, you can create and bind a function directly from the **Details** panel which will be called whenever the Menu Anchor is called. This allows you to provide extra functionality to the creation of your pop-up menu. For example, below we've run a check to determine if the player is _In Battle_ and if not, we allow them to access the pop-up menu. If they are we do not allow them to access the pop-up menu. 
+
+![](ExampleGetMenu.png)(w:840)
+
+### Menu Anchor Blueprint Functions
+
+The **Menu Anchor** widget has some widget specific functions that can be called through script which are described below. 
 
 ![](MenuAnchorNodes.png)
 
