#UE4doc #docs: changed "viewmode" to "View Mode" to match editor, removed "a" from plural, sp "outer"

diff --git a/Engine/Documentation/Source/Engine/Rendering/LightingAndShadows/QuickStart/1/LQS_1.INT.udn b/Engine/Documentation/Source/Engine/Rendering/LightingAndShadows/QuickStart/1/LQS_1.INT.udn
--- a/Engine/Documentation/Source/Engine/Rendering/LightingAndShadows/QuickStart/1/LQS_1.INT.udn
+++ b/Engine/Documentation/Source/Engine/Rendering/LightingAndShadows/QuickStart/1/LQS_1.INT.udn
@@ -56,7 +56,7 @@ Inside your **Blank** project (with **Starter Content**), follow the steps below
 
 	![](floorandroof.png)
 
-	To view your assets better, switch to **Wireframe** viewmode from the View options. 
+	To view your assets better, switch to **Wireframe** View Mode from the View options. 
 
 	![](setup3.png)
 
@@ -88,7 +88,7 @@ Inside your **Blank** project (with **Starter Content**), follow the steps below
 	| **400x410** | ![](setup23.png) | 
 	| **400x411** | ![](setup24.png) |
 
-	With all the wall pieces added (highlighted below), you should now have a outter walls and a small corner room. 
+	With all the wall pieces added (highlighted below), you should now have outer walls and a small corner room. 
 
 	![](wallsAdded.png)
 
