#UE4doc #docs: changed "are" to "is" for singular subject

diff --git a/Engine/Documentation/Source/Engine/Rendering/LightingAndShadows/QuickStart/6/LQS_6.INT.udn b/Engine/Documentation/Source/Engine/Rendering/LightingAndShadows/QuickStart/6/LQS_6.INT.udn
--- a/Engine/Documentation/Source/Engine/Rendering/LightingAndShadows/QuickStart/6/LQS_6.INT.udn
+++ b/Engine/Documentation/Source/Engine/Rendering/LightingAndShadows/QuickStart/6/LQS_6.INT.udn
@@ -22,7 +22,7 @@ Description:In this step we'll create a slightly more advanced light which uses
 
 %Steps%
 
-With our apartment almost complete, next we will create a more advanced **Point Light** that uses an **IES Profile** which are a lighting industry standard method of diagramming the brightness and falloff of light as it exits a particular real world light fixture.  
+With our apartment almost complete, next we will create a more advanced **Point Light** that uses an **IES Profile** which is a lighting industry standard method of diagramming the brightness and falloff of light as it exits a particular real world light fixture.  
 
 [REGION:note]
 	For more information, refer to the [IES Light Profiles](Engine/Rendering/LightingAndShadows/IESLightProfiles) documentation.
@@ -42,7 +42,7 @@ With our apartment almost complete, next we will create a more advanced **Point
 
 	![](ShowEngineContent.png)
 
-	Unreal Engine 4 provides some example IES Light Profiles to use, you can find others from the internet and import them. 
+	Unreal Engine 4 provides some example IES Light Profiles to use, but you can find others from the internet and import them as well. 
 
 1.	Select the **Complex_IES** profile. 
 
