#UE4doc #docs: changed "FlipBook" to "Flipbook" for documentation consistency

diff --git a/Engine/Documentation/Source/Engine/Rendering/Materials/Functions/Reference/Misc/Misc.INT.udn b/Engine/Documentation/Source/Engine/Rendering/Materials/Functions/Reference/Misc/Misc.INT.udn
--- a/Engine/Documentation/Source/Engine/Rendering/Materials/Functions/Reference/Misc/Misc.INT.udn
+++ b/Engine/Documentation/Source/Engine/Rendering/Materials/Functions/Reference/Misc/Misc.INT.udn
@@ -63,7 +63,7 @@ The SplitComponents function busts apart the components of an incoming color or
 
 ![Split Components](SplitComponents_Demo.png)
 
-### FlipBook
+### Flipbook
 
 This function takes in a Texture2D, such as a grid of sprite frames, and outputs an animation of them playing back in order.
 
@@ -76,7 +76,7 @@ This function takes in a Texture2D, such as a grid of sprite frames, and outputs
 | **Texture (Texture2D)** | Takes in a Texture2D containing a sprite sheet, or grid of frames of an animated texture.|
 | **UVs (Vector2)** | Takes in a set of UV coordinates to assist in tiling.|
 |[REGION:tablesection]Outputs[/REGION] ||
-| **Result** | Outputs an image as the result of the FlipBook.|
+| **Result** | Outputs an image as the result of the flipbook.|
 | **UVs** | Outputs the UV coordinates corresponding with a given frame of the texture sheet.|
 
 
