#UE4doc #docs: changed "effect" to "affect" since used as a verb

diff --git a/Engine/Documentation/Source/Engine/Rendering/Materials/Functions/Reference/Texturing/Texturing.INT.udn b/Engine/Documentation/Source/Engine/Rendering/Materials/Functions/Reference/Texturing/Texturing.INT.udn
--- a/Engine/Documentation/Source/Engine/Rendering/Materials/Functions/Reference/Texturing/Texturing.INT.udn
+++ b/Engine/Documentation/Source/Engine/Rendering/Materials/Functions/Reference/Texturing/Texturing.INT.udn
@@ -227,8 +227,8 @@ Outputs falloff results for the 3 primary world vectors based on the camera angl
 | -------- | -------- |
 |[REGION:tablesection]Inputs[/REGION] ||
 | **Blend Power (Scalar)** | The rate at which the textures blend from one to the other. |
-| **Use Reflection Vector (StaticBool)** | If *true*, uses ReflectionVector instead of CameraVector (Material's Normal input will effect results). This defaults to *false*. However, at this time, as there is no false input data for the switch, leaving this set to _false_ will yield no result.|
-| **Use Smoothable Normals (StaticBool)** | If *true*, uses ReflectionVector instead of CameraVector (Material's Normal input will effect results). This defaults to *false*. However, at this time, as there is no false input data for the switch, leaving this set to _false_ will yield no result.|
+| **Use Reflection Vector (StaticBool)** | If *true*, uses ReflectionVector instead of CameraVector (Material's Normal input will affect results). This defaults to *false*. However, at this time, as there is no false input data for the switch, leaving this set to _false_ will yield no result.|
+| **Use Smoothable Normals (StaticBool)** | If *true*, uses ReflectionVector instead of CameraVector (Material's Normal input will affect results). This defaults to *false*. However, at this time, as there is no false input data for the switch, leaving this set to _false_ will yield no result.|
 | **Smooth Reflection Percentage (Scalar)** | How much to smooth the normals being applied to the surface prior to blend.|
 |[REGION:tablesection]Outputs[/REGION] ||
 | **XY True** | The falloff results which blend toward white as the surface aligns to the XY plane.|
