#UE4doc #docs: removed spaces before and after hyphens in "How-To's", duplicate word "use"

diff --git a/Engine/Documentation/Source/Engine/Rendering/Materials/HowTo/MaterialHowTo.INT.udn b/Engine/Documentation/Source/Engine/Rendering/Materials/HowTo/MaterialHowTo.INT.udn
--- a/Engine/Documentation/Source/Engine/Rendering/Materials/HowTo/MaterialHowTo.INT.udn
+++ b/Engine/Documentation/Source/Engine/Rendering/Materials/HowTo/MaterialHowTo.INT.udn
@@ -57,30 +57,30 @@ If you already have a project that also includes the Starter Content, please fee
 1. With each of the above defined, click **Create Project** to create a new project. 
 
 
-## Material How - To's 
+## Material How-To's 
 
-Listed below are several How - To guides covering a wide range of topics regarding Materials. 
-Read over each description to get an understanding of what the How - To will cover then click to open and begin that How - To.
+Listed below are several How-To guides covering a wide range of topics regarding Materials. 
+Read over each description to get an understanding of what the How-To will cover then click to open and begin that How-To.
 
 
 |How - To... | Description |
 | ------- | ------------ |
-| **[Use the Main Material Node](Engine\Rendering\Materials\HowTo\Main_Material_Node)** | Setting up the Main Material node is key to get any Material to work. This How To will show you all you need to know about setting up the Main Material Node. |
-| **[Preview and Apply your Materials](Engine\Rendering\Materials\HowTo\Previewing_Applying)** | The ability to preview and apply your Materials is important to making any 3D project in UE4 look good. In this How To, you will learn not only how to preview your Materials but also how to apply those Materials to objects in the level. |
-| **[Use Instancing in your Materials](Engine\Rendering\Materials\HowTo\Instancing)** | Material Instancing is a powerful concept that can greatly help speed up any workflow. In this How To tutorial, you will learn how to setup your Materials so that they can use Material Instancing. |
-| **[Place Material Expressions and Functions](Engine\Rendering\Materials\HowTo\Placing_Expressions_Functions)** |Placing Material Expression nodes is a fundamental  part of building shaders. This How To will cover all you need to know about how to place Expression in the Material graph. |
-| **[Make Material Functions for your Materials](Engine\Rendering\Materials\HowTo\Making_Functions)** | Material Functions are snippets of Material code that can be recycled. In this How To, you will learn how you can use them in your Materials. |
-| **[Make Material Parameters](Engine\Rendering\Materials\HowTo\Making_Parameters)** | Material Parameters allow you to interact with Material Instances. In this How To, you will learn how you can setup and use use Material Parameters in your Materials.|
-| **[Use the Bump Offset Node](Engine\Rendering\Materials\HowTo\BumpOffset)** | The Bump Offset Material node can help add extra depth information to your Materials. This How To guide will show you everything you need to know about using the Bump Offset Material node in your Materials.    .|
-| **[Set Up Human Skin](Engine\Rendering\Materials\HowTo\Human_Skin)** | Setting up realistic Human skin can be tricky to get right. In this How To tutorial, we will go over how you can set up Human Skin to get the best results possible.|
+| **[Use the Main Material Node](Engine\Rendering\Materials\HowTo\Main_Material_Node)** | Setting up the Main Material node is key to get any Material to work. This How-To will show you all you need to know about setting up the Main Material Node. |
+| **[Preview and Apply your Materials](Engine\Rendering\Materials\HowTo\Previewing_Applying)** | The ability to preview and apply your Materials is important to making any 3D project in UE4 look good. In this How-To, you will learn not only how to preview your Materials but also how to apply those Materials to objects in the level. |
+| **[Use Instancing in your Materials](Engine\Rendering\Materials\HowTo\Instancing)** | Material Instancing is a powerful concept that can greatly help speed up any workflow. In this How-To tutorial, you will learn how to setup your Materials so that they can use Material Instancing. |
+| **[Place Material Expressions and Functions](Engine\Rendering\Materials\HowTo\Placing_Expressions_Functions)** |Placing Material Expression nodes is a fundamental  part of building shaders. This How-To will cover all you need to know about how to place Expression in the Material graph. |
+| **[Make Material Functions for your Materials](Engine\Rendering\Materials\HowTo\Making_Functions)** | Material Functions are snippets of Material code that can be recycled. In this How-To, you will learn how you can use them in your Materials. |
+| **[Make Material Parameters](Engine\Rendering\Materials\HowTo\Making_Parameters)** | Material Parameters allow you to interact with Material Instances. In this How To, you will learn how you can setup and use Material Parameters in your Materials.|
+| **[Use the Bump Offset Node](Engine\Rendering\Materials\HowTo\BumpOffset)** | The Bump Offset Material node can help add extra depth information to your Materials. This How-To guide will show you everything you need to know about using the Bump Offset Material node in your Materials.    .|
+| **[Set Up Human Skin](Engine\Rendering\Materials\HowTo\Human_Skin)** | Setting up realistic Human skin can be tricky to get right. In this How-To tutorial, we will go over how you can set up Human Skin to get the best results possible.|
 | **[Set  Up Emissive Materials](Engine\Rendering\Materials\HowTo\EmissiveGlow)** | Creating Materials that glow no matter what the lighting conditions is something that any 3D project needs. In this tutorial, we will go over how you can setup glowing Materials using the Emissive Material input.|
-| **[Use Masking in your Materials](Engine\Rendering\Materials\HowTo\Masking)** |Texture Masking is a powerful tool that will allow you to add endless variety to your objects for very little cost. This How To covers all you need to know about using Texture Masks in your Materials. |
-| **[Use Fresnel in your Materials](Engine\Rendering\Materials\HowTo\Fresnel)** |Sometimes characters need a boost to their silhouette to help them stand out in a scene. This is often referred to Rim lighting, but in UE4 this is known as Fresnel. This How To will show you all you need to know about how to add Fresnel to your Materials. |
-| **[Animate UV Coordinates](Engine\Rendering\Materials\HowTo\AnimatingUVCoords)** |Adding animation to UV coordinates is a cheap and very effective way to add the illusion of motion. In this How To, you will be shown how you can set up UV animation in your Materials. |
-| **[Add Detail Texturing](Engine\Rendering\Materials\HowTo\DetailTexturing)** |Adding extra detail to a surface so that it will look nice when viewed very closely can be achieved through Detail mapping. In this How To, you will be shown how you can add Detail Mapping to your Materials so they will always look good no matter what distance they are viewed from.|
-| **[Enable and use Refraction](Engine\Rendering\Materials\HowTo\Refraction)** |When light hits an object like glass, some of the light is changed because of refraction. In this How To, we will cover how you can use Refraction in your Materials.|
-| **[Enable and use Transparency](Engine\Rendering\Materials\HowTo\Transparency)** |The ability to simulate Transparency is key to being able to make believable water or glass. In this How To, we will cover how you can utilize Transparency in your Materials.|
-| **[Use Layered Materials](Engine/Rendering/Materials/HowTo/CreatingLayeredMaterials)** | %Engine/Rendering/Materials/HowTo/CreatingLayeredMaterials:description% |
-| **[Make Materials Shiny](Engine/Rendering/Materials/HowTo/ShinyMaterials)** | Making Materials that are shiny is something can be accomplished in just a few clicks of the mouse. In this How To, we will cover how you can turn any Material into a Material that shines.|
-| **[Colored Translucent Shadows](Engine/Rendering/Materials/HowTo/ColoredTransluscentShadows)** |Getting shadows to look like they are casting colored shadows into the world can be achieved by using Colored Translucent Shadows. In this How To, we will cover how you can create Materials that make use of Colored Translucent Shadows. |
+| **[Use Masking in your Materials](Engine\Rendering\Materials\HowTo\Masking)** |Texture Masking is a powerful tool that will allow you to add endless variety to your objects for very little cost. This How-To covers all you need to know about using Texture Masks in your Materials. |
+| **[Use Fresnel in your Materials](Engine\Rendering\Materials\HowTo\Fresnel)** |Sometimes characters need a boost to their silhouette to help them stand out in a scene. This is often referred to Rim lighting, but in UE4 this is known as Fresnel. This How-To will show you all you need to know about how to add Fresnel to your Materials. |
+| **[Animate UV Coordinates](Engine\Rendering\Materials\HowTo\AnimatingUVCoords)** |Adding animation to UV coordinates is a cheap and very effective way to add the illusion of motion. In this How-To, you will be shown how you can set up UV animation in your Materials. |
+| **[Add Detail Texturing](Engine\Rendering\Materials\HowTo\DetailTexturing)** |Adding extra detail to a surface so that it will look nice when viewed very closely can be achieved through Detail mapping. In this How-To, you will be shown how you can add Detail Mapping to your Materials so they will always look good no matter what distance they are viewed from.|
+| **[Enable and use Refraction](Engine\Rendering\Materials\HowTo\Refraction)** |When light hits an object like glass, some of the light is changed because of refraction. In this How-To, we will cover how you can use Refraction in your Materials.|
+| **[Enable and use Transparency](Engine\Rendering\Materials\HowTo\Transparency)** |The ability to simulate Transparency is key to being able to make believable water or glass. In this How-To, we will cover how you can utilize Transparency in your Materials.|
+| **[Use Layered Materials](Engine/Rendering/Materials/HowTo/CreatingLayeredMaterials)** | %Engine/Rendering/Materials/How-To/CreatingLayeredMaterials:description% |
+| **[Make Materials Shiny](Engine/Rendering/Materials/HowTo/ShinyMaterials)** | Making Materials that are shiny is something can be accomplished in just a few clicks of the mouse. In this How-To, we will cover how you can turn any Material into a Material that shines.|
+| **[Colored Translucent Shadows](Engine/Rendering/Materials/HowTo/ColoredTransluscentShadows)** |Getting shadows to look like they are casting colored shadows into the world can be achieved by using Colored Translucent Shadows. In this How-To, we will cover how you can create Materials that make use of Colored Translucent Shadows. |
 
