TypoCorrection

diff --git a/Engine/Documentation/Source/Engine/Animation/Tools/VertexAnimationTool/VertexAnimationTools.INT.udn b/Engine/Documentation/Source/Engine/Animation/Tools/VertexAnimationTool/VertexAnimationTools.INT.udn
--- a/Engine/Documentation/Source/Engine/Animation/Tools/VertexAnimationTool/VertexAnimationTools.INT.udn
+++ b/Engine/Documentation/Source/Engine/Animation/Tools/VertexAnimationTool/VertexAnimationTools.INT.udn
@@ -1,4 +1,4 @@
-Availability:Public
+﻿Availability:Public
 Title:Vertex Animation Tool
 Crumbs: %ROOT%, Engine, Engine/Animation
 Description: User guide for the 3D Max Vertex Animation tool set.
@@ -9,7 +9,7 @@ Related: Engine\Animation\Tools\VertexAnimationTool\VAT_TL_Meshes
 
 ## Overview
 Storing complex animation data into 2D Textures or a meshes UV's is a great way to reduce animation overhead while still maintaining the required look and feel of the animation. 
-In the past you might have used [Morph Targets](https://docs.unrealengine.com/latest/INT/Engine/Content/FBX/MorphTargets/index.html) to do the exact same thing, however this new method has a few advantages that Morph Targets do not.
+In the past you might have used [Morph Targets](Engine/Content/FBX/MorphTargets) to do the exact same thing, however this new method has a few advantages that Morph Targets do not.
 One of those advantages is allowing you to use complex animation data inside systems that would not otherwise be capable of utilizing Morph Targets like the Cascade Particle Editor.
 In the following documents we will cover how you can use this script to generate content for your UE4 projects.
 
@@ -35,7 +35,7 @@ The tool generates the data in the Texture using the following formula.
 
 Due to this limitation this tool is best suited for Visual Effects or ambient Static Meshes that require animations but might not justify using a complex animation rig to do so. 
 This tool will also **NOT WORK** with Skeletal Mesh Animations as the bone transforms are not available to the Material Editor.
-This means that if you want to affect the vertices of a Skeletal Mesh in a similar manner you will need to use [Morph Targets](https://docs.unrealengine.com/latest/INT/Engine/Content/FBX/MorphTargets/index.html) instead.
+This means that if you want to affect the vertices of a Skeletal Mesh in a similar manner you will need to use [Morph Targets](Engine/Content/FBX/MorphTargets) instead.
 
 ### Vertex Animation Tools Break Down
 Inside of the Vertex Animation Tools you will find two entirely different methods for affecting the vertices of Static Meshes. 
@@ -50,7 +50,8 @@ In the following section we will go over what these two methods are and how they
 	|**Animation Options:** | Animation Options allows you to pick between using animation that was created using the Timeline in 3Ds Max or use individual Key Frame Meshes that were created in 3Ds Max or another 3D package like Maya or Blender and then exported out of that package frame by frame so the animation could then be re-constructed inside 3Ds Max.|
 	|**Process Animated Meshes:** | This button will process the animated meshes that are in the 3Ds Max scene, creating and then exporting the required Textures. |
 	|**Anim Start:** |This option defines the frame that the animation should start from. |
-	|**Anim Start:** |This option defines the frame that the animation should end on. |
+	|**Anim End
+:** |This option defines the frame that the animation should end on. |
 	|**Frame Skip:** |This option allows you skip frames to try and save on Texture space. |
 	|**Process Selected Meshes:** |This option, only available if Key Frame Meshes is enabled, will process the selected Key Frame Meshes in the order that they were selected. |
 
