TypoCorrection

diff --git a/Engine/Documentation/Source/Engine/Animation/Tools/VertexAnimationTool/VAT_KF_Meshes/VAT_KF_Meshes.INT.udn b/Engine/Documentation/Source/Engine/Animation/Tools/VertexAnimationTool/VAT_KF_Meshes/VAT_KF_Meshes.INT.udn
--- a/Engine/Documentation/Source/Engine/Animation/Tools/VertexAnimationTool/VAT_KF_Meshes/VAT_KF_Meshes.INT.udn
+++ b/Engine/Documentation/Source/Engine/Animation/Tools/VertexAnimationTool/VAT_KF_Meshes/VAT_KF_Meshes.INT.udn
@@ -1,7 +1,7 @@
-Availability:Public
+﻿Availability:Public
 Title:Vertex Animation Tool - Key Framed Meshes
 Crumbs: %ROOT%, Engine, Engine/Animation
-Description: User guide for the 3D Max Vertex Animation tool set.
+Description: User guide for the 3Ds Max Vertex Animation tool set.
 
 [TOC(start:2 end:3)]
 
@@ -17,7 +17,7 @@ In this section we will go over how to use this method as well as any difference
 Please note that while we are using 3Ds Max for this example, the method in which the content is created and setup inside 3Ds Max is no different than how content could be set up inside other 3D Packages such as Maya or Blender.
 [/region]
 
-1. First inside 3s Max create a new scene and then create a Teapot with the following settings and place it at 0,0,0 in the world.
+1. First inside 3Ds Max create a new scene and then create a Teapot with the following settings and place it at 0,0,0 in the world.
 	* **Radius:** 70.0
 	* **Segments:** 5
 
@@ -33,7 +33,7 @@ Make sure the **Object** is set to **Copy** and then set the **Number of Copies*
 	![](VAT_KFM_Teapot_Copy.gif)(convert:false)
 
 1. Then select the 2nd Teapot mesh and set the **Melt Amount** in the Melt modifier to 25. 
-Then select the next Teapot in line and set the **Melt Amount** of the Melt modifier to 50 and the next one to 75% and finally set the last Teapot **Melt Amount** to 100.
+Then select the next Teapot in line and set the **Melt Amount** of the Melt modifier to 50 and the next one to 75, and finally set the last Teapot **Melt Amount** to 100.
 
 	![](VAT_KFM_Teapot_Melt_Amount.gif)(convert:false)
 
@@ -68,7 +68,7 @@ To do this all you need to do is press the **Paint Selection Sequence** button t
 ## UE4 Content Setup & Use for Key Framed Meshes
 Once the Teapot mesh has been exported from 3Ds Max, it is now time for the Teapot mesh to be imported and setup so that it will work the same inside UE4.
 
-1. First import the Teapot mesh, **SM_Teapot_PSS.FBX** by dragging it form the location it was exported to into the **Content Browser**, making sure that the following options are enabled / disabled.
+1. First import the Teapot mesh, **SM_Teapot_PSS.FBX** by dragging it from the location it was exported to into the **Content Browser**, making sure that the following options are enabled / disabled.
 
 	* Disable the **Auto Generate Collision** option.
 	* Disable the **Remove Degenerates** option.
@@ -79,7 +79,7 @@ Once the Teapot mesh has been exported from 3Ds Max, it is now time for the Teap
 		![](VAT_SM_Import_Options.png)
 
 1. After the Teapot mesh has been imported, open the mesh up by Double - Clicking on it with the Left Mouse button in Content Browser. 
-Once the mesh has been opened enable the **Use Full Precision UVs** option and set the **Distance Field Resolution Scale** to **0.0** then press the **Apply Changes** button..
+Once the mesh has been opened enable the **Use Full Precision UVs** option and set the **Distance Field Resolution Scale** to **0.0** then press the **Apply Changes** button.
 
 	![](VAT_KFM_UE4_Setup.png)
 
