TypoCorrection

diff --git a/Engine/Documentation/Source/Engine/Animation/Tools/VertexAnimationTool/VAT_TL_Meshes/VAT_TL_Meshes.INT.udn b/Engine/Documentation/Source/Engine/Animation/Tools/VertexAnimationTool/VAT_TL_Meshes/VAT_TL_Meshes.INT.udn
--- a/Engine/Documentation/Source/Engine/Animation/Tools/VertexAnimationTool/VAT_TL_Meshes/VAT_TL_Meshes.INT.udn
+++ b/Engine/Documentation/Source/Engine/Animation/Tools/VertexAnimationTool/VAT_TL_Meshes/VAT_TL_Meshes.INT.udn
@@ -1,7 +1,7 @@
-Availability:Public
+﻿Availability:Public
 Title:Vertex Animation Tool - Timeline Meshes
 Crumbs: %ROOT%, Engine, Engine/Animation
-Description: User guide for the 3D Max Vertex Animation tool set.
+Description: User guide for the 3Ds Max Vertex Animation tool set.
 
 [TOC(start:2 end:3)]
 
@@ -29,7 +29,7 @@ In the following section we will cover how you can use the 3Ds Max Animation Tim
 	
 	![](VAT_Tool_Settings_Setup.png)
 
-1. Now that the correct tool options have been set, press the **Process Animated Meshes** button to being the Texture creation process. 
+1. Now that the correct tool options have been set, press the **Process Animated Meshes** button to begin the Texture creation process. 
 About half way through the process a Windows Dialogue box will be displayed asking for a location to store the Textures that are generated.
 For this example a folder was placed on the Desktop called **Vertex_Anims** and the Textures where saved to that folder.
 When completed the folder that was placed on the Desktop should now have two Texture files inside of it,
@@ -50,7 +50,7 @@ Please keep in mind that each and every step listed below **MUST BE** completed
 ### Texture Import & Setup
 
 [region:tip]
-If you are unfamiliar with the UE4 Texture Import process please read the [Texture Import Guide](https://docs.unrealengine.com/latest/INT/Engine/Content/Types/Textures/Importing/index.html) for more information.
+If you are unfamiliar with the UE4 Texture Import process please read the [](Engine/Content/Types/Textures/Importing) documentation for more information.
 [/region]
 
 The first thing that needs to be completed is that both Textures that were exported from 3Ds Max need to be imported into UE4. 
@@ -90,7 +90,7 @@ Import the Teapot Static Mesh in the same way the Texture were imported by first
 When you do this the following dialogue box will be displayed. From this dialogue make sure that the following options are set / disabled.
 
 [region:tip]
-If you are unfamiliar with how to import Static Meshes into UE4 please check out the [FBX Static Mesh Pipeline](https://docs.unrealengine.com/latest/INT/Engine/Content/FBX/StaticMeshes/index.html) document for more information about this..
+If you are unfamiliar with how to import Static Meshes into UE4 please check out the [](Engine/Content/FBX/StaticMeshes) document for more information about this.
 [/region]
 
 * Disable the **Auto Generate Collision** option.
@@ -117,12 +117,12 @@ With the Static Mesh now set up correctly, make sure to press the **Apply Change
 Now that all of the required content has been imported and properly setup, it is now time to create a Material that can use the content.
 
 [region:tip]
-If you are unfamiliar with how Materials work inside UE4 please refer to the [Materials](https://docs.unrealengine.com/latest/INT/Engine/Rendering/Materials/index.html) documentation for more information.
+If you are unfamiliar with how Materials work inside UE4 please refer to the [](Engine/Rendering/Materials) documentation for more information.
 [/region]
 
 1. First create a new Material by Right-Clicking with the Mouse and then choosing the **Material** option that is under the **Create Basic Asset** section. Name the Material **MAT_Vertex_Animation**.
 
-1. Next open up the **MAT_Vertex_Animation** Material and locate the two Textures that were imported into the Generic Browser in the previous step. Select both Textures and then Drag them from the Content Browser and Drop them into the Material Graph.  
+1. Next open up the **MAT_Vertex_Animation** Material and locate the two Textures that were imported into the Content Browser in the previous step. Select both Textures and then Drag them from the Content Browser and Drop them into the Material Graph.  
 	
 	![](VAT_Materail_Adding_Textures.gif)(convert:false)
 
