TypoCorrection

diff --git a/Engine/Documentation/Source/Programming/UnrealArchitecture/TMap/TMap.INT.udn b/Engine/Documentation/Source/Programming/UnrealArchitecture/TMap/TMap.INT.udn
--- a/Engine/Documentation/Source/Programming/UnrealArchitecture/TMap/TMap.INT.udn
+++ b/Engine/Documentation/Source/Programming/UnrealArchitecture/TMap/TMap.INT.udn
@@ -1,4 +1,4 @@
-Title:TMap
+﻿Title:TMap
 Crumbs:
 Description:
 Availability:Public
@@ -99,7 +99,7 @@ Here, the resulting map is equivalent to using Add/Emplace to add them individua
 
 ## Iteration
 
-Iteration over TMaps is similar to TArrays.  You can use C++’s ranged-for feature, remembering that the element type is a TPair:
+Iteration over TMaps is similar to TArrays.  You can use C++'s ranged-for feature, remembering that the element type is a TPair:
 
 	for (auto& Elem : FruitMap)
 	{
@@ -317,7 +317,7 @@ You can sort by key or by value using the KeySort and ValueSort methods respecti
 
 ## Operators
 
-Like TArray, TMap is a regular value type and as such can be copied via the standard copy constructor or assignment operator.  As maps strictly own their elements, copying a map is ‘deep’ and so the new map will have its own copy of the elements:
+Like TArray, TMap is a regular value type and as such can be copied via the standard copy constructor or assignment operator.  As maps strictly own their elements, copying a map is 'deep' and so the new map will have its own copy of the elements:
 
 	TMap<int32, FString> NewMap = FruitMap;
 	NewMap[5] = "Apple";
