#UE4 Docs: Adding 4.8 meta data and fixing some bolding based on feedback from doc team.

diff --git a/Engine/Documentation/Source/Engine/Animation/Tools/VertexAnimationTool/VAT_KF_Meshes/VAT_KF_Meshes.INT.udn b/Engine/Documentation/Source/Engine/Animation/Tools/VertexAnimationTool/VAT_KF_Meshes/VAT_KF_Meshes.INT.udn
--- a/Engine/Documentation/Source/Engine/Animation/Tools/VertexAnimationTool/VAT_KF_Meshes/VAT_KF_Meshes.INT.udn
+++ b/Engine/Documentation/Source/Engine/Animation/Tools/VertexAnimationTool/VAT_KF_Meshes/VAT_KF_Meshes.INT.udn
@@ -2,6 +2,7 @@
 Title:Vertex Animation Tool - Key Framed Meshes
 Crumbs: %ROOT%, Engine, Engine/Animation
 Description: User guide for the 3Ds Max Vertex Animation tool set.
+Version: 4.8
 
 [TOC(start:2 end:3)]
 
diff --git a/Engine/Documentation/Source/Engine/Animation/Tools/VertexAnimationTool/VAT_TL_Meshes/VAT_TL_Meshes.INT.udn b/Engine/Documentation/Source/Engine/Animation/Tools/VertexAnimationTool/VAT_TL_Meshes/VAT_TL_Meshes.INT.udn
--- a/Engine/Documentation/Source/Engine/Animation/Tools/VertexAnimationTool/VAT_TL_Meshes/VAT_TL_Meshes.INT.udn
+++ b/Engine/Documentation/Source/Engine/Animation/Tools/VertexAnimationTool/VAT_TL_Meshes/VAT_TL_Meshes.INT.udn
@@ -2,6 +2,7 @@
 Title:Vertex Animation Tool - Timeline Meshes
 Crumbs: %ROOT%, Engine, Engine/Animation
 Description: User guide for the 3Ds Max Vertex Animation tool set.
+Version: 4.8
 
 [TOC(start:2 end:3)]
 
diff --git a/Engine/Documentation/Source/Engine/Animation/Tools/VertexAnimationTool/VertexAnimationTools.INT.udn b/Engine/Documentation/Source/Engine/Animation/Tools/VertexAnimationTool/VertexAnimationTools.INT.udn
--- a/Engine/Documentation/Source/Engine/Animation/Tools/VertexAnimationTool/VertexAnimationTools.INT.udn
+++ b/Engine/Documentation/Source/Engine/Animation/Tools/VertexAnimationTool/VertexAnimationTools.INT.udn
@@ -4,6 +4,7 @@ Crumbs: %ROOT%, Engine, Engine/Animation
 Description: User guide for the 3D Max Vertex Animation tool set.
 Related: Engine\Animation\Tools\VertexAnimationTool\VAT_KF_Meshes
 Related: Engine\Animation\Tools\VertexAnimationTool\VAT_TL_Meshes
+Version: 4.8
 
 [TOC(start:2 end:3)]
 
@@ -18,7 +19,7 @@ Before we can begin using the script there are few things to be aware of so you
 
 ### Script Location
 [region:warning]
-This script is only available in Unreal Engine 4 (UE4) Version **4.8** or later. Earlier versions of UE4 **WILL NOT** have this script or required Material Functions available for use.
+This script is only available in Unreal Engine 4 (UE4) Version **4.8** or later.
 [/region]
 
 The Vertex Animation Script will be in the following location and is called **VertexAnimationTools.ms.**
@@ -34,7 +35,7 @@ The tool generates the data in the Texture using the following formula.
 		Final texture resolution:  X = number of vertex in the mesh, Y = Number of frames captured.
 
 Due to this limitation this tool is best suited for Visual Effects or ambient Static Meshes that require animations but might not justify using a complex animation rig to do so. 
-This tool will also **NOT WORK** with Skeletal Mesh Animations as the bone transforms are not available to the Material Editor.
+This tool will also **not work** with Skeletal Mesh Animations as the bone transforms are not available to the Material Editor.
 This means that if you want to affect the vertices of a Skeletal Mesh in a similar manner you will need to use [Morph Targets](Engine/Content/FBX/MorphTargets) instead.
 
 ### Vertex Animation Tools Break Down
