#UE4doc #docs: Changed instances of "BaseColor" to "Base Color" (as it appears in the Material Editor) sp:" separable"

diff --git a/Engine/Documentation/Source/Engine/Rendering/Materials/LightingModels/SubSurfaceProfile/SSSProfileLightingModel.INT.udn b/Engine/Documentation/Source/Engine/Rendering/Materials/LightingModels/SubSurfaceProfile/SSSProfileLightingModel.INT.udn
--- a/Engine/Documentation/Source/Engine/Rendering/Materials/LightingModels/SubSurfaceProfile/SSSProfileLightingModel.INT.udn
+++ b/Engine/Documentation/Source/Engine/Rendering/Materials/LightingModels/SubSurfaceProfile/SSSProfileLightingModel.INT.udn
@@ -10,7 +10,7 @@ Description:Description and technical details of the subsurface profile shading
 The ability to rendering realistic looking human skin is a must for any modern Video Game engine. To fulfill this need, Unreal Engine 4 (UE4) now offers a shading method specifically for skin or wax like surfaces called **Subsurface Profile**. 
 While the Subsurface Profile shading model has similar properties to the **Subsurface** shading model, its key difference is in how it renders. 
 SubsurfaceProfile uses a rendering approach that is based in Screen Space because this helps to better display the subtle subsurface effects seen in human skin where backscattering is a secondary effect only seen in few cases like ears for example.
-The following document will cover the what a Subsurface Profile is and how you can use them in your work. 
+The following document will cover what a Subsurface Profile is and how you can use them in your work. 
 
 [REGION:lightbox]
 [![](Results.png)(w:920 convert:false)](Results.png)
@@ -68,7 +68,7 @@ Once open, you can adjust the individual properties of the Subsurface scattering
 * ** Scatter Radius:** The distance in world space units to perform the scatter. 
 
 * ** Subsurface Color:** The Subsurface Color can be used as a weight for the subsurface effect. Black means there is no subsurface scattering. 
-White means all lighting is entering the material and gets scattered around. A non grayscale value give more control over what
+White means all lighting is entering the material and gets scattered around. A non-grayscale value give more control over what
 color contributions are entering the surface resulting in a more complex looking shading. 
 
 * **Falloff Color:** The Falloff Color defines the material scattering color once the light has entered the material. 
@@ -85,8 +85,8 @@ You should avoid using a vivid color here if you want to get a more complex shad
 The Screen Space Subsurface shading profile does not differ much from the Lit shading mode, with the main difference being that the Metallic input has been repurposed and is not available to be used.  
 
 **Base Color Input:** The Base Color input is used as usual for Diffuse lighting. There is no extra Subsurface Scattering color as the Screen Space subsurface scatter is not supposed to change color or brightness, it is just redistributing the lighting to nearby pixels. 
-So if a material should scatter in a specific color it needs to be expressed as part of the BaseColor. 
-The BaseColor should be the final color as if the material is viewed from a large distance where the scattering cannot be distinguished from diffuse lighting. 
+So if a material should scatter in a specific color it needs to be expressed as part of the Base Color. 
+The Base Color should be the final color as if the material is viewed from a large distance where the scattering cannot be distinguished from diffuse lighting. 
 
 [REGION:note]
  Human skin can be seen as a thin pale layer with a more vibrant red color below its surface and this is how the default colors look like. The visible scatter distance in Human skin is about 1.2 cm and that is the default of the Subsurface Profile.
@@ -119,16 +119,16 @@ At the moment, the Subsurface scattering profile shading model does not differ m
 The Subsurface scattering profile is based on work from the work from [[REDACTED_PERSON]]([REDACTED_PROFILE_URL]). Make sure to check out his webpage for many useful tips on how to make your 3D images look more realistic.
 [/REGION]
 
-We separate the non specular (non view dependent) lighting contributions to support a specular on top of subsurface material and down sample for better performance. 
-Similar to a Gaussian blur we filter the image with a 2 pass (assuming a separatable kernel) post process. 
+We separate the non-specular (non-view dependent) lighting contributions to support a specular on top of subsurface material and down sample for better performance. 
+Similar to a Gaussian blur we filter the image with a 2 pass (assuming a separable kernel) post process. 
 The filter kernel depends on the subsurface scattering profile which is stored in the GBuffer (Max 255 active profiles per scene). 
 The kernel has colored weights and specific sample positions which can be scaled in the profile (defined in units/cm). 
-In a final step, we recombine the scattered light contribution with the full resolution image. To separate the view and non view dependent lighting, we store a weighting term in the scene color alpha channel. 
+In a final step, we recombine the scattered light contribution with the full resolution image. To separate the view and non-view dependent lighting, we store a weighting term in the scene color alpha channel. 
 This approximation requires a 64 bit render target (see r.SceneColorFormat) and this approximation works for most cases
 
 It successfully takes out the specular but you get a more desaturated non view dependent color for those specular pixels. This can be improved by using two 32-bit render targets for all lighting passes. That has the same memory bandwidth but on some hardware this might be slower. This might be something we want to change (added code complexity).
 
-Here we have an example where the Specular was removed before the blurring was applied. Notice how in the end image(image on the far right) the Specular reflection is crisp and smooth. This is the effect we wanted to achieve.
+Here we have an example where the Specular was removed before the blurring was applied. Notice how in the end image (image on the far right) the Specular reflection is crisp and smooth. This is the effect we wanted to achieve.
 
 [REGION:lightbox]
 [![](Good_Combination.png)(w:920 convert:false)](Good_Combination.png)
@@ -157,7 +157,7 @@ The following image shows a bit more of the internals of the system. This view c
 While the Subsurface scattering profile shading model is a step forward when it comes to rendering skin, there are some limitations in what it can do.
  _Please note that as this system becomes more and more polished, this list could change. _
 
-* The feature does not work in the non deferred (mobile) rendering mode.
+* The feature does not work in the non-deferred (mobile) rendering mode.
 * Setting a large screen scatter radius will show banding artifacts in extreme lighting conditions.
 * Currently there is no backscatter of light. 
 * Currently a grey outline occurs when SSS Materials are occluded by non SSS Materials.
