#UE4doc #docs: Added hyphen to "non-grayscale" for consistency

diff --git a/Engine/Documentation/Source/Engine/Rendering/Materials/MaterialInputs/MaterialInputs.INT.udn b/Engine/Documentation/Source/Engine/Rendering/Materials/MaterialInputs/MaterialInputs.INT.udn
--- a/Engine/Documentation/Source/Engine/Rendering/Materials/MaterialInputs/MaterialInputs.INT.udn
+++ b/Engine/Documentation/Source/Engine/Rendering/Materials/MaterialInputs/MaterialInputs.INT.udn
@@ -179,7 +179,7 @@ The Subsurface Color input is only enabled when the [Shading Model](Engine/Rende
 
 ## Ambient Occlusion
 
-Ambient Occlusion is used to help simulate the self shadowing that happens within crevices of a surface. Generally, this input will be connected to an AO map of some type, which are often created within 3D modeling packages such as Maya, 3ds Max, or ZBrush.
+Ambient Occlusion is used to help simulate the self-shadowing that happens within crevices of a surface. Generally, this input will be connected to an AO map of some type, which is often created within 3D modeling packages such as Maya, 3ds Max, or ZBrush.
 
 ![](NormalNetwork.png)
 
