#UE4doc #docs: removed unnecessary word, "the" duplicate word, "be"

diff --git a/Engine/Documentation/Source/Engine/Rendering/Materials/LightingModels/SubSurfaceProfile/SSSProfileLightingModel.INT.udn b/Engine/Documentation/Source/Engine/Rendering/Materials/LightingModels/SubSurfaceProfile/SSSProfileLightingModel.INT.udn
--- a/Engine/Documentation/Source/Engine/Rendering/Materials/LightingModels/SubSurfaceProfile/SSSProfileLightingModel.INT.udn
+++ b/Engine/Documentation/Source/Engine/Rendering/Materials/LightingModels/SubSurfaceProfile/SSSProfileLightingModel.INT.udn
@@ -45,7 +45,7 @@ The default settings for the Subsurface Profile are close to Caucasian skin. How
 
 ![Enable SubsurfaceProfile](1.png)
 
-The Subsurface Profile can also be be overridden in a Material Instance. To do this, you first need to open the Material Instance that you wish to change. 
+The Subsurface Profile can also be overridden in a Material Instance. To do this, you first need to open the Material Instance that you wish to change. 
 Then in the **Details** section of the Material Instance, enable **Override Subsurface Profile** and then supply the Subsurface profile you want to use in the **Subsurface Profile** input. 
 
 ![](6.png)
@@ -134,7 +134,7 @@ Here we have an example where the Specular was removed before the blurring was a
 [![](Good_Combination.png)(w:920 convert:false)](Good_Combination.png)
 [/REGION]
 
-Here we have an example where the Specular was not removed before the blurring was applied. Notice how in the end image(image on the far right) the Specular reflection is dull and looks a bit stretched. This is not the correct way to have this effect render.  
+Here we have an example where the Specular was not removed before the blurring was applied. Notice how in the end image (image on the far right) the Specular reflection is dull and looks a bit stretched. This is not the correct way to have this effect render.  
 
 [REGION:lightbox]
 [![](Bad_Combination.png)(w:920 convert:false)](Bad_Combination.png)
