#UE4doc #docs: sp: section, duplicate word "you", " its" to "it is" (proper usage), duplicate word "can"

diff --git a/Engine/Documentation/Source/Engine/Rendering/Materials/MaterialInstances/Editor/MaterialInstanceEditor.INT.udn b/Engine/Documentation/Source/Engine/Rendering/Materials/MaterialInstances/Editor/MaterialInstanceEditor.INT.udn
--- a/Engine/Documentation/Source/Engine/Rendering/Materials/MaterialInstances/Editor/MaterialInstanceEditor.INT.udn
+++ b/Engine/Documentation/Source/Engine/Rendering/Materials/MaterialInstances/Editor/MaterialInstanceEditor.INT.udn
@@ -67,16 +67,16 @@ The **Details** pane in the Material Instance editor is where all of the 'work'
  To enable a parameter, use the **Left Mouse Button** and click on the **Grayed Out Box** that is next to the name of the parameter you want to enable. When you do this, the box will
  no longer be grayed out and you will be able to adjust the parameter to your liking. 
 
-* **Parameter Groups** - In this sectio,n you can adjust all of the values that have been exposed to the Material Instance from the Parent Material. 
+* **Parameter Groups** - In this section you can adjust all of the values that have been exposed to the Material Instance from the Parent Material. 
 * **General** - The General section allows you to change the Parent Material, Phys Material, and also adjust how this Material Instance will influence the Lightmass build. 
-* **Material Overrides** - Material Overrides allows you you to adjust Material base values 
+* **Material Overrides** - Material Overrides allows you to adjust Material base values 
 
 
 ### Parent List
 
 ![](MIE_Parent.png)
 
-Since Material Instances can have other Material Instances as their parent, sometimes its difficult to find the original Material that the Material Instance is based off of.
+Since Material Instances can have other Material Instances as their parent, sometimes it is difficult to find the original Material that the Material Instance is based off of.
 The parent list solves that problem by displaying the chain of parents of the current Material Instance all the way up to the root Material that is the start of the chain.
 For example, the parent list displayed above shows a Material Instance named **Instance_Metallic_Medium** that has the Material Instance **MasterDemoMaterial_Parent** as its parent.
 We can see in the parent list that **MasterDemoMaterial_Parent** has **MasterDemoMaterial** as its parent. The instance currently being edited is shown in **bold**.
@@ -107,7 +107,7 @@ samples being used by the Material.
 
 ![](MIE_Ins_Count.png)
 
-If your Material has some kind of movement in it you and you want to preview it, you can can enable Real - Time preview by clicking on the stop watch icon.
+If your Material has some kind of movement in it you and you want to preview it, you can enable Real - Time preview by clicking on the stop watch icon.
 
 ![](MIE_Real_Time.png)
 
