#UE4doc #docs: sp: changed "backfaces" to "back faces" to be consistent with documentation (ie. "front faces") sp: shadow

diff --git a/Engine/Documentation/Source/Engine/Rendering/Materials/MaterialProperties/MaterialProperties.INT.udn b/Engine/Documentation/Source/Engine/Rendering/Materials/MaterialProperties/MaterialProperties.INT.udn
--- a/Engine/Documentation/Source/Engine/Rendering/Materials/MaterialProperties/MaterialProperties.INT.udn
+++ b/Engine/Documentation/Source/Engine/Rendering/Materials/MaterialProperties/MaterialProperties.INT.udn
@@ -17,7 +17,7 @@ or by **Right-Clicking** the Material in the **Content Browser** then selecting
 
 * Accessing the Property Matrix:
 
-![](PropertiesContextMenu_Mat.png)(h:445)
+![](PropertiesContextMenu_Mat.png)(h:445) 
 
 * The Property Matrix:
 
@@ -46,7 +46,7 @@ The base material node has a number of properties which affect how the material
 | **Decal Response**| This defines how the Material reacts on DBuffer decals(Affects look, performance, and texture/sample usage). Non DBuffer Decals can be disabled on the primitive(e.g. static mesh)|
 | **Shading Model** | [INCLUDE:#ShadingModels] |
 | **Opacity Mask Clip Value** | This is the reference value which the OpacityMask input of masked materials will clip against per-pixel.  Any value higher than OpacityMaskClipValue will pass and the pixel will be drawn (opaque), any value lower will fail and the pixel will be discarded (transparent). |
-| **TwoSided** | The normal will be flipped on backfaces, which means lighting will be calculated for both the front faces and the backfaces.  This is commonly used on foliage to keep from having to double up the number of polygons.  Two Sided does not work correctly with static lighting, since the mesh still uses only a single UV set for light mapping.  As a result, two sided materials with static lighting will be shaded the same on both sides. |
+| **Two Sided** | The normal will be flipped on back faces, which means lighting will be calculated for both the front faces and the back faces.  This is commonly used on foliage to keep from having to double up the number of polygons.  Two Sided does not work correctly with static lighting, since the mesh still uses only a single UV set for light mapping.  As a result, two sided materials with static lighting will be shaded the same on both sides. |
 | **Tangent Space Normal** | Tangent space normal's are calculated from the surface of the object, with the Z-Axis (blue) always pointing directly away from the surface. World space normal's use the world coordinate system to calculate pixel angle, ignoring the surface's original orientation. Tangent space calculations are a bit more costly in terms of performance, but are generally more convenient, as these are usually the types of normal maps you can create in 2D applications such as Photoshop.   Visually, tangent-space normal maps tend to look primarily blue-ish, while world-space maps are more vividly rainbow-colored.|
 | **Decal Blend Mode** | [INCLUDE:#DecalBlendModes] |
 | **Use Material Attributes** | This checkbox will cause the Material's primary node to be condensed down to a single input labeled "Material Attributes." This is useful when you need to blend in multiple Materials using Layered Materials, or when defining multiple Material types using a Make Material Attributes expression node. See the Layered Materials documentation for more information. |
@@ -84,7 +84,7 @@ Translucent self shadowing is a nice way to get a volumetrically lit translucent
 | **Translucent Self Shadow Second Density Scale** | This is a second self shadow density that can be set for variation. An internal gradient is created between this value and the Translucent Self Shadow Density Scale. |
 | **Translucent Self Shadow Second Opacity** | This sets the opacity value for secondary self shadowing., serving as a way to scale the gradient effect between self shadow and second self shadow. |
 | **Translucent Backscattering Exponent** | This controls the backscattering that is used when employing the Subsurface shading model along with Translucency. Larger values give a smaller, brighter backscattering highlight. This value is only used within volumetric translucent shadows formed from a directional light. |
-| **Translucent Multiple Scattering Extinction** | This gives a colored extinction value - which basically amounts to a shadow color - for objects that have a volumetric translucent shadoow, such as smoke or steam. |
+| **Translucent Multiple Scattering Extinction** | This gives a colored extinction value - which basically amounts to a shadow color - for objects that have a volumetric translucent shadow, such as smoke or steam. |
 | **Translucent Shadow Start Offset** | This is a world-space offset for the self shadow created within a translucent volume. The higher the number, the farther the shadow moves away from the lighting source. |
 
 
