#UE4doc #docs: sp: capitalized "Note that the..."

diff --git a/Engine/Documentation/Source/Engine/Rendering/Materials/MaterialProperties/LightingModels/LightingModels.INT.udn b/Engine/Documentation/Source/Engine/Rendering/Materials/MaterialProperties/LightingModels/LightingModels.INT.udn
--- a/Engine/Documentation/Source/Engine/Rendering/Materials/MaterialProperties/LightingModels/LightingModels.INT.udn
+++ b/Engine/Documentation/Source/Engine/Rendering/Materials/MaterialProperties/LightingModels/LightingModels.INT.udn
@@ -102,7 +102,7 @@ Some examples include acrylic or lacquer clear coats, and colored films over met
 
 ![](ClearCoat01.png)
 [region:caption]
-(note the flecks were done in the material editor and are not part of the shading model)
+(Note that the flecks were done in the material editor and are not part of the shading model)
 [/region]
 
 When using the Clear Coat Shading Model, you are given access to the following inputs:
