#UE4doc #docs: removed unnecessary commas, made "name" singular, numerous edits

diff --git a/Engine/Documentation/Source/Engine/Rendering/Materials/ParameterCollections/ParemeterCollections.INT.udn b/Engine/Documentation/Source/Engine/Rendering/Materials/ParameterCollections/ParemeterCollections.INT.udn
--- a/Engine/Documentation/Source/Engine/Rendering/Materials/ParameterCollections/ParemeterCollections.INT.udn
+++ b/Engine/Documentation/Source/Engine/Rendering/Materials/ParameterCollections/ParemeterCollections.INT.udn
@@ -64,18 +64,18 @@ In the following section we will look at how we can add a Material Collection No
 
 1. The first thing that needs to be done is that we need to create a Material Collection Parameter node and fill it in with data. 
 	To do this, **Right-Click** in the **Content Browser** and then select the **Materials & Textures** section and then select **Material Parameter Collection**. 
-	This will create a Material Parameter Collection in the **Content Browser**. Once created, you need to name it. For this example, we will be using , **MPC_Global_Material_Prams** as the name. 
-	When completed, you should have something that looks like this in your **Content Browser**.
+	This will create a Material Parameter Collection in the **Content Browser**. Once created, you need to name it. For this example we will be using **MPC_Global_Material_Prams** as the name. 
+	When completed you should have something that looks like this in your **Content Browser**.
 
 	![](MPC_Create_New_MPC.png)
 
 1. Once the Material Collection Parameter has been created open it up by **Double-Clicking** on it with the **Left Mouse Button**. 
 	Once the Material Collection Parameter is open add one Scalar Parameters to and name it **Global_Desatration**. 
-	Leave the default value of Global_Desatration at 0.0. When completed, your Material Collection Parameter should look like this. 
+	Leave the default value of Global_Desatration at 0.0. When completed your Material Collection Parameter should look like this. 
 	![](MPC_Setup_MPC.png)
 
 1. Now that the Material Collection Parameter has been created we now need to create a simple Material to use it with. 
-   For this example, we will be creating the following Material. 
+   For this example we will be creating the following Material. 
 
 	![](MPC_Simple_Material.png)
 	
@@ -85,7 +85,7 @@ In the following section we will look at how we can add a Material Collection No
 1. With the nodes now hooked up to the Main Material Node, it is now time to add both the Desaturation Material Expression node as well as the Material Expression Parameter node to the Material. 
 	To add the Desaturation Material Expression node, search for Desaturation in the **Palette** and then drag it from the Palette into the Material Graph. 
 	Once that has been placed, locate the Material Parameter Collection in the **Content Browser** and then drag it from the **Content Browser** into the Material Graph. 
-	When completed, your Material should look something like this. 
+	When completed your Material should look something like this. 
 
 	![](MPC_Added_New_Nodes.png)
 
@@ -94,7 +94,7 @@ In the following section we will look at how we can add a Material Collection No
 	Next, we are going to need to connect the output of the Material Collection Parameter to the **Fraction** input on the Desaturation node.
 	Once that is completed, connect the output of the Desaturation node to the Base Color input on the Main Material node. 
 	Finally, make sure to press the **Apply** button to compile your Material. 
-	When completed, your Material Graph should look like this. 
+	When completed your Material Graph should look like this. 
 
 	![](MPC_Completed_Material.png)
 
@@ -121,7 +121,7 @@ This is an extremely powerful way to manipulate Material Collection Parameter as
 
 1. Before we can begin, we need to do a few things first. 
 	The first thing that we need to do is to make some [Material Instances](Engine/Rendering/Materials/MaterialInstances) out of our Material. 
-	For this example, we need to create three Material Instances based off of the Material that was created above. 
+	For this example we need to create three Material Instances based off of the Material that was created above. 
 	After the Material Instances have been created make sure and adjust the color so that you can tell them apart. 
 	For this example, the colors Red, Green, and Blue have been chosen as the colors for our Material Instances. 
 	When completed, your **Content Browser** should look like this. 
@@ -130,16 +130,16 @@ This is an extremely powerful way to manipulate Material Collection Parameter as
 
 1. Now that the Material Instances have been created, it is now time to apply them to various objects in the scene by clicking on the Material Instance with the **Left Mouse Button** in the **Content Browser** and then Dragging them onto the objects in the scene. 
 	Once you are over the object that you want to use the Material Instance on, release the **Left Mouse Button** to apply the Material Instance to the object.
-	When completed, you should have something that looks like this.
+	When completed you should have something that looks like this.
 
 	![](MPC_Applied_Materials.png)
 
 1. With the Material Instances now applied to objects in the level, it is now time to open up the Level Blueprint. 
-	To open up the Level Blueprint, click on **Blueprints** icon on the main tool bar and then from the drop down menu, select **Level Blueprint**.
+	To open up the Level Blueprint, click on **Blueprints** icon on the main tool bar and then from the drop down menu select **Level Blueprint**.
 
 	![](MPC_Open_Level_Blueprint.png)
 
-1. Once the Level Blueprint has been opened, it is now time to begin laying down some nodes. For this example, we will be needing the following nodes and variables.
+1. Once the Level Blueprint has been opened, it is now time to begin laying down some nodes. For this example we will be needing the following nodes and variables.
 	* ** Left Mouse Button Pressed**
 	* ** Timeline**
 	* ** Set Scalar Parameter Value**
@@ -150,14 +150,14 @@ This is an extremely powerful way to manipulate Material Collection Parameter as
 
 1. With the correct nodes added to the Blueprint Graph, it is now time to setup them up. 
 	The first thing that we are going to setup is going to be the  ** Set Scalar Parameter Value**. 
-	To set this up, all we have to do is first define what Material Parameter Collection we are going to use. 
+	To set this up all we have to do is first define what Material Parameter Collection we are going to use. 
 	To do this, use your **Left Mouse Button** to click on the **Collection** input and then select which Material Parameter Collection you would wish to use.
 	Next, select the **Parameter Name** you wish to manipulate, in this case it will be **Global_Desaturation**. When completed, your **Set Scalar Parameter Value** node should look like this. 
 
 	![](MPC_SSPV_Node.png)
 
 1. Next we are going to need to set up our [Timeline](Engine/Blueprints/UserGuide/Timelines/Editor) so that when the **Left Mouse Button** is pressed, the colors will change over time and then back once the Left Mouse Button has been released.
-	To set up the Timeline, first use the **Left Mouse Button** and **Double-Click** on the Timeline to open it up. 
+	To set up the Timeline, first use the **Left Mouse Button** and **Double-Click** on the Timeline to open it. 
 	When you do, you should see something like this. 
 	![](MPC_Timeline_Editor.png)
 
@@ -165,7 +165,7 @@ This is an extremely powerful way to manipulate Material Collection Parameter as
 	The first thing that we need to do is to add a **Float Track** to the Timeline by clicking on the **Add Float Track** icon. 
 	Then we need to add two Key Tracks by using the **Right Mouse Button** and clicking anywhere in the light gray area and then selecting **Add Key** from the menu. 
 	We want to add two keys, one at 0 time and value of 0 and one Key at time 1 and value of 1. 
-	Once the key points have been added, do not forget to name the track so that it is easier to know what it does. For this example, we will use **Des_Amount** as the names. 
+	Once the key points have been added, do not forget to name the track so that it is easier to know what it does. For this example, we will use **Des_Amount** as the name. 
 	When completed, you should have something that looks like this. 
 	[OBJECT:EmbeddedVideo]
 		[PARAMLITERAL:width]
@@ -180,16 +180,16 @@ This is an extremely powerful way to manipulate Material Collection Parameter as
 	[/OBJECT]
 
 1. With the **Timeline** now set up correctly, you can close the Timeline menu and get back to the **Event Graph** for the **Level Blueprint**. 
-	Once you have the Level Blueprint available, we can start to connect all of the nodes together. 
+	Once you have the Level Blueprint available, you can start to connect all of the nodes together. 
 	When completed, your Blueprint should look something like this. 
 	![](MPC_Connecting_Nodes.png)
 
 1. Now that everything is connected, make sure to **Compile** your Level Blueprint by pressing the **Compile** button and then close the Blueprint Editor.
 	Once the Blueprint Editor has been closed, it is now time to test to see our changes happen in game. To do this, all you need to do is press the **Play** icon that is on the Main Toolbar.
 	Once you press play, a new window will be displayed with the level running inside of it.
-	To see the Material Collection Parameter in action, all you need to do is hold down the **Left Mouse Button**.
-	When you do, you should see the colors of the Materials in the scene go from their original color to grey.
-	When you release the **Left Mouse Button**, you should then see the Materials go from grey back to their original color. 
+	To see the Material Collection Parameter in action, hold down the **Left Mouse Button**
+	and the colors of the Materials in the scene will go from their original color to grey.
+	Release the **Left Mouse Button**, and the Materials will return their original color. 
 	[OBJECT:EmbeddedVideo]
 		[PARAMLITERAL:width]
 		640
@@ -205,18 +205,18 @@ This is an extremely powerful way to manipulate Material Collection Parameter as
 ## Limitations and Performance Characteristics 
 
 [EXCERPT:MaterialWarning]
-A Material can reference at most two different MaterialParameterCollections. One is typically used for game-wide values, 
-and the other can be used for level specific parameters. A collection can have at most 1024 scalar parameters and 1024 vector 
+A Material can reference, at most, two different MaterialParameterCollections. One is typically used for game-wide values 
+while the other can be used for level specific parameters. A collection can have up to 1024 scalar parameters and 1024 vector 
 parameters.
 [/EXCERPT:MaterialWarning]
 
 Modifying the number of parameters in a collection will cause a recompile of all Materials that reference that collection. 
-If you need to add many parameters, it can be faster to add many parameters up front in an empty map.
+If you need to add numerous parameters, it can be faster to add a lot of parameters up front in an empty map.
 
 If you rename a parameter, any Materials referencing that parameter will continue to work as expected. However, any Blueprint 
 referencing that parameter will now be broken. You must reassign the ParameterName of the Blueprint function which operates on 
 the collection to fix this. 
 
-Updating values in a MaterialParameterCollection is very efficient, much more efficient than setting many different parameters 
-on many Material instances.
+Updating values in a MaterialParameterCollection is much more efficient than setting many different parameters 
+on various Material instances.
 
