#UE4doc #docs: sp: benefits, necessary

diff --git a/Engine/Documentation/Source/Engine/Rendering/Materials/PhysicallyBased/PhysicallyBasedShaders.INT.udn b/Engine/Documentation/Source/Engine/Rendering/Materials/PhysicallyBased/PhysicallyBasedShaders.INT.udn
--- a/Engine/Documentation/Source/Engine/Rendering/Materials/PhysicallyBased/PhysicallyBasedShaders.INT.udn
+++ b/Engine/Documentation/Source/Engine/Rendering/Materials/PhysicallyBased/PhysicallyBasedShaders.INT.udn
@@ -21,7 +21,7 @@ that are involved with our physically based approach. For a complete rundown of
 Physically based shading means we approximate what light actually does as opposed to approximating what we
 intuitively think it should do. The end result is more accurate and typically more natural looking. Physically
 based Materials will work equally well in all lighting environments. In addition, the Material values can be
-less complex and interdependent, resulting in a more intuitive interface. These benifits are applicable even for non-photorealistic
+less complex and interdependent, resulting in a more intuitive interface. These benefits are applicable even for non-photorealistic
 rendering as evident in the latest films from Pixar [[4]](#Smits) and Disney [[3]](#Burley).
 
 For these reasons and more, Unreal Engine 4 has adopted a new physically based Material and shading model. For an in depth,
@@ -109,7 +109,7 @@ Roughness 0 to 1. Nonmetal top, metal bottom.
 #### Mapping Roughness
 
 [EXCERPT:RoughnessMap]
-Roughness is a property that will frequently be mapped on your objects, in order to add the most physical variation to the surface.
+Roughness is a property that will frequently be mapped on your objects in order to add the most physical variation to the surface.
 [/EXCERPT:RoughnessMap]
 
 [REGION:fullwidth]
@@ -165,7 +165,7 @@ Small scale geometry, especially details only present in the high poly and baked
 which is typically an AO map with very short trace distance. This is multiplied by the final BaseColor before output and multiplied with 0.5 (Specular default) as the Specular output.
 To be clear this is BaseColor = Cavity*OldBaseColor, Specular = Cavity*0.5.
 
-For advanced use, this can be used to control the index of refraction (IOR). We have not found this to be necesary for 99% of Materials. Below are Specular values based off of measured IOR.
+For advanced use, this can be used to control the index of refraction (IOR). We have not found this to be necessary for 99% of Materials. Below are Specular values based off of measured IOR.
 
 **Measured Specular Values:**
 
