#UE4doc #docs: sp: duplicate word "the"

diff --git a/Engine/Documentation/Source/Engine/Rendering/Overview/RenderingOverview.INT.udn b/Engine/Documentation/Source/Engine/Rendering/Overview/RenderingOverview.INT.udn
--- a/Engine/Documentation/Source/Engine/Rendering/Overview/RenderingOverview.INT.udn
+++ b/Engine/Documentation/Source/Engine/Rendering/Overview/RenderingOverview.INT.udn
@@ -129,7 +129,7 @@ The **Ambient Cubemap** effect applies a cubemap texture to the lighting of the
 
 ### Bloom
 
-**Bloom** is a real world light phenomena that can greatly add to the perceived realism of a rendered image at a moderate render performance cost. Bloom can be seen by the naked eye when looking at very bright objects that are on a much darker background. Even brighter objects also cause other effects (streaks, lens flares) but those are not covered by the classic bloom effect. Because our displays (e.g. TV, TFT, ...) usually do not support HDR (high dynamic range), we cannot really render very bright objects. Instead we simulate the effects that happen in the eye (retina subsurface scattering), when light hits the film (film subsurface scattering) or in front of the the camera (milky glass filter). The effect might not always be physically correct, but it can help to hint the relative brightness of objects or add realism to the LDR (low dynamic range) image.
+**Bloom** is a real world light phenomena that can greatly add to the perceived realism of a rendered image at a moderate render performance cost. Bloom can be seen by the naked eye when looking at very bright objects that are on a much darker background. Even brighter objects also cause other effects (streaks, lens flares) but those are not covered by the classic bloom effect. Because our displays (e.g. TV, TFT, ...) usually do not support HDR (high dynamic range), we cannot really render very bright objects. Instead we simulate the effects that happen in the eye (retina subsurface scattering), when light hits the film (film subsurface scattering) or in front of the camera (milky glass filter). The effect might not always be physically correct, but it can help to hint the relative brightness of objects or add realism to the LDR (low dynamic range) image.
 
 [REGION:fullwidth]
 ![Bloom Effect](bloom.png)
