#UE4doc #docs: sp: sp: real time, changed "its" to "it is," sp: distributions

diff --git a/Engine/Documentation/Source/Engine/Rendering/ParticleSystems/Cascade/CascadeEditor.INT.udn b/Engine/Documentation/Source/Engine/Rendering/ParticleSystems/Cascade/CascadeEditor.INT.udn
--- a/Engine/Documentation/Source/Engine/Rendering/ParticleSystems/Cascade/CascadeEditor.INT.udn
+++ b/Engine/Documentation/Source/Engine/Rendering/ParticleSystems/Cascade/CascadeEditor.INT.udn
@@ -260,7 +260,7 @@ The time menu allows you to make adjustments to playback speed of the **Viewport
 | --- | --- |
 | [REGION:tablesection]Time Menu[/REGION] || 
 | **Play/Pause** | If checked, the simulation is playing. If unchecked, the simulation pauses.  |
-| **Realtime** | If checked, the simulation plays in reatime. Unchecked, the simulation only plays when something updates the **Viewport** pane. |
+| **Realtime** | If checked, the simulation plays in real time. Unchecked, the simulation only plays when something updates the **Viewport** pane. |
 | **Loop** | If checked, the particle system will restart once it reaches the end of its cycle. If unchecked, it will play once only. |
 | **AnimSpeed** | This opens a submenu allowing you to adjust playback speed percentage between 100%, 50%, 25%, 10%, and 1%.  |
 
@@ -426,7 +426,7 @@ Note that you need to make sure the Distribution you are modifying is a 'curve'
 
 To create points in the Graph Editor, **ctrl-left click** on the spline for the value you want.  The easiest way to do this is by turning off all the other modules by using the checkbox discussed above.  All modules start at 0 with a single key at time 0. **Ctrl-Left clicking** on the spline anywhere in the timeline will create a point there.  This point can be dragged around at will, but as discussed above, if the spline represents a vector (XYZ), then it will move all 3 keys for that vector in time but not in value.
 
-**Right clicking** on a keypoint will bring up a menu and let you manually enter the Time or Value of that keypoint. If its a key in a color curve, it will also let you select its color using the color picker.
+**Right clicking** on a keypoint will bring up a menu and let you manually enter the Time or Value of that keypoint. If it is a key in a color curve, it will also let you select its color using the color picker.
 
 If the module is ColorOverLife, then the splines rendered will reflect the current color at that time, while the points will be colored to reflect the particular channel for that spline.
 
@@ -455,9 +455,9 @@ This calculation runs from top to bottom within the emitter stack. For example,
 
 ### Distribution Types
 
-Distributions are special data types designed to numbers in highly specialized ways. Most of the primary properties available within any particle module will have distrubutions assigned to them. These distrubutions might define that property as having a single value, a range, or to be interpolated along a curve.
+Distributions are special data types designed to numbers in highly specialized ways. Most of the primary properties available within any particle module will have distributions assigned to them. These distributions might define that property as having a single value, a range, or to be interpolated along a curve.
 
-The type of distribution you use will change based on the property. As an example, if you want the lifetime of all of your particles to be 5 seconds, then a distrubution of _DistributionFloatConstant_ would allow you to hold a single value. However, if you wanted the lifetime of the particles to have a random value between 2 and 5 seconds, you would switch to a _DistributionFloatUniform_, within which you could specify such a range. 
+The type of distribution you use will change based on the property. As an example, if you want the lifetime of all of your particles to be 5 seconds, then a distribution of _DistributionFloatConstant_ would allow you to hold a single value. However, if you wanted the lifetime of the particles to have a random value between 2 and 5 seconds, you would switch to a _DistributionFloatUniform_, within which you could specify such a range. 
 
 For new users, distributions may seem confusing at first, but you will quickly see how they allow a great range of freedom over how the various properties within your particle systems are handled. See the [Distributions Reference](Engine/Basics/Distributions) page for more details.
 
