#UE4doc #docs: sp: breakdown, punctuation sp: sp: real time, changed "its" to "it is," sp: distributions

diff --git a/Engine/Documentation/Source/Engine/Rendering/ParticleSystems/LODs/LODs.INT.udn b/Engine/Documentation/Source/Engine/Rendering/ParticleSystems/LODs/LODs.INT.udn
--- a/Engine/Documentation/Source/Engine/Rendering/ParticleSystems/LODs/LODs.INT.udn
+++ b/Engine/Documentation/Source/Engine/Rendering/ParticleSystems/LODs/LODs.INT.udn
@@ -39,7 +39,7 @@ The following section of the Cascade toolbar covers the LOD controls.
 
 ![Cascade_Toolbar.png](Cascade_Toolbar.png)
 
-A break down of each control follows.
+A breakdown of each control follows:
 
 
 | **Button** | **Name** | **Description** |
@@ -74,7 +74,7 @@ The following section will go over the intended design flow for creating particl
 
 
 
-In this example, we have a single SubUV emitter. It is set up in such a way that SubUV coordinate 0 shows a number 1, coordinate 1 shows a 2, and so on. The texture applied looks like this:
+In this example, we have a single SubUV emitter. It is set up in such a way that SubUV coordinate 0 shows a number 1, coordinate 1 shows a number 2, and so on. The texture applied looks like this:
 
 ![ImageTest_Tile.png](ImageTest_Tile.png)(w:256)
 
