#UE4 DOCS: removed note about RequiredAPI, fixing https://jira.ol.epicgames.net/browse/UEDOC-1511

will create a new Story for updating these pages

diff --git a/Engine/Documentation/Source/Programming/UnrealArchitecture/Reference/Classes/Specifiers/MinimalAPI/ClassSpecifier_MinimalAPI.INT.udn b/Engine/Documentation/Source/Programming/UnrealArchitecture/Reference/Classes/Specifiers/MinimalAPI/ClassSpecifier_MinimalAPI.INT.udn
--- a/Engine/Documentation/Source/Programming/UnrealArchitecture/Reference/Classes/Specifiers/MinimalAPI/ClassSpecifier_MinimalAPI.INT.udn
+++ b/Engine/Documentation/Source/Programming/UnrealArchitecture/Reference/Classes/Specifiers/MinimalAPI/ClassSpecifier_MinimalAPI.INT.udn
@@ -2,10 +2,6 @@ Title: MinimalAPI
 Crumbs:%ROOT%, Programming, Programming/UnrealArchitecture/Reference, Programming/UnrealArchitecture/Reference/Classes, Programming/UnrealArchitecture/Reference/Classes/Specifiers
 Description:
 Availability:Public
-Version: 4.5
+Version: 4.7
 
 Causes only the class's type information to be exported for use by other modules. The class can be cast to, but the functions of the class cannot be called (with the exception of inline methods). This improves compile times of exporting everything for classes that do not need all of their functions accessible in other modules.
-
-[REGION:note]
-The **RequiredAPI** function specifier can be used with individual functions to make them accessible in other modules. This can be used in combination with the `MinimalAPI` class specifier to make the class and a subset of its functions available while still keeping compile times as low as possible.
-[/REGION]
