#UE4doc #docs: minor edit

diff --git a/Engine/Documentation/Source/Engine/Rendering/ParticleSystems/Optimization/Results/Results.INT.udn b/Engine/Documentation/Source/Engine/Rendering/ParticleSystems/Optimization/Results/Results.INT.udn
--- a/Engine/Documentation/Source/Engine/Rendering/ParticleSystems/Optimization/Results/Results.INT.udn
+++ b/Engine/Documentation/Source/Engine/Rendering/ParticleSystems/Optimization/Results/Results.INT.udn
@@ -117,7 +117,7 @@ Bounds are costly the check and update each frame, therefore it is a good idea t
 
 There are some considerations for fixed vs. non-fixed bounds to keep in mind.
 
-Fast moving effects such as rocket trails, tracers, or anything attached to a projectile etc. should not have fixed bounds set, or have the fixed bounds set so large that the emitter cannot fly outside of the bounds.
+Fast moving effects such as rocket trails, tracers, or anything attached to a projectile etc. should not have fixed bounds set nor have the fixed bounds set so large that the emitter cannot fly outside of the bounds.
 
 When setting bounds, be aware of local vs. world space coordinates defined in your system.  If an emitter is rotated the opposite direction of your bounds orientation, it is possible for elements of the effect to fly out of your bounds and elements popping out of view can occur.
 
