#UE4doc #docs: duplicate word "the," changed "relate" to "relates"

diff --git a/Engine/Documentation/Source/Engine/Rendering/ParticleSystems/Overview/KeyConcepts.INT.udn b/Engine/Documentation/Source/Engine/Rendering/ParticleSystems/Overview/KeyConcepts.INT.udn
--- a/Engine/Documentation/Source/Engine/Rendering/ParticleSystems/Overview/KeyConcepts.INT.udn
+++ b/Engine/Documentation/Source/Engine/Rendering/ParticleSystems/Overview/KeyConcepts.INT.udn
@@ -102,7 +102,7 @@ If you change a property to a distribution type that changes over time, some mod
 
 * Absolute time is basically the containing emitter time. If you have an emitter set up for 3 loops of 2 seconds, Absolute time for modules in that emitter will run from 0 to 2 seconds three times.
 
-* Relative time is between 0 and 1, indicating the the lifetime of each particle.
+* Relative time is between 0 and 1, indicating the lifetime of each particle.
 
 
 Descriptions of the currently available modules can be found at [ParticleSystemReference](Engine/Rendering/ParticleSystems/Reference).
@@ -111,7 +111,7 @@ Descriptions of the currently available modules can be found at [ParticleSystemR
 
 ## Emitters, Particle Systems, and Emitter Actors
 
-As you work with Cascade to create your own particle effects, it is important to keep in mind how each one of the objects relate to one another. In this document, we have already discussed the concept of modules, but they only make up one component of a complete particle effect. All told, the components of a particle system are **modules**, **emitters**, **particle systems**, and **emitter Actors**. A solid way to remember how they relate is like so:
+As you work with Cascade to create your own particle effects, it is important to keep in mind how each one of the objects relates to one another. In this document, we have already discussed the concept of modules, but they only make up one component of a complete particle effect. All told, the components of a particle system are **modules**, **emitters**, **particle systems**, and **emitter Actors**. A solid way to remember how they relate is like so:
 
 * **Modules**, which define particle behavior and are placed within...
 * **Emitters**, which are used to emit a specific type of particle for an effect, and any number of which can be placed within a...
