#UE4doc #docs: proper conjugation: "Updates," changed "realtime" to "real time" sp: dependent, in reference to engine feature: "Real-Time"

diff --git a/Engine/Documentation/Source/Engine/Rendering/ParticleSystems/Reference/BaseClasses/ParticleSystemClass/ParticleSystemClass.INT.udn b/Engine/Documentation/Source/Engine/Rendering/ParticleSystems/Reference/BaseClasses/ParticleSystemClass/ParticleSystemClass.INT.udn
--- a/Engine/Documentation/Source/Engine/Rendering/ParticleSystems/Reference/BaseClasses/ParticleSystemClass/ParticleSystemClass.INT.udn
+++ b/Engine/Documentation/Source/Engine/Rendering/ParticleSystems/Reference/BaseClasses/ParticleSystemClass/ParticleSystemClass.INT.udn
@@ -19,8 +19,8 @@ The ParticleSystem class contains the following public variable(s):
 | **Orient ZAxis Toward Camera** | If _true_, the local Z-axis of the ParticleSystem will be oriented towards the camera at all times. |
 | **SecondsBeforeInactive** | If the particle system is not rendered for this amount of time (in seconds), it will go inactive no longer being ticked. Enter a zero to prevent the system from ever being forced to be inactive. |
 | [REGION:tablesection]Thumbnail[/REGION] ||
-| **ThumbnailWarmup** | The amount of time to warm-up the system before catching the thumbnail render when _Use Realtime Thumbnail_ is checked. |
-| **Use Realtime Thumbnail** | If _true_, the thumbnail displayed in the **Content Browser** for the ParticleSystem asset will be automatically captured using the default camera position and reflect the current settings and appearance of the ParticleSystem instead of using the saved thumbnail. [REGION:note]Using realtime thumbnail rendering can slow down performance of the **Content Browser**.[/REGION] |
+| **ThumbnailWarmup** | The amount of time to warm-up the system before catching the thumbnail render when _Use Real-Time Thumbnail_ is checked. |
+| **Use Real-Time Thumbnail** | If _true_, the thumbnail displayed in the **Content Browser** for the ParticleSystem asset will be automatically captured using the default camera position and reflect the current settings and appearance of the ParticleSystem instead of using the saved thumbnail. [REGION:note]Using Real-Time thumbnail rendering can slow down performance of the **Content Browser**.[/REGION] |
 | [REGION:tablesection]LOD[/REGION] ||
 | **LOD Distance Check Time** | This is how often (in seconds) the system should perform a distance check to determine the LOD level to utilize. (Only used when LODMethod is set to Automatic.) |
 | **LOD Method** | This enumeration indicates the method the system should utilize for selecting appropriate LOD levels. Two methods are available: [INCLUDE:#LODMethod] |
@@ -74,8 +74,8 @@ The ParticleSystem class contains the following public variable(s):
 [REGION:simpletable]
 | Mode | Description |
 | --- | --- |
-| `EPSUM_RealTime` | Update the emitters in realtime. |
-| `EPSUM_FixedTime` | Update the emitters at a fixed time step. This locks the system to the update time of the game and is performance dependant, so a lower framerate lowers the update time of the system (slowing it down), whereas a higher framerate speeds it up. Usually only used in special cases. |
+| `EPSUM_RealTime` | Updates the emitters in real time. |
+| `EPSUM_FixedTime` | Updates the emitters at a fixed time step. This locks the system to the update time of the game and is performance dependent, so a lower framerate lowers the update time of the system (slowing it down), whereas a higher framerate speeds it up. Usually only used in special cases. |
 [/REGION]
 `EPSUM_FixedTime` will step the system a given amount regardless of the current frame-rate. This mode should only be used when the effect timing is not key to another system - such as tying the emitter to an animation.
 [/EXCERPT:SystemUpdate]
