#UE4doc #docs: Changed title from "What is Bump Offset?" to "Bump Offset" duplicate word "the"

diff --git a/Engine/Documentation/Source/Engine/Rendering/Materials/HowTo/BumpOffset/BumpOffset.INT.udn b/Engine/Documentation/Source/Engine/Rendering/Materials/HowTo/BumpOffset/BumpOffset.INT.udn
--- a/Engine/Documentation/Source/Engine/Rendering/Materials/HowTo/BumpOffset/BumpOffset.INT.udn
+++ b/Engine/Documentation/Source/Engine/Rendering/Materials/HowTo/BumpOffset/BumpOffset.INT.udn
@@ -13,7 +13,7 @@ Related: Engine/Rendering/Materials
 **Bump Offset** mapping give a surface the illusion of depth by modifying the UV coordinates in a creative way to help further displace the texels from the surface of the object, giving the illusion that the surface has more details than it really does. 
 In this How To example, we will be covering not only how you can find the Bump Offset Material Expression, we will also be covering how you can utilize the Bump Offset node in your Materials.
 
-## What is Bump Offset
+## Bump Offset
 
 [INCLUDE:Engine\Rendering\Materials\ExpressionReference\Utility#BumpOffset]
 
@@ -78,7 +78,7 @@ Once named, **Double-Click** on the Material in the **Content Browser** to open
 	When you first connect from the Bump Offset Material Expression to the UV input on any of the three textures, you will notice a red warning placed at the bottom of the Bump Offset Material Expression. The Bump Offset Material Expression needs a Height Map in order to work correctly and currently one has not been provided. Once an input to the Height has been entered, this error will go away.   
 	[/Region]
 
-1. Once the Bump Offset Material Expression has been connected to the UV input of all the Textures, it now needs to be setup so that it can work. The first thing that needs to be done is that we need to give it a Height input. To do this, grab the Texture we are using for the Base Color input, T_Brick_Clay_New_D, and duplicate it by first selecting it with the mouse and then pressing **CRTL + W** on the keyboard. Once the node is duplicated, move it so that it is in front of the Bump Offset Material Expression and the connect the Alpha output of the Texture into the Height input of the the Bump Offset Material Expression. When completed, you should have something that looks like this. 
+1. Once the Bump Offset Material Expression has been connected to the UV input of all the Textures, it now needs to be setup so that it can work. The first thing that needs to be done is that we need to give it a Height input. To do this, grab the Texture we are using for the Base Color input, T_Brick_Clay_New_D, and duplicate it by first selecting it with the mouse and then pressing **CRTL + W** on the keyboard. Once the node is duplicated, move it so that it is in front of the Bump Offset Material Expression and the connect the Alpha output of the Texture into the Height input of the Bump Offset Material Expression. When completed, you should have something that looks like this. 
 	![](Bump_Offest_Duplicate_For_HM.gif)(convert:false)
 
 1. Now the settings for the Bump Offset Material Expression need to be set so that we do not get what is commonly called UV swimming. To adjust the settings of the Bump Offset Material Expression, select the Bump Offset Material Expression in the Material Graph by clicking on it with the **Left Mouse Button** and  then in the **Details** panel, set the **Height Ratio** to **0.005**. When completed, you should have something that looks like this. 
