#UE4doc #docs: Changed title to "Translucent Shadow Color," sp: know to known, "The" to "the"

diff --git a/Engine/Documentation/Source/Engine/Rendering/Materials/HowTo/ColoredTransluscentShadows/ColoredTransluscentShadows.INT.udn b/Engine/Documentation/Source/Engine/Rendering/Materials/HowTo/ColoredTransluscentShadows/ColoredTransluscentShadows.INT.udn
--- a/Engine/Documentation/Source/Engine/Rendering/Materials/HowTo/ColoredTransluscentShadows/ColoredTransluscentShadows.INT.udn
+++ b/Engine/Documentation/Source/Engine/Rendering/Materials/HowTo/ColoredTransluscentShadows/ColoredTransluscentShadows.INT.udn
@@ -19,11 +19,11 @@ There are many applications where this could be useful to inject some color into
  ![](translucent_T.JPG)(w:450)
 
 
-## What is Translucent Shadow Color
+## Translucent Shadow Color
 
-Translucent Shadow Color is the term used to describe how shadows will pick up colors of objects that they pass through also know as **Transmission**. 
+Translucent Shadow Color is the term used to describe how shadows will pick up colors of objects that they pass through also known as **Transmission**. 
 The amount of light that passes through the material is determined is directly determined by its opacity value between 0 and 1 and how much light is being cast onto the material. 
-For example, if the opacity was set to 0, the light transmission for a colored shadow from the material would be none. If The opacity value were set to 1, then it would be fully opaque. 
+For example, if the opacity was set to 0, the light transmission for a colored shadow from the material would be none. If the opacity value were set to 1, then it would be fully opaque. 
 
 
 ## Material Setup
