#UE4doc #docs: sp: known, Changed title to "Emmisive Glow" sp: cannot, bulb

diff --git a/Engine/Documentation/Source/Engine/Rendering/Materials/HowTo/EmissiveGlow/EmissiveGlow.INT.udn b/Engine/Documentation/Source/Engine/Rendering/Materials/HowTo/EmissiveGlow/EmissiveGlow.INT.udn
--- a/Engine/Documentation/Source/Engine/Rendering/Materials/HowTo/EmissiveGlow/EmissiveGlow.INT.udn
+++ b/Engine/Documentation/Source/Engine/Rendering/Materials/HowTo/EmissiveGlow/EmissiveGlow.INT.udn
@@ -12,11 +12,11 @@ Related: Resources/ContentExamples/MaterialNodes/1_5
 
 ## Overview 
 
-During the course of project development, the need might arise for a Self Illuminated Material or glowing Material. 
-Inside of Unreal Engine 4(UE4), Self Illuminated Materials are know as Emissive Materials. 
+During the course of project development, the need might arise for a Self-Illuminated Material or glowing Material. 
+Inside of Unreal Engine 4(UE4), Self-Illuminated Materials are known as Emissive Materials. 
 In the following How - To guide, we will be covering how you setup your Materials to use the Emissive input as well as how you can use the Emissive input to actually light the world that you are creating. 
 
-## What is Emissive Glow
+## Emissive Glow
 
 Emissive Glow or in UE4, Emissive Materials, give artists a very cheap and effective way to give the illusion that a Material is casting light when it really is not. 
 Emissive Materials do this by allowing artist to push the values of the Emissive input higher than 1.0 which will push the Material into the HDR range, 
@@ -24,7 +24,7 @@ giving off a Bloom effect that you might see when looking at a very bright light
 
 ## Lit & Unlit Emissive Materials
 
-Emissive Materials are self illuminating which means that they can be used with both Lit and Unlit Shading Models. 
+Emissive Materials are self-illuminating which means that they can be used with both Lit and Unlit Shading Models. 
 However whenever possible, Emissive Materials should use the Unlit Shading Model as it will result in a shader that is cheaper to render.
  There is also no visual difference between the two different Shading Models in terms of how the Emissive is rendered. 
  So whenever possible, use the Unlit Shading Mode as it will be cheaper. 
@@ -36,7 +36,7 @@ However whenever possible, Emissive Materials should use the Unlit Shading Model
 	other shader inputs would still work with the lighting in the level. 
 
 * **When to use the Unlit Shading Model**: You should use the Unlit Shading Model if your Emissive Material does not have to use the other shader inputs. An example of this would be if you are 
-	using the Emissive Material strictly to simulate light source such as a light card or the surface area of a light blub. 
+	using the Emissive Material strictly to simulate light source such as a light card or the surface area of a light bulb. 
 
 ## Emissive Material Creation
 
@@ -60,7 +60,7 @@ All of the content that is used in the following sections can be found if you in
 	![](EM_Adding_Nodes_&_Textures_To_Material.png)
 
 	[region:tip]
-	If you can not find the **Tech_Panel** textures that are used in this How - To, that means that you might have forgotten to include the **Starter Content** with your project. 
+	If you cannot find the **Tech_Panel** textures that are used in this How - To, that means that you might have forgotten to include the **Starter Content** with your project. 
 	To find this content, you can create another project, make sure to include the Start Content, and then use the [Migrate Tool](Engine/Content/Browser/UserGuide/Migrate) 
 	to move the textures from that project into your project. 
 	[/region]
