#UE4doc #docs: sp: Changed title to "Detail Texturing" and "can not" to "cannot"

diff --git a/Engine/Documentation/Source/Engine/Rendering/Materials/HowTo/DetailTexturing/DetailTexturing.INT.udn b/Engine/Documentation/Source/Engine/Rendering/Materials/HowTo/DetailTexturing/DetailTexturing.INT.udn
--- a/Engine/Documentation/Source/Engine/Rendering/Materials/HowTo/DetailTexturing/DetailTexturing.INT.udn
+++ b/Engine/Documentation/Source/Engine/Rendering/Materials/HowTo/DetailTexturing/DetailTexturing.INT.udn
@@ -16,7 +16,7 @@ When viewing a Material very closely, you will notice that the Textures used in
 To address this issue, you can use what is known as a Detail Texture to help hide a Textures pixilation when a Material is viewed very closely. 
 
 
-## What is Detail Texturing
+## Detail Texturing
 
 Detail texturing allows you to give the illusion of more detail in a Texture by bringing in a highly repeated Diffuse and Normal map combination which layers over the original Diffuse and Normal for an object. 
 This gives the illusion of greater detail at close range than normally would be possible. 
@@ -73,8 +73,8 @@ One way is not better than the other as both give the same results in the end.
 All of the content that is used in the following sections can be found if you include the **Starter Content** with your project. While the techniques that are shown here will work with any Textures, if you want to follow along, please make sure that your project does include the **Starter Content**.
 [/region]
 
-However, before we begin, lets first set up a very simple Material that we can use for testing. 
-To do this, first use the mouse to **Right-Click** in the **Content Browser** and then from the **Create Basic Asset** section of the pop-up Menu, select **Material**. 
+However, before we begin, first set up a very simple Material to use for testing. 
+To do this, **Right-Click** in the **Content Browser** and then from the **Create Basic Asset** section of the pop-up Menu, select **Material**. 
 Once selected, a new **Material** will be created. Name the Material **Detail_Material**. When completed, your **Content Browser** should look like the following.
 
 ![](DT_Creat_New_Material.png)
@@ -91,7 +91,7 @@ Once selected, a new **Material** will be created. Name the Material **Detail_Ma
 	* **T_Ground_Gravel_D**
 	* **T_Ground_Moss_N**
 
-	When completed, your Material Graph should look like the following. 
+	When completed, your Material Graph should look like this:
 
 	![](DT_Adding_Textures.png)
 
@@ -119,13 +119,13 @@ Once selected, a new **Material** will be created. Name the Material **Detail_Ma
 	
 	![](DT_All_Connections_Made.png)
 
-1. By adjusting the values of the **Constant** Material Expression nodes, you can effect how the Detail Textures will look. Here is an example of what happens when you set the **Scale** to a value of 1 then 5 and finally 10. 
+1. By adjusting the values of the **Constant** Material Expression nodes, you can affect how the Detail Textures will look. Here is an example of what happens when you set the **Scale** to a value of 1 then 5 and finally 10. 
 
 	![](DT_Different_Scale_Settings.png)
 
 ### Without the Detail Material Function
 
-If for whatever reason you can not use the **Detail Texturing** Material Function, you can always build this functionality inside the Material Graph using Material Expression nodes and the directions below.
+If for whatever reason you cannot use the **Detail Texturing** Material Function, you can always build this functionality inside the Material Graph using Material Expression nodes and the directions below.
 
 1. First make a copy of **Detail_Material** by **Right-Clicking** on the Material inside the **Content Browser** and then selecting **Duplicate** from the Menu. 
 
@@ -190,7 +190,7 @@ This is often referred to as Distance Based Texture Blending and can be set up f
 
 	![](DT_Distance_Based_Completed_Material.png)
 
-1. To see the effect in action, all you have to do is hold down the **Right Mouse Button** while inside the Material Preview window and then move the mouse forward and backwards. When you get far enough away from the preview mesh, you should see the Material transition from one texture to the other one. If you can not see this happening, adjust the first **Constant** that is plugged into the **Divide** Material Expression, currently set to 512, to something smaller as smaller numbers will allow the effect to happen closer to the object. 
+1. To see the effect in action, all you have to do is hold down the **Right Mouse Button** while inside the Material Preview window and then move the mouse forward and backwards. When you get far enough away from the preview mesh, you should see the Material transition from one texture to the other one. If you cannot see this happening, adjust the first **Constant** that is plugged into the **Divide** Material Expression, currently set to 512, to something smaller as smaller numbers will allow the effect to happen closer to the object. 
 	[OBJECT:EmbeddedVideo]
 	[PARAMLITERAL:width]
 	640
