#UE4doc #docs: Changed title to "Fresnel" sp: "cannot" sp: "effect" duplicate word, "the"

diff --git a/Engine/Documentation/Source/Engine/Rendering/Materials/HowTo/Fresnel/Fresnel.INT.udn b/Engine/Documentation/Source/Engine/Rendering/Materials/HowTo/Fresnel/Fresnel.INT.udn
--- a/Engine/Documentation/Source/Engine/Rendering/Materials/HowTo/Fresnel/Fresnel.INT.udn
+++ b/Engine/Documentation/Source/Engine/Rendering/Materials/HowTo/Fresnel/Fresnel.INT.udn
@@ -19,7 +19,7 @@ extra complexity that can quickly become hard to manage. Unreal Engine 4 (UE4) o
 This way, the artist has more control over the look and feel of the effect and can even use it for things like adding lighting to an unlit object to save on performance. 
 
 
-## What is Fresnel 
+## Fresnel 
 
 Fresnel is the term used to describe how the light you see reflects at different intensities based off of the angle you are viewing it from. 
 For example, if you are standing over a pool looking straight down at the pool, you will not see a lot of reflections in the water. 
@@ -32,7 +32,7 @@ The result is then clamped to [0,1] so you do not have any negative color in the
 
 ![](FN_Caculation_Example.png)
 
-In the middle of the Sphere, where the 0 is, you will see that there is no Fresnel effect. This because the camera is pointing directly at the surface normal. 
+In the middle of the Sphere, where the 0 is, you will see that there is no Fresnel effect. This is because the camera is pointing directly at the surface normal. 
 As the surface normal becomes more and more perpendicular to the camera, or getting closer and closer to the number 1, the Fresnel effect becomes more and more visible 
 and this is the type of behavior we want. 
 
@@ -53,7 +53,7 @@ It is **highly recommended** that you use this version to begin with and then mo
 |--------- | ---------- |
 | **ExponentIn** | This controls the fall off of the Fresnel Effect. |
 | **BaseReflectFrctionIn** | This specifies the fraction of specular reflection when the surface is viewed from straight on. Setting this to a value of 1 effectively disables Fresnel |
-| **Normal** | Here you can input a Normal to effect the way the Fresnel effect is rendered. This property can not be set by supplying either a Normal map or a Vector 3 input.|
+| **Normal** | Here you can input a Normal to effect the way the Fresnel effect is rendered. This property cannot be set by supplying either a Normal map or a Vector 3 input.|
 
 You have two options for adjusting the values for the Fresnel Material Expression.
 * Select the Fresnel Material Expression in the Material Graph and then input values via the **Details** Panel.
@@ -76,13 +76,13 @@ This node should only be used if the Fresnel that is under **Utility** does not
 
 | Property | Description|
 |--------- | ---------- |
-| **Normal Vector (V3)** |  Here you can input a Normal to effect the way the Fresnel effect is rendered. This property can not be set by supplying either a Normal map or a Vector 3 input. |
+| **Normal Vector (V3)** |  Here you can input a Normal to effect the way the Fresnel effect is rendered. This property cannot be set by supplying either a Normal map or a Vector 3 input. |
 | **Camera Vector(V3)** |This allows you to effect the Camera that is used to determine the Dot Product between the Camera and the surface normal. |
 | **Invert Fresnel(B)** | This will allow you to invert the effect of the Fresnel. This is useful if you want the Fresnel to only be applied in the center or at the edges. |
 | **Power(S)** | This allows you to control the power of the Fresnel effect. Entering bigger numbers will tighten up the areas affected. Smaller numbers will allow more areas to be affected by the Fresnel.  |
 | **Use Cheap Contrast(B)** | This will allow the Fresnel to use a cheaper method for calculating contrast between lights and darks. |
-| **Cheap Contrast Dark (S)** | This controls how dark the Fresnel will be. This will only have an affect if **Use Cheap Contrast** is enabled.  |
-| **Cheap Contrast Bright (S)** | This controls how bright the Fresnel effect will be. This will only have an affect if **Use Cheap Contrast** is enabled.|
+| **Cheap Contrast Dark (S)** | This controls how dark the Fresnel will be. This will only have an effect if **Use Cheap Contrast** is enabled.  |
+| **Cheap Contrast Bright (S)** | This controls how bright the Fresnel effect will be. This will only have an effect if **Use Cheap Contrast** is enabled.|
 
 
 ## How to Use Fresnel in your Material 
@@ -95,9 +95,9 @@ Creating a Material that uses Fresnel can be accomplished in the following steps
 
 	![](FN_Creat_New_Material.png)
 
-1. Open up the new Material by **Double-Clicking** on it with the **Left Mouse Button** in the **Content Browser** . 
+1. Open up the new Material by **Double-Clicking** on it with the **Left Mouse Button** in the **Content Browser. ** 
 	Once the Material is opened, locate the needed Material Expression nodes by searching for them in the **Palette** search menu.
-	To place Material Expression nodes into the the Material Graph, first select the node you want to place with the **Left Mouse Button** and with the **Left Mouse Button** still held down, drag the node into the Material Graph. 
+	To place Material Expression nodes into the Material Graph, first select the node you want to place with the **Left Mouse Button** and with the **Left Mouse Button** still held down, drag the node into the Material Graph. 
 	For this Material, the following type and amount of nodes are needed. 
 
 	| Node Name | Amount|
@@ -162,14 +162,14 @@ Notice how as the number increases, the Fresnel effect is pushed closer and clos
 
 ## Fresnel & Normal Maps
 
-If you input a Normal Map into the Normal input of the Fresnel, you will see that you can effect the surface in a number of very interesting ways. 
+If you input a Normal Map into the Normal input of the Fresnel, you will see that you can affect the surface in a number of very interesting ways. 
 This is because of the Normal maps influence over the surface normal. As we learned before, the Fresnel works by checking to see if the surface normal is perpendicular to the camera or not. 
 If it is perpendicular to the camera, then the Fresnel effect will take place. On a smooth sphere, this means that the Fresnel will only take place at the edges of object. 
 However when a Normal map is introduced, we change the surface normal to no longer be smooth, allowing the Fresnel effect to take happen in places that it normally would not if the 
 surface Normal had been smooth. 
 
 For a visual example at what is going on, take a look at the following image showing the effect the surface normal has on the Fresnel. 
-On the left hand side, 1, is what the Fresnel node will out put with a flat surface normal. Now take a look at the right hand side, 2. 
+On the left (image 1) is what the Fresnel node will output with a flat surface normal. Now take a look at the right hand side (image 2.)
 In this image, a Normal map was input which drastically changes the surface normal. 
 Because the surface normal is no longer flat everywhere, you get more of the Fresnel being displayed. 
 
@@ -189,7 +189,7 @@ The Material that was used to create the images above was the same Material that
 ![](FN_Adding_Normal_Map.png)
 
 
-The Fresnel Material Expression can also be supplied with a Normal map that can effect the way the Fresnel is rendered. 
+The Fresnel Material Expression can also be supplied with a Normal map that can affect the way the Fresnel is rendered. 
 This is particularly useful if you want to use the Fresnel to display something like an energy shield on a Material but still have that Material look normal when the shield is not active. 
 To achieve this using the Material above, the following things need to be done.
 
@@ -204,12 +204,12 @@ To achieve this using the Material above, the following things need to be done.
 ![](FN_Fresnel_Normal_Input.png)	
 
 Notice how the bricks look like they have nothing to do with the Fresnel. 
-You can even take this one step by adding an **Multiply** and **Constant** Expression node and then using that to control the intensity of Fresnel effect. 
+You can even take this one step by adding a **Multiply** and **Constant** Expression node and then using that to control the intensity of Fresnel effect. 
 Here is how that would look when added to the Material from above. 
 
 ![](FN_Controal_Fresnel_Effect.png)
 
-By adjusting the **Constant** that was added, you can control how the Fresnel effects the surface. 
+By adjusting the **Constant** that was added, you can control how the Fresnel affects the surface. 
 Setting the number to 0 turns off the effect and setting it to anything else will gradually reveal the effect like in the example below.
 
 ![](FN_Effect_Amount.png)
