#UE4doc #docs: sp: non-tiled, non-directional

diff --git a/Engine/Documentation/Source/Engine/Rendering/Materials/HowTo/Human_Skin/HumanSkin.INT.udn b/Engine/Documentation/Source/Engine/Rendering/Materials/HowTo/Human_Skin/HumanSkin.INT.udn
--- a/Engine/Documentation/Source/Engine/Rendering/Materials/HowTo/Human_Skin/HumanSkin.INT.udn
+++ b/Engine/Documentation/Source/Engine/Rendering/Materials/HowTo/Human_Skin/HumanSkin.INT.udn
@@ -103,9 +103,9 @@ In order for your Normal Map to work correctly for skin, you need to make sure t
 *  When combining the object's Normal map with the Detail Normal map, make sure that you use the Material Function **BlendAngleCorrectedNormals** to combine as you will yield better results. 
    It is also very useful to expose parameters to adjust the mixing of the Normal maps. You can take the detail normal map value, add some scalar value to the z component and normalize. 
    This way you can transition to a flat detail Normal Map from your regular one. 
-* As skin has different shading properties in different areas, it is best to combine the Roughness with an object mapped (non tiled) texture. 
+* As skin has different shading properties in different areas, it is best to combine the Roughness with an object mapped (non-tiled) texture. 
   As the detail Normal map causes the roughness to go up in the distance, we need to compensate for that.  We also need to verify the shading at multiple scales / view distances to ensure that it looks good at all distances from the camera.
-* Modern real-time rendering implements shading with per pixel lighting. Normally this leads to very high quality results but for very detailed Normal maps. By using Mip Maps, we avoid shimmering and get a better performance but Mip Mapping does not work well for Normal maps as  details in lower mips simply do not exist and the surface appears much more shiny than it should be. To counter that effect, we can adjust the roughness accordingly. This does not capture anisotropic details but for non directional features, it is a good approximation. 
+* Modern real-time rendering implements shading with per pixel lighting. Normally this leads to very high quality results but for very detailed Normal maps. By using Mip Maps, we avoid shimmering and get a better performance but Mip Mapping does not work well for Normal maps as  details in lower mips simply do not exist and the surface appears much more shiny than it should be. To counter that effect, we can adjust the roughness accordingly. This does not capture anisotropic details but for non-directional features, it is a good approximation. 
 
 
 ### Opacity
