#UE4doc #docs: Changed title to "Material Instancing" sp: the, affect

diff --git a/Engine/Documentation/Source/Engine/Rendering/Materials/HowTo/Instancing/Instancing.INT.udn b/Engine/Documentation/Source/Engine/Rendering/Materials/HowTo/Instancing/Instancing.INT.udn
--- a/Engine/Documentation/Source/Engine/Rendering/Materials/HowTo/Instancing/Instancing.INT.udn
+++ b/Engine/Documentation/Source/Engine/Rendering/Materials/HowTo/Instancing/Instancing.INT.udn
@@ -15,9 +15,9 @@ Related: Engine/Rendering/Materials/MaterialProperties
 
 Having to create, setup, and tweak standard Materials can be a very time consuming process in UE4. 
 To help speed up and streamline this process, Unreal Engine 4 offers a special type of Material called a **Material Instance**. 
-In the following How - To, you will learn how you can setup and use Material Instances in your projects.
+In the following How-To, you will learn to setup and use Material Instances in your projects.
 
-## What is Material Instancing
+## Material Instancing
 
 Material Instancing is a way to create one single Material, referred to as the Parent Material, that can then be used as a base to make a wide variety of different looking Materials. 
 To achieve this flexibility Material Instancing uses a concept called inheritance. 
@@ -53,7 +53,7 @@ You can add  **Parameter Nodes** to the Material Graph in two completely differe
 	![](Materails_Instances_Convert_A_Node.gif)(convert:false)
 
 [region:note]
-Keep in mind that you can not do this for every Material node. If you do not see the **Convert To Parameter** option when you **Right-Click** on a node, then that node can not be converted into a Parameter Node. 
+Keep in mind that you cannot do this for every Material node. If you do not see the **Convert To Parameter** option when you **Right-Click** on a node, then that node cannot be converted into a Parameter Node. 
 [/region]
 
 ## Creating a Material Instance
@@ -136,7 +136,7 @@ The names that you set for parameters inside of the Material reflect the names t
 
 The full power of the Material Instance system comes into play when you need to have lots and lots of Materials that look different but still have similar controls or if you need to control a Material in a specific way during game play. 
 In the following example, we will make three different Instanced Materials that are based off of the Parent Material that was created in the previous steps. 
-We will then change the Parent Material to show how changes to he Parent will effect the Children. 
+We will then change the Parent Material to show how changes to the Parent will affect the Children. 
 
 1. First we are going to need to make some more Material Instances using the **Material_Master** Material that was created above as the Parent Material. 
    To do this, use the Mouse to **Right-Click** on the **Material_Master** in the **Content Browser** and then select **Create New Material Instance** from the menu. 
