#UE4doc #docs: Changed title to "The Main Material Node," sp: cannot

diff --git a/Engine/Documentation/Source/Engine/Rendering/Materials/HowTo/Main_Material_Node/Main_Material_Node.INT.udn b/Engine/Documentation/Source/Engine/Rendering/Materials/HowTo/Main_Material_Node/Main_Material_Node.INT.udn
--- a/Engine/Documentation/Source/Engine/Rendering/Materials/HowTo/Main_Material_Node/Main_Material_Node.INT.udn
+++ b/Engine/Documentation/Source/Engine/Rendering/Materials/HowTo/Main_Material_Node/Main_Material_Node.INT.udn
@@ -16,7 +16,7 @@ The Main Material node of the Material Graph is where all of the Material Expres
 The results of what is input into the Main Material node is what will be displayed when the Material is compiled and applied to objects in game. 
 In the following How - To, we will cover all you need to know about how to use the Main Material node.
 
-## What is the Main Material Node
+## The Main Material Node
 
 Materials are created using a specialized coding language called High Level Shading Language or HLSL.
 HLSL allows Materials to talk directly to the graphics hardware giving artist and coders better control over what is being displayed on screen.
