#UE4doc #docs: Changed title to "Material Functions,"

diff --git a/Engine/Documentation/Source/Engine/Rendering/Materials/HowTo/Making_Functions/Making_Functions.INT.udn b/Engine/Documentation/Source/Engine/Rendering/Materials/HowTo/Making_Functions/Making_Functions.INT.udn
--- a/Engine/Documentation/Source/Engine/Rendering/Materials/HowTo/Making_Functions/Making_Functions.INT.udn
+++ b/Engine/Documentation/Source/Engine/Rendering/Materials/HowTo/Making_Functions/Making_Functions.INT.udn
@@ -21,7 +21,7 @@ The Fresnel style effect is most commonly used for highlighting the edges of obj
 Do not forget to check out the [Creating Material Functions]( Engine/Rendering/Materials/Functions/Creating) document as it also contains a lot of valuable information about how to create and use Material Functions.
 [/region]
 
-## What are Material Functions
+## Material Functions
 
 [INCLUDE:Engine/Rendering/Materials/Functions#functionsintro] 
 
@@ -50,8 +50,8 @@ Creating Material Functions can be accomplished in the few following steps.
 	![](Material_Function_Editor.png)
 
 	[region:tip]
-	The Material Function editor works in a very similar manner to the Material Editor works. 
-	The two biggest differences between being what the **Details** panel does and how the **Final Material** node works. 
+	The Material Function editor works in a very similar manner to how the Material Editor works, 
+	the two biggest differences between being what the **Details** panel does and how the **Final Material** node works. 
 	You can read more about how the Material Function editor works but checking out the **[Material Functions Overview](Engine/Rendering/Materials/Functions/Overview)**
 	[/region]
 
