Minor correction - page description.

diff --git a/Engine/Documentation/Source/Programming/Tutorials/PlayerInput/4/PlayerInput_4.INT.udn b/Engine/Documentation/Source/Programming/Tutorials/PlayerInput/4/PlayerInput_4.INT.udn
--- a/Engine/Documentation/Source/Programming/Tutorials/PlayerInput/4/PlayerInput_4.INT.udn
+++ b/Engine/Documentation/Source/Programming/Tutorials/PlayerInput/4/PlayerInput_4.INT.udn
@@ -1,7 +1,8 @@
 Availability: Public
 Title:4. On Your Own!
 Crumbs:%ROOT%, Programming, Programming/Tutorials, Programming/Tutorials/PlayerInput
-Description:Binding input to C++ code
+Description:Ideas for further exploration!
+
 
 [VAR:Steps]
 [OBJECT:Navigation]
