#UE4doc #docs: Changed title to "Parameters," "an" to "a,"

diff --git a/Engine/Documentation/Source/Engine/Rendering/Materials/HowTo/Making_Parameters/Making_Parameters.INT.udn b/Engine/Documentation/Source/Engine/Rendering/Materials/HowTo/Making_Parameters/Making_Parameters.INT.udn
--- a/Engine/Documentation/Source/Engine/Rendering/Materials/HowTo/Making_Parameters/Making_Parameters.INT.udn
+++ b/Engine/Documentation/Source/Engine/Rendering/Materials/HowTo/Making_Parameters/Making_Parameters.INT.udn
@@ -14,7 +14,7 @@ Related: Engine/Matinee/HowTo/MHT_6
 To interact with Material Instances, you need to use a special kind of Material Expression node called a **Material Parameter**. 
 In this How To, we will be looking at how you can add or convert Material Expression nodes into Material Parameters nodes and also how you can use those Material Parameters nodes inside of your Materials.
 
-## What are Parameters
+## Parameters
 
 Material Parameters are a special kind of Material Expression nodes that allow you to change the look of a Material without having to recompile the Material to see the effects.
 Material Parameter Expression nodes work like any other Material Node that you might use in the Material Graph but with one key difference. 
@@ -32,7 +32,7 @@ The simplest way to make a Material Parameter out of a Material Expression is to
 ![](Making_Parameters_Right_Click_Convert.gif)(convert:false)
 
 [region:note]
-Not all Material Expressions can be converted into Material Parameters. The best way to see if an Material Expression can be converted into a Material Parameter is to **Right-Click** on it and check to see if the **Convert to Parameter** option is available. If you do not see the **Convert to Parameter** option, then the Material Expression cannot be converted.
+Not all Material Expressions can be converted into Material Parameters. The best way to see if a Material Expression can be converted into a Material Parameter is to **Right-Click** on it and check to see if the **Convert to Parameter** option is available. If you do not see the **Convert to Parameter** option, then the Material Expression cannot be converted.
 [/region]
 
 ### Adding Parameters from the Palette
