#UE4doc #docs: Changed title to "Texture Masks," sp: "affected," "a" to "an"

diff --git a/Engine/Documentation/Source/Engine/Rendering/Materials/HowTo/Masking/Masking.INT.udn b/Engine/Documentation/Source/Engine/Rendering/Materials/HowTo/Masking/Masking.INT.udn
--- a/Engine/Documentation/Source/Engine/Rendering/Materials/HowTo/Masking/Masking.INT.udn
+++ b/Engine/Documentation/Source/Engine/Rendering/Materials/HowTo/Masking/Masking.INT.udn
@@ -15,10 +15,10 @@ Related: Engine/Rendering/Materials/HowTo/Instancing
 ## Overview 
 
 When creating 3D assets, you might find that you need the ability to define different surfaces types within the same Material. 
-A simple and cheap way to achieve this is to use a Texture Mask that defines which surface should be effected by which section of the Material. 
+A simple and cheap way to achieve this is to use a Texture Mask that defines which surface should be affected by which section of the Material. 
 In the following How To, we will cover how you can use Texture Masking inside of your Unreal Engine 4 (UE4) Materials.
 
-## What is a Texture Mask
+## Texture Masks
 
 A texture Mask is a grayscale texture, or a single channel (R, G, B, or A) of a texture, used to limit the area of an effect inside of a Material. 
 Quite often, a mask will be contained within a single channel of another texture, such as the **Alpha Channel** of the Diffuse or Normal map. 
@@ -39,7 +39,7 @@ In the following image, we can see this in action.
 
 On the right hand side, outlined in green, is what the Mask Texture looks like. 
 Notice how only a few areas are painted **White** while all other areas are painted **Black**. 
-This is because the we only want to effect the areas that are in **White**.
+This is because the we only want to affect the areas that are in **White**.
 On the left hand side, you can see what areas this Mask Texture will mask out when applied to the chair mesh.
 
 [region:note]
@@ -59,7 +59,7 @@ If do you pack something into the Alpha Channel of a texture, make sure that you
 ## Using Texture Masks
 
 Mask Textures can be used a number of different ways inside the UE4 Material Editor. 
-This can range from defining a emissive light source to being used as a Roughness texture. 
+This can range from defining an emissive light source to being used as a Roughness texture. 
 In the following section, we will cover some of the most common ways to use Texture masking in UE4. 
 
 ### Emissive Mask
@@ -94,7 +94,7 @@ Opening up the chair Material, **M_Chair**, we can see a perfect example of Text
 ![](TM_Chair_Material.png)
 
 In this Material, the Mask Texture,**T_Chair_M**, is used to control every aspect of how the chair looks. 
-From what sections should be metal or non metal to what color those sections should have, the Mask Texture helps to define all of this.
+From what sections should be metal or non-metal to what color those sections should have, the Mask Texture helps to define all of this.
 In the following image, you can see a breakdown of how each channel of the Mask Texture for SM_Chair works.
 On the **Left-Hand side** is what the Mask looks like as a Texture.
 The smaller images that you see on the **Right-Hand side** show each of the individual channels and what area those channels will effect on the chair Static Mesh. 
@@ -115,7 +115,7 @@ Here is a breakdown of what type of information is stored inside each channel of
 * **Alpha Channel**: Whole object Mask. This is not currently used by the Material.
 
 
-When all put together, the chair looks like the following inside a UE4 level.
+When all put together, the chair looks like the following inside a UE4 level:
 
 ![](TM_SM_Chair.png)
 
