#UE4doc #docs: Changed title to "Material Expressions," changed "making it looks like its " to "making it look like it is"

diff --git a/Engine/Documentation/Source/Engine/Rendering/Materials/HowTo/Placing_Expressions_Functions/Placing_Expressions_Functions.INT.udn b/Engine/Documentation/Source/Engine/Rendering/Materials/HowTo/Placing_Expressions_Functions/Placing_Expressions_Functions.INT.udn
--- a/Engine/Documentation/Source/Engine/Rendering/Materials/HowTo/Placing_Expressions_Functions/Placing_Expressions_Functions.INT.udn
+++ b/Engine/Documentation/Source/Engine/Rendering/Materials/HowTo/Placing_Expressions_Functions/Placing_Expressions_Functions.INT.udn
@@ -14,12 +14,12 @@ Material Expressions and Material Functions are the building blocks used to crea
 Each expression or function is a self-contained black box that can run little snippets of code on their inputs and output the results. 
 In this How To guide, we will be learning the various ways that we can place Material Expressions and Material Functions inside of the Material Graph. 
 
-## What are Material Expressions
+## Material Expressions
 
 Material Expression nodes are tiny snippets of HLSL code that perform very specific tasks inside of a Material. 
 Materials are built using combinations of Material Expression nodes are combined together to do certain tasks. 
 For example, if you want a Texture to move inside of a Material, you can use the Panner Material Expression node. 
-This node contains specific HLSL code that tells the Material to move the UV coordinates of the Texture that it is applied to in a specific direction, making it looks like its moving. 
+This node contains specific HLSL code that tells the Material to move the UV coordinates of the Texture that it is applied to in a specific direction, making it look like it is moving. 
 As you will soon see,  Material Expression nodes are key to building and interacting with the Materials you create.
 
 
@@ -103,7 +103,7 @@ The Palette offers the ability to filter your search so that you can limit your
 You can also add Material Functions by **Right-Clicking** anywhere in the Material graph and then selecting the Material Function you want from the menu context menu by clicking on it with the **Left Mouse Button** to place it in the Material graph. You can also use the **Keyboard** to search for the Material Function you want by typing the name of what you are looking for in the **Search** box. When you find the one you are looking for, press the **Enter** key to place the Function node in the Material Graph.
 
 [region:note]
-The **Right-Click** menu can not filter by Functions or Expression and displays them both when you are searching.
+The **Right-Click** menu cannot filter by Functions or Expression and displays them both when you are searching.
 [/region]
 
 ![](Placing_Functions_Function_Search.gif)(convert:false)
@@ -123,8 +123,8 @@ One thing that is exclusive to Material Functions is that you also have the abil
 
 [region:note]
 To find the Material Functions in the **Content Browser**, you will need to look in the Engine folder. 
-However by default, the Engine folder is disabled and to see it you will need to enable it. 
-To do this use the **Left Mouse Button** to click on to View Options that is on the lower right hand corner of the **Content Browser**. 
+However, by default the Engine folder is disabled and to see it you will need to enable it. 
+To do this, use the **Left Mouse Button** to click on to View Options that is on the lower right hand corner of the **Content Browser**. 
 Then from the pop-up menu, select the **Show Engine Content** option to display the content in the Engine Folder.  
 
 ![](Placing_Functions_Show_Engine_Content.gif)(convert:false)
