#UE4doc #docs: Changed title to "Previewing and Applying Materials," sp: "the"

diff --git a/Engine/Documentation/Source/Engine/Rendering/Materials/HowTo/Previewing_Applying/Previewing_Applying.INT.udn b/Engine/Documentation/Source/Engine/Rendering/Materials/HowTo/Previewing_Applying/Previewing_Applying.INT.udn
--- a/Engine/Documentation/Source/Engine/Rendering/Materials/HowTo/Previewing_Applying/Previewing_Applying.INT.udn
+++ b/Engine/Documentation/Source/Engine/Rendering/Materials/HowTo/Previewing_Applying/Previewing_Applying.INT.udn
@@ -16,7 +16,7 @@ Previewing and applying your Materials is something that you will have to do qui
 Unreal Engine 4 (UE4) makes both previewing and applying materials a very straight forward process.
 In the following How - To, we will cover all you need to know about how to preview and apply your Materials in UE4.
 
-## What is Previewing and Applying Materials
+## Previewing and Applying Materials
 
 Previewing and Applying your Materials is of critical importance to the Material creation process because this is how you will be able to see and apply the results of what was created in the Material Graph.
 When you are Previewing a Material, you are seeing what the changes would look like.
@@ -95,7 +95,7 @@ To enable this functionality, you need to make sure that the Material meets the
 ## How to Apply your Materials 
 
 There are two main ways that you can apply Materials to objects inside of UE4.
-* **Drag and Drop** - You can drag any Material from the **Content Browser** and then drop it on any object in the level to change its material like in he following example.
+* **Drag and Drop** - You can drag any Material from the **Content Browser** and then drop it on any object in the level to change its material like in the following example:
 ![](Previewing_Applying_Drag_&_Drop.gif)(convert:false)
 
 * **Apply via Object's Details Panel** - You can also apply materials via the object's **Details** panel by using the **Use selected asset from Content Browser** button like in the following example.
