#UE4doc #docs: Changed title to "Refraction," "everyday" to "every day"

diff --git a/Engine/Documentation/Source/Engine/Rendering/Materials/HowTo/Refraction/Refraction.INT.udn b/Engine/Documentation/Source/Engine/Rendering/Materials/HowTo/Refraction/Refraction.INT.udn
--- a/Engine/Documentation/Source/Engine/Rendering/Materials/HowTo/Refraction/Refraction.INT.udn
+++ b/Engine/Documentation/Source/Engine/Rendering/Materials/HowTo/Refraction/Refraction.INT.udn
@@ -15,9 +15,9 @@ When light passes from one medium to another, say from air to water, the directi
 This lighting phenomenon is referred to as Refraction and happens because the speed at which light travels can be altered by the different types of materials light come in contact with. 
 Inside Unreal Engine 4, you can simulate this lighting phenomena in your Materials by using Refraction. 
 
-## What is Refraction
+## Refraction
 
-Refraction is the term used to describe the change in direction of a wave(light) due to a change in its transmission medium. 
+Refraction is the term used to describe the change in direction of a light wave due to a change in its transmission medium. 
 In other words, when light comes in contact with certain surfaces, like water or glass, the light will get slightly bent because those surfaces effect the speed at which light travels through them.
 The best example of Refraction can be seen by placing part of a pencil in water.
 The part of the pencil that is placed in the water will appear bent  where the pencil and water meet due to Refraction.
@@ -34,7 +34,7 @@ This is happening because the speed at which light travels changes as it hits th
 
 Index of Refraction or IOR is an optical measurement that can be used to tell you exactly how much light will be bent when light propagates from one medium to another. 
 Since this lighting phenomenon can be measured, we can use these measurements inside UE4 to affect how the Refraction will look giving a better sense of realism to our 3D world. 
-Here are some IOR values for objects that you might interact with everyday. 
+Here are some IOR values for objects that you might interact with every day. 
 
 |Material | IOR Value|
 |---|---|
@@ -91,7 +91,7 @@ This tutorial was written assuming that you included the Starter Content with yo
 	![](R_Adding_Test_Meshes.png)
 
 	[region:note]
-	If you can not find the Cube, Cone, or Sphere, this could mean that you did not include the Starter Content with your project. If you did not include the Starter Content, you can either create a new project, this time make sure to include the Starter Content, or you can [Migrate](Engine/Content/Browser/UserGuide/Migrate) the Starter Content from another project into this one.
+	If you cannot find the Cube, Cone, or Sphere, this could mean that you did not include the Starter Content with your project. If you did not include the Starter Content, you can either create a new project, this time make sure to include the Starter Content, or you can [Migrate](Engine/Content/Browser/UserGuide/Migrate) the Starter Content from another project into this one.
 	[/region]
 
 1. Now that there are some testing shapes, apply the Refraction Material Instance to the shapes by first selecting the Material Instance with the **Left Mouse Button** in the **Content Browser** and then dragging the Material Instance over the shape. Release the **Left Mouse Button** when you are over the shape you want to apply the Material Instance to.  When completed, you should have something that looks like this in your level.
@@ -186,7 +186,7 @@ Refraction depth bias is a way to prevent closer objects from rendering into the
 However this can increases the disconnect between the surface and where the refraction starts.
 You can adjust the **Refraction Depth Bias** in the following two places inside UE4. 
 
-* **Material Editor** - You can find the **Refraction Depth Bias** in the **Details** panel under the Material section. You will need to press the white triangle(Highlighted in Green) to expose this property.
+* **Material Editor** - You can find the **Refraction Depth Bias** in the **Details** panel under the Material section. You will need to press the white triangle (Highlighted in Green) to expose this property.
 
 	![](R_RDB_Materail_Editor.png)
 
