#UE4doc #docs: Changed title to "Shininess," usage: affect vs effect (affect is a verb and effect is a noun) sp: enabled

diff --git a/Engine/Documentation/Source/Engine/Rendering/Materials/HowTo/ShinyMaterials/ShinyMaterials.INT.udn b/Engine/Documentation/Source/Engine/Rendering/Materials/HowTo/ShinyMaterials/ShinyMaterials.INT.udn
--- a/Engine/Documentation/Source/Engine/Rendering/Materials/HowTo/ShinyMaterials/ShinyMaterials.INT.udn
+++ b/Engine/Documentation/Source/Engine/Rendering/Materials/HowTo/ShinyMaterials/ShinyMaterials.INT.udn
@@ -17,7 +17,7 @@ Inside of Unreal Engine 4, you can simulate the shininess of objects using Metal
 In the following How - To, you will learn everything you need to know about how to introduce Shininess to your Materials.
 
 
-## What is Shininess
+## Shininess
 
 Shininess simply refers to a surface's ability to reflect light. 
 The more light a surface can reflect, the more shiny that surface is. 
@@ -49,7 +49,7 @@ Here is an example of how Metallic is affected by Roughness.
 ![](Engine\Rendering\Materials\PhysicallyBased\roughness_metal.png)
 
 The sphere on the far left has a Metallic value of 1 and a Roughness value of 0 so it appears to be very shiny. 
-While the sphere on the far right has a Metallic value of 1 and a Roughness value of 1. 
+The sphere on the far right, however, has a Metallic value of 1 and a Roughness value of 1. 
 This gives the metal a very dull look, like it has been used quite a bit. 
 
 ## Specular & Shininess
@@ -58,7 +58,7 @@ This gives the metal a very dull look, like it has been used quite a bit.
 In 99 percent of cases, you will never need to adjust the Specular input and the default value of 0.5 will be good enough for most Materials. 
 [/region]
 
-Specular can also have an affect on the Shininess of a Material. 
+Specular can also have an effect on the Shininess of a Material. 
 Pushing the Specular value closer to 1 will make the reflections and Specular highlights in the Material appear stronger and more apparent 
 while lowering that value closer to 0 will weaken the reflections and Specular highlight until they are almost non-existent. 
 Here is an example showing how the intensity of the reflections and Specular highlight are affected as the Specular value is changed from 0 to 1.
@@ -67,7 +67,7 @@ Here is an example showing how the intensity of the reflections and Specular hig
 
 The Specular, like with every other input, is also heavily affected by the Roughness. 
 Even with the Specular input set to 1, the effect of the Specular can be canceled out by setting the Roughness to a value of 1 as well.
-Also adjusting the Specular will have no effect on the Material if Metallic is enable.
+Also adjusting the Specular will have no effect on the Material if Metallic is enabled.
 Again, in **99 Percent** of cases, you will never need to enable this value and using the default value of 0.5 will give best results. 
 
 ## Using Shininess in Materials
@@ -85,7 +85,7 @@ content page for information about how to move content between projects. This wa
 
 	![](Engine\Rendering\Materials\HowTo\Transparency\Trans_Create_New_Material.png)
 
-1. Name the Material **Shiny_Material** and then open it up by using the **Left Mouse Button** and **Double-Clicking** on the Material in the **Content Browser**. 
+1. Name the Material **Shiny_Material** and then open it up by using the **Left Mouse Button** and **Double-Clicking** on the Material in the **Content Browser.** 
 
 1. With the Material open add the following Textures and Material Expression Nodes to the Material Graph. Please note that the Texture can be found in the Textures Folder of the Starter Content.
 
@@ -96,7 +96,7 @@ content page for information about how to move content between projects. This wa
 
 	![](Shiny_Required_Nodes.png)
 
-1. Once the Material Expression nodes have been placed, the Scalar Parameters now need to be named and filled in with default values. Here are the names and default values that should be used. 
+1. Once the Material Expression nodes have been placed, the Scalar Parameters now need to be named and filled in with default values. Here are the names and default values that should be used:
 	
 	* ** Metallic:** 0
 	* **Specular:**  0.5
@@ -104,7 +104,7 @@ content page for information about how to move content between projects. This wa
 
 	![](Shiny_Named_Nodes.png)
 
-1. Now it is time to start to connect all of the Material Expression nodes to their corresponding inputs on the main Material node. When completed, the Material Graph should look something like this.
+1. Now it is time to start to connect all of the Material Expression nodes to their corresponding inputs on the main Material node. When completed, the Material Graph should look something like this:
 
 	![](Shiny_Connected_Nodes.png)
 
