sp: look and other misc edits

diff --git a/Engine/Documentation/Source/Engine/Rendering/Materials/HowTo/Subsurface_Profile/Subsurface_Profile.INT.udn b/Engine/Documentation/Source/Engine/Rendering/Materials/HowTo/Subsurface_Profile/Subsurface_Profile.INT.udn
--- a/Engine/Documentation/Source/Engine/Rendering/Materials/HowTo/Subsurface_Profile/Subsurface_Profile.INT.udn
+++ b/Engine/Documentation/Source/Engine/Rendering/Materials/HowTo/Subsurface_Profile/Subsurface_Profile.INT.udn
@@ -71,7 +71,7 @@ In the following steps, we will be creating a Material and a Material Instance t
 | **Roughness** | 0.35 |
 | **Subsurface_Amount**| 0.5|
 
-	When everything has been filled in, your shader graph should looks something like this. 
+	When everything has been filled in, your shader graph should look something like this. 
 	![](SSP_Named_Value_Nodes.png)
 
 1. With the Material Expression nodes now setup, it is time to connect all of them to the main Material node. 
@@ -109,7 +109,7 @@ When completed, your **Content Browser** should look something like this.
 	![](SSP_Material_Instance.png)
 
 1. To change the Subsurface Profile in a Material Instance, first make sure that you enable the **Override Subsurface Profile** and then provide the Subsurface Profile that you want to use 
-	by using the **Left Mouse  Button** to click on the **Subsurface Profile** input. Once you click on the input you should see the Subsurface Profile that are available to choose from. 
+	by using the **Left Mouse  Button** to click on the **Subsurface Profile** input. Once you click on the input you should see the Subsurface Profiles that are available to choose from. 
 	![](SSP_SSP_In_Material_Instance.png)
 
 	[region:note]
@@ -146,7 +146,7 @@ When completed, your **Content Browser** should look something like this.
 	**StaticMesh'/Game/StarterContent/Props/SM_Chair.SM_Chair'**
 
 	[region:tip]
-	If you can not find the SM_Chair mesh, it is probably because you did not include the Starter Content with your project. 
+	If you cannot find the SM_Chair mesh, it is probably because you did not include the Starter Content with your project. 
 	To find more about the Starter Content, please check out the [Start Content](Engine/Content/Packs) documentation page.
 	[/region]
 
@@ -173,7 +173,7 @@ When completed, your **Content Browser** should look something like this.
 	* ** Scatter Radius:** The distance in world space units to perform the scatter. 
 
 	* ** Subsurface Color:** The Subsurface Color can be used as a mask or a weight for the subsurface effect. Black means there is no subsurface scattering. 
-	White means all lighting is entering the material and gets scattered around. A non grayscale value give more control over what
+	White means all lighting is entering the material and gets scattered around. A non-grayscale value give more control over what
 	color contributions are entering the surface resulting in a more complex looking shading. 
 
 	* **Falloff Color:** The Falloff Color defines the material scattering color once the light has entered the material. 
@@ -191,8 +191,8 @@ Setting the **Scatter Radius** of the Subsurface profile to a really big number
 
 In the image on the Left, the Scattering Radius has been set to 50. 
 Notice how it looks like there are multiple images rendered on the surface. 
-Now take a look at the image on the RIght. In this image the Scatter Radius has been set to 5.0. 
-Notice how the image looks softer and more natural, this is the type of effect we are trying to achieve.
+Now take a look at the image on the Right. In this image the Scatter Radius has been set to 5.0. 
+Notice how the image looks softer and more natural. This is the type of effect we are trying to achieve.
 
 ### Combining a Material Instance with a Subsurface Profile
 Combining a Material Instance with a Subsurface profile is a great way to get the most
