#UE4doc #docs: changed "lets" to "let's"

diff --git a/Engine/Documentation/Source/Engine/Rendering/Materials/HowTo/Subsurface_Scattering/Subsurface_Scattering.INT.udn b/Engine/Documentation/Source/Engine/Rendering/Materials/HowTo/Subsurface_Scattering/Subsurface_Scattering.INT.udn
--- a/Engine/Documentation/Source/Engine/Rendering/Materials/HowTo/Subsurface_Scattering/Subsurface_Scattering.INT.udn
+++ b/Engine/Documentation/Source/Engine/Rendering/Materials/HowTo/Subsurface_Scattering/Subsurface_Scattering.INT.udn
@@ -36,7 +36,7 @@ When completed, your **Content Browser** should look like this.
 
 ## Setting up a Subsurface Material
 
-Continuing with our Material from above that has Subsurface Scattering enabled, lets setup a basic Material so that we can see the Subsurface scattering in action in a level. 
+Continuing with our Material from above that has Subsurface Scattering enabled, let's setup a basic Material so that we can see the Subsurface scattering in action in a level. 
 	
 1. We need to place down some Material Expression nodes so that we have something to work with. 
    For this example, we are going to be adding the following nodes. 
@@ -102,7 +102,7 @@ Continuing with our Material from above that has Subsurface Scattering enabled,
 	**StaticMesh'/Game/StarterContent/Props/SM_Chair.SM_Chair'**
 
 	[region:tip]
-	If you can not find the SM_Chair mesh, it is probably because you did not include the Starter Content with your project. 
+	If you cannot find the SM_Chair mesh, it is probably because you did not include the Starter Content with your project. 
 	To find more about the Starter Content, please check out the [Starter Content](Engine/Content/Packs) documentation page.
 	[/region]
 
@@ -128,7 +128,7 @@ Notice as the numbers go from 0 to 1, the amount of light we see passing through
 ![](SSM_Different_Opacity_Settings.png)
 
 [region:tip]
-While the Opacity does help to get rid of a lot of the scattered light, you might notice that there is still some Subsurface scattering happening. To get completely rid of the effect, you are also going to have to adjust the **Value** of the Subsurface Color. More on this in the Subsurface Color Value section below. 
+While the Opacity does help to get rid of a lot of the scattered light, you might notice that there is still some Subsurface scattering happening. To get completely rid of the effect, you are also going to have to adjust the **Value** of the Subsurface Color (More on this in the Subsurface Color Value section below.)
 [/region]
 
 ### Subsurface Color Value
