#UE4doc #docs: corrected sentence tense shift (x3), duplicate words "of" "as"

diff --git a/Engine/Documentation/Source/Engine/Rendering/ParticleSystems/Reference/Modules/Location/LocationModules.INT.udn b/Engine/Documentation/Source/Engine/Rendering/ParticleSystems/Reference/Modules/Location/LocationModules.INT.udn
--- a/Engine/Documentation/Source/Engine/Rendering/ParticleSystems/Reference/Modules/Location/LocationModules.INT.udn
+++ b/Engine/Documentation/Source/Engine/Rendering/ParticleSystems/Reference/Modules/Location/LocationModules.INT.udn
@@ -20,7 +20,7 @@ The Initial Location module is used to set the initial location of a particle at
 |[REGION:tablesection]Location[/REGION]||
 |**Start Location** | A vector distribution indicating the location the particle should be emitted, relative to the emitter. Value is retrieved based on the EmitterTime at the spawn of the particle.|
 |**Distribute over NPoints**|When this is set to a non-zero number, it provides a number of evenly-spaced points along a uniform distribution from which to emit. For instance, if the min and max values of a distribution defined a line 100 units long, setting this value to 2 would cause particles to only be emitted from the actual min and max locations, nowhere in between.|
-|**Distribute Threshold**|When _Distribute over NPoints_ is set to any non-zero number, this provides a percentage for how many of the particles will emit from those evenly-spaced points. For example, a value of 1 will cause all particles to emit from those points. A value of 0.5 would cause half of the particles to emit from those points, while the other half were still spawned randomly.|
+|**Distribute Threshold**|When _Distribute over NPoints_ is set to any non-zero number, this provides a percentage for how many of the particles will emit from those evenly-spaced points. For example, a value of 1 will cause all particles to emit from those points. A value of 0.5 would cause half of the particles to emit from those points, while the other half would still spawn randomly.|
 
 In Spawn, the module retrieves the appropriate values from the distribution using the current emitter time. If the particle emitter does not have the _Use Local Space_ flag, this value is transformed into world-space. The Particle.Location field then has this value added to it.
 
@@ -36,7 +36,7 @@ The Initial Loc (Seeded) module is identical to the [Initial Location](#InitialL
 |[REGION:tablesection]Location[/REGION]||
 |**Start Location** | A vector distribution indicating the location the particle should be emitted, relative to the emitter. Value is retrieved based on the EmitterTime at the spawn of the particle.|
 |**Distribute over NPoints**|When this is set to a non-zero number, it provides a number of evenly-spaced points along a uniform distribution from which to emit. For instance, if the min and max values of a distribution defined a line 100 units long, setting this value to 2 would cause particles to only be emitted from the actual min and max locations, nowhere in between.|
-|**Distribute Threshold**|When _Distribute over NPoints_ is set to any non-zero number, this provides a percentage for how many of the particles will emit from those evenly-spaced points. For example, a value of 1 will cause all particles to emit from those points. A value of 0.5 would cause half of the particles to emit from those points, while the other half were still spawned randomly.|
+|**Distribute Threshold**|When _Distribute over NPoints_ is set to any non-zero number, this provides a percentage for how many of the particles will emit from those evenly-spaced points. For example, a value of 1 will cause all particles to emit from those points. A value of 0.5 would cause half of the particles to emit from those points, while the other half would still spawn randomly.|
 
 
 
@@ -53,7 +53,7 @@ The World Offset module is used to offset the initial location of a particle. Th
 |[REGION:tablesection]Location[/REGION]||
 |**Start Location** | A vector distribution indicating the world-space offset the particle should use. Value is retrieved based on the EmitterTime at the spawn of the particle.|
 |**Distribute over NPoints**|When this is set to a non-zero number, it provides a number of evenly-spaced points along a uniform distribution from which to emit. For instance, if the min and max values of a distribution defined a line 100 units long, setting this value to 2 would cause particles to only be emitted from the actual min and max locations, nowhere in between.|
-|**Distribute Threshold**|When _Distribute over NPoints_ is set to any non-zero number, this provides a percentage for how many of the particles will emit from those evenly-spaced points. For example, a value of 1 will cause all particles to emit from those points. A value of 0.5 would cause half of the particles to emit from those points, while the other half were still spawned randomly.|
+|**Distribute Threshold**|When _Distribute over NPoints_ is set to any non-zero number, this provides a percentage for how many of the particles will emit from those evenly-spaced points. For example, a value of 1 will cause all particles to emit from those points. A value of 0.5 would cause half of the particles to emit from those points, while the other half would still spawn randomly.|
 
 
 
@@ -284,7 +284,7 @@ The Triangle module allows the user to specify the shape and location of a trian
 |**Start Offset**|A vector distribution that defines an offset to the location of the triangle.|
 |**Height**|A float distribution that defines the length of the legs of the triangle.|
 |**Angle**|A float distribution that provides the angle between the primary leg and the hypotenuse. |
-|**Thickness**|A float distribution that provides the thickness of of the triangle, turning it into a 3D wedge. Particles will be emitted from within the 3-dimensional area.|
+|**Thickness**|A float distribution that provides the thickness of the triangle, turning it into a 3D wedge. Particles will be emitted from within the 3-dimensional area.|
 
 
 ### Skeletal Mesh Vertex/Surface Location (Skel Vert/Surf Loc)
@@ -316,7 +316,7 @@ The Skel Vert/Surf Loc module allows for particles to be emitted from a skeletal
 [REGION:simpletable]
 |**Property**|**Description**|
 | --- | --- |
-|_(Numbered Entry)_|Specifies the name of the bone on the SkeletalMesh to which a vertex or surface must be weighted in order to serve as as the source for the particles.|
+|_(Numbered Entry)_|Specifies the name of the bone on the SkeletalMesh to which a vertex or surface must be weighted in order to serve as the source for the particles.|
 [/REGION]
 [/EXCERPT:VertSourceLoc]
 
