#UE4doc #docs: sp: "parameter"

diff --git a/Engine/Documentation/Source/Engine/Rendering/ParticleSystems/Reference/Modules/Parameter/ParameterModules.INT.udn b/Engine/Documentation/Source/Engine/Rendering/ParticleSystems/Reference/Modules/Parameter/ParameterModules.INT.udn
--- a/Engine/Documentation/Source/Engine/Rendering/ParticleSystems/Reference/Modules/Parameter/ParameterModules.INT.udn
+++ b/Engine/Documentation/Source/Engine/Rendering/ParticleSystems/Reference/Modules/Parameter/ParameterModules.INT.udn
@@ -30,7 +30,7 @@ The Dynamic module is used to pass four scalar values to the material(s) used by
 | --- | --- |
 |**Param Name** | The name of the associated parameter in the DynamicParameter expression of the material. This property is read-only and populated automatically when a Particle Dynamic material expression exists within the material applied to this emitter.|
 |**Use Emitter Time** | If _true_, use the emitter time to retrieve the distribution value for the parameter. Otherwise, use the particle-relative time.|
-|**Spawn Time Only** | If _true_, only set the paramater value when the particle is spawned. Otherwise, the value is updated each frame.|
+|**Spawn Time Only** | If _true_, only set the parameter value when the particle is spawned. Otherwise, the value is updated each frame.|
 |**Value Method** | Specifies the method used to get the parameter value. It includes the following options:[INCLUDE:#ValueMethod]|
 |**Scale Velocity By Param Value** | If _true_, the velocity value passed to the material will be scaled by the value set in the _Param Value_ distribution.|
 |**Param Value** | A float distribution for setting the parameter value to use with the _EDPV_UserSet_ method.|
