#UE4doc #docs: sp: "one"

diff --git a/Engine/Documentation/Source/Engine/Rendering/ParticleSystems/Reference/Modules/Required/Required.INT.udn b/Engine/Documentation/Source/Engine/Rendering/ParticleSystems/Reference/Modules/Required/Required.INT.udn
--- a/Engine/Documentation/Source/Engine/Rendering/ParticleSystems/Reference/Modules/Required/Required.INT.udn
+++ b/Engine/Documentation/Source/Engine/Rendering/ParticleSystems/Reference/Modules/Required/Required.INT.udn
@@ -40,7 +40,7 @@ Every particle emitter contains this module which contains all the properties th
 | **Delay First Loop Only** | When _true_, and emitter with an EmitterDelay of greater than 0 and EmitterLoops greater than 1 will only perform the delay on the first loop. |
 | [REGION:tablesection]Sub UV[/REGION] ||
 | [REGION:note]The Sub UV data indicates that the emitter should utilize sub-images of the texture applied. This is useful for implementing simple animations on particles.[/REGION] ||
-| **Interpolation Method** | An enumeration indicating the method that should be utilized to interpolate between sub-images. This can be on of the following: [INCLUDE:#InterpMethod] |
+| **Interpolation Method** | An enumeration indicating the method that should be utilized to interpolate between sub-images. This can be one of the following: [INCLUDE:#InterpMethod] |
 | **Sub Images_Horizontal** | The number of sub-images in the texture on the horizontal (x) axis of the texture. |
 | **Sub Images_Vertical** | The number of sub-images in the texture on the vertical (y) axis of the texture. |
 | **Scale UV** | Indicates that the UV values should be scaled to properly `fit' the sub-image size. This is only used for mesh emitters. |
