Jira UEDOC-1412, 1117
UMG Quick Start Updates (will probably need additional updates).

diff --git a/Engine/Documentation/Source/Engine/UMG/QuickStart/1/UMGQS_1.INT.udn b/Engine/Documentation/Source/Engine/UMG/QuickStart/1/UMGQS_1.INT.udn
--- a/Engine/Documentation/Source/Engine/UMG/QuickStart/1/UMGQS_1.INT.udn
+++ b/Engine/Documentation/Source/Engine/UMG/QuickStart/1/UMGQS_1.INT.udn
@@ -1,7 +1,7 @@
 Availability: Public
-Title:1. Create a New Project
+Title:1. Required Project Setup
 Crumbs:%ROOT%, Engine, Engine/UMG, Engine/UMG/QuickStart
-Description:The first step in the UMG Quick Start is to create a New Project.
+Description:Here we will create a Widget Blueprint used for the HUD and set our character up to pass information to the HUD. 
 
 
 [VAR:Steps]
@@ -21,25 +21,119 @@ Description:The first step in the UMG Quick Start is to create a New Project.
 
 %Steps%
 
+[REGION:note]
+For this tutorial we are using the **Blueprint First Person** template with **Starter Content** enabled. If you do not know how to create a new project and use a template, please see the [](Engine/Basics/Projects/Browser) page for more information. 
+[/REGION]
 
+1. Click the **Add New** button inside the **Content Browser**, then under _User Interface_ select **Widget Blueprint** and name it **HUD**.
 
+	![](UMGQS1.png)
 
-1. From the Unreal Engine Launcher, click the **Launch Unreal Engine** Button. 
+	All User Interface elements (HUD, Menus, etc.) will be created and housed inside a **Widget Blueprint**. The Widget Blueprint allows you to visually layout your UI elements as well as provide scripted functionality for those elements. We are creating this now so that we can tell it to be displayed when our player character spawns in the world, however we will set it up later. 
+	
+	[REGION:note]
+	For more information on **Widget Blueprints** see the [Widget Blueprint](Engine/UMG/UserGuide/WidgetBlueprints) documentation. 
+	[/REGION]
 
-1. The [Unreal Project Browser](Engine/Basics/Projects/Browser) will appear. Click on the **New Project** tab.
+1. Create two more **Widget Blueprints**, one called **MainMenu** and another called **PauseMenu**. 
 
-1. For this guide, we will be using the **Blueprint First Person** template. On the **Blueprint** tab, select **First Person** (1 below). 
+	![](AllWidgets.png)
 
-	![](UMGQS_1.png)(w:600)
-	
-1. Make sure that **Include starter content** is enabled (2), enter a project name in the **Name** field (3), then select **Create Project** (4).
+1. **Right-click** in the **Content Browser** and create a new **Level** called **Main**. 
 
-1. The new project will now open up in the Unreal Editor. 
+	![](CreateLevel.png)
 
-[REGION: note]
-If this is your first time running the editor, you will be greeted with the Unreal Editor 4 **Welcome Tutorial**. This is an interface walkthrough and is a really good way to catch your bearings and learn how to navigate the editor.
-[/REGION]
+	We will use this later in the guide for our Main Menu setup. 
+
+1. In the **Content Browser**, open the **FirstPersonCharacter** Blueprint located under the **Content/FirstPersonBP/Blueprints** folder. 
+
+	![](UMGQS2.png)
+
+	This is the playable character's Blueprint in which we will create some information to pass to our **HUD** Widget Blueprint for display. 
+
+1.	In the **My Blueprint** window, click the **Add Variable** Button. 
+
+	![](UMGQS3.png)
+
+1.	In the **Details** panel for the new variable, name it **Health**, change it to a **Float** variable type and set the **Default Value** to **1.0**. 
+
+	![](UMGQS4.png)
+
+	This variable will represent the player character's Health that we will display on the HUD. 
+
+1.	Create another **Float** variable called **Energy** with a **Default Value** of **1.0**. 
+
+	![](UMGQS5.png)
+
+	This will represent the player character's Energy that we will display on the HUD. 
+
+	[REGION:note]
+	Unable to enter a Default Value? Click the **Compile** button from the Tool Bar to compile the Blueprint then try again. 
+	[/REGION]
+
+1.	Create another variable of the **Integer** type called **Ammo** with a **Default Value** of **25**. 
+
+	![](UMGQS6.png)
+
+1.	Create one more **Integer** variable called **MaxAmmo** also with a **Default Value** of **25**. 
+
+1.	Inside the **Graph** window, **Right-click** and add an **Event Begin Play** node. 
+
+1.	Drag off the **Event Begin play** then add a **Create Widget** node with **Class** set to your **HUD** Widget Blueprint. 
+
+	![](UMGQS7.png)
+
+1.	Off the **Return Value** of the **Create HUD_C Widget**, select **Promote to Variable** and name **HUD Reference**. 
+
+	![](UMGQS8.png)
 
+	This will create our HUD Widget Blueprint when the game is started and store it as a variable that we can access later. This is useful for calling functions or setting properties of the HUD later on, for example if you wanted to hide the HUD while the game is paused you can access the HUD through this variable. 
 
+1.	Drag off the out pin of the **Set** node and add the **Add to Viewport** node. 
+
+	![](UMGQS9.png)
+
+	As the name suggests, this will add the Widget Blueprint specified as the target to the player's viewport, drawing it onscreen. 
+
+## Adjusting Character Variables
+
+The last thing we will do before we start generating our actual HUD is provide a way for our character variables to change. 
+
+1. Inside the **FirstPersonCharacter** Blueprint, while holding **Alt** drag in the **Energy** variable and place it next to the Jump script. 
+
+	![](Step2_1.png)(w:720)
+
+1. Hold **Ctrl** and drag in a copy of the **Energy** variable and connect it to a **Float - Float** node set to **0.25** and connect as shown. 
+
+	![](Step2_2.png)
+
+	Each time the character jumps, 0.25 will be subtracted from the character's current Energy value. 
+
+1. Set up the same script for the **Health** variable but use an **F** Key Event (or any other key press event) for testing.
+
+	![](Step2_3.png)
+
+	This will allow us to test that our Health variable is properly displayed through our HUD whenever a key is pressed. 
+
+1.	Find the **Spawn Projectile** script and following the **InputAction Fire** event, **Right-click** and add a **Branch** node. 
+
+1.	For the **Branch** condition, **Ctrl** drag in the **Ammo** variable connected to a **Integer > Integer** variable set to **0** and connect as shown.
+
+	![](Step2_4.png)(w:820)
+
+	This will only allow the player to fire a projectile if their ammo is greater than zero. 
+
+1.	At the end of the **Spawn Projectile** script following the **Play Sound at Location** node, set **Ammo** to equal **Ammo - 1**. 
+
+	![](Step2_5.png)
+
+	This will subtract 1 ammo each time the player fires a projectile, giving us something to display on our HUD. 
+
+1.	**Compile** and **Save** then close the Blueprint. 
+
+
+In the next step we will set up the displaying of our Health, Energy and Ammo variables on our HUD and get them working in-game. 
 
 %Steps%
+
+
diff --git a/Engine/Documentation/Source/Engine/UMG/QuickStart/2/UMGQS_2.INT.udn b/Engine/Documentation/Source/Engine/UMG/QuickStart/2/UMGQS_2.INT.udn
--- a/Engine/Documentation/Source/Engine/UMG/QuickStart/2/UMGQS_2.INT.udn
+++ b/Engine/Documentation/Source/Engine/UMG/QuickStart/2/UMGQS_2.INT.udn
@@ -1,7 +1,7 @@
 Availability: Public
-Title:2. Creating a Widget Blueprint
+Title:2. Displaying Health, Energy and Ammo
 Crumbs:%ROOT%, Engine, Engine/UMG, Engine/UMG/QuickStart
-Description:The second step in the UMG Quick Start is to create a Widget Blueprint used for our HUD.
+Description:In this step we will hook up and display the Health, Energy and Ammo of the player on our HUD. 
 
 
 [VAR:Steps]
@@ -22,34 +22,155 @@ Description:The second step in the UMG Quick Start is to create a Widget Bluepri
 
 %Steps%
 
+## Visual: Health, Energy and Ammo
 
+Here we will begin setting up the visual aspects of our HUD. 
 
-[OBJECT:TopicList]
-	[PARAM:icon]
-	[/PARAM]
-	[PARAM:title]
-	[/PARAM]
-	[PARAM:description]
+1.	Open your **HUD** Widget Blueprint to access the **Widget Blueprint Editor**. 
 
-This step involves creating a [Widget Blueprint](Engine/UMG/UserGuide/WidgetBlueprints) which will be used to draw the HUD on screen. 
+	![](Step2_6.png)(w:640)
 
-1. With project open, inside the **Game** folder of the **Content Browser**, click the **New** button and select **New Folder** and call it **UMG**. 
+	Inside the Widget Blueprint Editor we can provide the visual layout and scripted functionality of our HUD. 
+	 
+	[REGION:note]
+	For more information on each of the different facets of the Widget Blueprint Editor, see [](Engine/UMG/UserGuide/WidgetBlueprints). 
+	[/REGION]
 
-1. Open the **UMG** folder, then **Right-click** and select **User Interface** then **Widget Blueprint** and name it **GameHUD**.  
+1.	Inside the **Palette** window under **Panel**, drag a **Horizontal Box** onto the **CanvasPanel** in the **Hierarchy** window. 
 
-		![](New_3.png)
+	![](Step2_7.png)
 
-		No matter what type of HUD elements you wish to create in **UMG**, you will always first need to create a **Widget Blueprint**. 
+	**Panel** Widgets are sort of like containers for other widgets and provide additional functionality for the widgets that are inside them.
+	
+1.	Also under **Panel**, drag two **Vertical Boxes** onto the **Horizontal Box**. 
 
-In the next section, we will begin adding elements to the **Widget Blueprint** that we want to display on screen. 
+	![](Step2_8.png)
 
-	[/PARAM]
-	[PARAM:links]
-	[/PARAM]
-[/OBJECT]
+1.	Under **Common**, drag two **Text** widgets onto the first Vertical Box and two **Progress Bars** onto the second Vertical Box. 
+
+	![](Step2_9.png)
+
+1.	Select the **Horizontal Box**, then in the graph resize the box and position it in the upper left corner of the window. 
+
+	![](Step2_10.png)
+
+	The progress bars will appear very small, but don't worry we will fix that next. 
+
+1.	Select both **Progress Bars**, then in the **Details** panel set both to **Fill** for the **Size**. 
+
+	![](Step2_11.png)
+
+1.	Select the **Vertical Box** containing the Progress Bars and set it to **Fill** as well. 
+
+	![](Step2_12.png)
+
+1.	Select the **Horizontal Box** again and resize it so the bars line up with the text. 
+
+	![](Step2_13.png)
+
+1.	Select the top-most **Text** widget in the **Hierarchy** window then in the **Details** panel under **Content** enter **Health :**. 
+
+	![](Step2_14.png)
+
+	Do the same for the other **Text** widget but label it as **Energy** so your graph looks like below. 
+
+	![](Step2_15.png)
+
+1.	Select the **Progress Bar** next to health and in the **Details** panel set the **Fill Color and Opacity** to a green color. 
+
+	![](Step2_24.png)(w:800)
+
+	You may notice that Progress Bar does not change color when assigning a color. This is because the **Percent** value to fill the bar is set to 0.0 (you can change this to test out different colors, later we will link this to the health value of our character). 
+
+1.	Also set a Fill Color for the **Energy** bar (we selected an Orange color). 
+		
+1.	For the display of Ammo, using the methods from above, add widgets to the **Hierarchy** so your setup is laid out like below. 
+
+	![](Step2_16.png)
+
+1.	With the **Horizontal Box** containing the ammo information selected, resize it and position it in upper right corner.
+
+	![](Step2_17.png)
+
+1.	With the **Horizontal Box** still selected, in the **Details** panel click **Anchors** and select the upper-right Anchor. 
+
+	![](Step2_18.png)
+
+	This will move the Anchor Medallion to the upper-right corner of your screen. 
+
+	![](Step2_19.png)
+
+	By **Anchoring** a Widget, you are indicating to the Widget that it is to appear in its placed position relative to the position of the Anchor Medallion regardless of screen size. In other words, when the screen size changes, move the widget the same distance away from the Anchor Medallion's position (in this example the same distance away from the upper-right corner of the screen). 
+
+	[REGION:note]
+	You can test different screen sizes by clicking and changing the **Preview Size** option inside the graph. 
+	[/REGION]
 
+## Script: Health, Energy and Ammo
+
+Now that we have our visual layout, we can add the hooks to provide the functionality behind our HUD elements. 
+
+1.	Click the **Graph** button in the upper-right corner of the Widget Blueprint Editor window. 
+
+	![](Step2_20.png)
+
+1.	In the **Graph** under the **Event Construct**, **Right-click** and add a **Get Player Character** node then **Cast To FirstPersonCharacter**. 
+
+	![](Step2_21.png)
+
+	Here we are checking if the Character Blueprint being used by the Player Character is indeed the **First Person Character** Blueprint. 
+
+1.	Drag off the **As First Person Character** pin and select **Promote to Variable** (name it **My Character**) then connect as shown. 
+
+	![](Step2_22.png)
+
+	This will allow us to access the variables contained inside our **FirstPersonCharacter** Blueprint. 
+
+1.	Return to the **Designer** tab and select the **Progress Bar** next to **Health**. 
+
+1.	In the **Details** panel under **Appearance**, click the **Bind** option next to **Percent** and set it to use **Health** from **MyCharacter**. 
+
+	![](Step2_23.png)
+
+	This binds the value of the Progress Bar to our **Health** variable inside our **First Person Character** Blueprint. Now whenever our Health value changes in our Character Blueprint it will automatically be updated in our HUD as well. You may also notice that the Progress Bar does not change after binding it, that is because we've stated this to occur upon **Event Construct** (or when the HUD is constructed, which occurs when the game is started).
+
+1.	Select the **Progress Bar** next to Energy and repeat the process above, binding **Percent** to **Energy** from **MyCharacter**. 
+
+	![](Step2_25.png)
+
+1.	Select the 25 following the Ammo text, then in the **Details** panel for **Text**, click **Bind** and **Create Binding**. 
+
+	![](Step2_26.png)
+
+	Previously we used **Sub-Object Properties** to bind properties of the same type, we can also create our own custom bindings. Here we are going to bind our **Text** property to an **Integer** property from our **First Person Character** Blueprint giving us a display of our current Ammo.  
+
+1.	Inside the function that is created and automatically opened, hold **Ctrl** and drag the **MyCharacter** variable into the graph. 
+
+	![](Step2_27.png)(w:820)
+
+1.	Off **MyCharacter**, get **Ammo** and connect it to the **Return Value** of the **Return Node**. 
+
+	![](Step2_28.png)
+
+	A conversion node **To Text** will automatically be created for you when connecting. 
+
+1.	Repeat the process above for the other 25 text and creating a binding for the "Max Ammo" text. 
+
+	![](Step2_29.png)
+
+	You probably don't need to do this if your Max Ammo is going to remain the same, however this will allow you to change it. 
+
+1.	Click **Compile** and **Save** then click the **Play** button to play in the editor. 
+
+	![](Example_GIF.gif)(convert:false)
+
+	Our Health, Energy and Ammo values are now displayed on our HUD and reflect our current values from our Character Blueprint. Pressing **Space** will have the character jump and deplete Energy, pressing the **Left Mouse Button** will fire the weapon reducing our Ammo, and pressing **F** will cause us to lose Health (which you could hook up to a damage system at a later date). 
+	
+In the next section we will create a Main Menu which we can use to load into our game that we've setup. 
 
 %Steps%
 
 
 
+
+
diff --git a/Engine/Documentation/Source/Engine/UMG/QuickStart/3/UMGQS_3.INT.udn b/Engine/Documentation/Source/Engine/UMG/QuickStart/3/UMGQS_3.INT.udn
--- a/Engine/Documentation/Source/Engine/UMG/QuickStart/3/UMGQS_3.INT.udn
+++ b/Engine/Documentation/Source/Engine/UMG/QuickStart/3/UMGQS_3.INT.udn
@@ -1,7 +1,7 @@
 Availability: Public
-Title:3. Setting Up the Character Variables
+Title:3. Creating a Main Menu
 Crumbs:%ROOT%, Engine, Engine/UMG, Engine/UMG/QuickStart
-Description:The third step in the UMG Quick Start is to setup our character and variables that will be displayed through the HUD. 
+Description:In this step we will create the visual layout for our Main Menu. 
 
 
 [VAR:Steps]
@@ -24,74 +24,95 @@ Description:The third step in the UMG Quick Start is to setup our character and
 
 ---
 
+First bit of business is to create the layout of our Main Menu which we will do below. 
 
-[OBJECT:TopicList]
-	[PARAM:icon]
-	[/PARAM]
-	[PARAM:title]
-	[/PARAM]
-	[PARAM:description]
+1.	Open your **MainMenu** Widget Blueprint. 
 
-In the following steps, we will add some variables to our character that we want to display on our HUD. 
+1.	On the **Designer** tab, drag and drop widgets from the **Palette** onto the **Hierarchy** to achieve the setup below. 
 
-1. In the **Game/Blueprints** folder, open the **MyCharacter** Blueprint. 
+	![](Hierarchy1.png)
 
-1. Inside the **MyBlueprint** window, create two **Float** variables (**P_Health**, **P_Energy**). 
+	For reference we have an **Image**, a **Vertical Box** with **Text** and 3 **Buttons** that have **Text** on our **Canvas Panel**. 
 
-1. Click **Compile** then in the **Details** panel, set the **Default Value** for each **1.0**. 
+1.	Select the **Vertical Box** in the **Hierarchy** then **Right-click** and **Copy** it. 
 
-1. Also create two **Int** variables (**P_AmmoMax**, **P_CurrentAmmo**), **Compile** and set them both to **25**.
+1.	**Right-click** on the **Canvas Panel** and select **Paste** to paste a second copy of the **Vertical Box**. 
 
-		![](New_1.png)
+	![](Hierarchy2.png)
+	 
+1.	In the new **Vertical Box**, **Right-click** copy a **Button** and paste it to create a fourth Button/Text widget. 
 
-1. In the graph window, locate the "Jump" section, then while holding down the **Alt** key, drag in the **P_Energy** variable and place it following **Jump**. 
+	![](Hierarchy3.png)
 
-		![](New_2.png)(w:800)
+1.	Select the first **Vertical Box** and in the **Details** panel, rename it to **MainMenu**, check the **isVariable** option and set **ZOrder** to **1**. 
 
-		This creates a **Set** node which will allow us to set this variable following the jump, thus reducing energy.
-		
-1. While holding down the **Control** key, drag in the **P_Energy** variable again and connect it to a **Float - Float** node (set to **.25**), then connect as shown.
+	![](Hierarchy4.png)
 
-		![](New_3.png)
-		
-		This will **Get** the current energy value and subtract 0.25 from it each time jump is executed, which we will display on the HUD.
+	We are naming the widget for clarity and setting it to a variable so that we can access it as well as setting the **ZOrder** so that it appears on top of our image which we will set in a moment. 
 
-1. **Right-click** somewhere under the jump section, search for an add an **F** key event. 
+1.	Rename each of the **Buttons** as shown in the image below. 
 
-		![](New_4.png)
+	![](Hierarchy5.png)
 
-		We are going to simulate taking damage and subtract health whenever the **F** key is pressed.
-								
-1. **Alt** drag in **P_Health**, then **Control** drag in **P_Health** to **Set**/**Get** the variable similar to the method used to get energy above.
+	By updating the names of our buttons so that we know what each one does, this will make it easier to script functionality for them. 
 
-1. Connect the nodes as shown below to the **F** key **Pressed** pin. 
+1.	Update each of the **Text** widgets as shown below via the **Text** section under _Content_ in the **Details** panel. 
 
-		![](New_5.png)
+	![](Hierarchy6.png)
 
-		We added the **Random Float in Range** node to subtract a random value between **0.025** and **0.075** for variety in damage applied.
+	Here we are updating the text that is displayed on each of the buttons as well as the menu headers. 
 
-1. Find the "Spawn projectile" section and following the **InputActionFire** node, add a **Branch** node and connect the **True** to the **Montage Play** node.
+1.	Holding **Ctrl** select each of the **Buttons** then in the **Details** panel set a **Tint** color for **Hovered** and set each to **Fill** for **Size**. 
 
-1. For the **Branch** condition, **Control** drag in the **P_CurrentAmmo** variable and connect it to a **>** node (set to **0**), then connect to the **Condition**. 
+	![](Hierarchy7.png)
 
-		![](New_6.png)(w:800)
+	Here we are adjusting the size of all of our buttons as well as defining a color for the button to take when mousing over it. 
 
-		This will only allow the player to fire a projectile if their current ammo is greater than 0.
+1.	Select the **Image** in the **Hierarchy** then in the **Details** under _Appearance_ and _Brush_, select a **Texture** or **Material** to use. 
 
-1. Also in "Spawn projectile", at the end of the code following the **Play Sound** node, **Set** the **P_CurrentAmmo** to **P_CurrentAmmo - 1**. 
+	![](Hierarchy8.png)
 
-		![](New_7.png)
+	[REGION:note]
+	If you want to use the same one as this guide, you can download it from here: [Example Background](ExampleBackground.png)
+	
+	Drag-and-drop the Example Background into **Unreal Engine** to import it (click **Yes** at any confirmation messages). 
+	[/REGION]
 
-		This will subtract one from our ammo each time a shot is fired, which we will convey through our HUD.
+1.	In the **Designer** window, resize the image so that it takes up the entire layout. 
 
-1. Click the **Compile** and **Save** buttons, then close the Blueprint. 
+	![](Hierarchy9.png)
 
-	[/PARAM]
-	[PARAM:links]
-	[/PARAM]
-[/OBJECT]
+1.	In the **Details** panel for the **Image**, click the **Anchor** button and choose the fill screen option. 
+
+	![](Hierarchy10.png)
+
+1.	Select both **Vertical Boxes** and **Anchor** them to the left-center position. 
+
+	![](Hierarchy11.png)
+
+1.	Select the **Options** Vertical Box and set it’s **Visibility** setting to **Hidden**. 
+
+	![](Hierarchy12.png)
 
-In the next section, we will hook up the Ammo, Health, and Energy to our **Widget Blueprint** to display it on our in-game HUD. 
+	We do not want options to be shown by default, we will tell it to be visible through script. 
+
+1.	Select the **MainMenu** Vertical Box and resize and position it as desired in the **Designer** window. 
+
+	![](Hierarchy13.png)
+
+1.	You can adjust the **Font**, **Size** and **Alignment** from the **Details** panel of each **Text** widget. 
+
+	![](Hierarchy14.png)
+
+	You can experiment with these settings to get your desired effect. 
+
+1.	Setup your **Options** Vertical Box as desired then move it directly behind your **Main Menu** Vertical Box. 
+
+	![](Hierarchy15.png)
+
+	Since only one of these will be visible at a time, it is okay that they overlap one another. 
+
+With the visual layout of our menu complete, in the next step we will script the menu's functionality. 
 
 ---
 
@@ -100,3 +121,4 @@ In the next section, we will hook up the Ammo, Health, and Energy to our **Widge
 
 
 
+
diff --git a/Engine/Documentation/Source/Engine/UMG/QuickStart/4/UMGQS_4.INT.udn b/Engine/Documentation/Source/Engine/UMG/QuickStart/4/UMGQS_4.INT.udn
--- a/Engine/Documentation/Source/Engine/UMG/QuickStart/4/UMGQS_4.INT.udn
+++ b/Engine/Documentation/Source/Engine/UMG/QuickStart/4/UMGQS_4.INT.udn
@@ -1,7 +1,7 @@
 Availability: Public
-Title:4. Displaying Ammo, Health and Energy
+Title:4. Scripting the Main Menu
 Crumbs:%ROOT%, Engine, Engine/UMG, Engine/UMG/QuickStart
-Description:The fourth step in the UMG Quick Start is to hook up the Ammo, Health and Energy so they are displayed in-game. 
+Description:In this step we will script the functionality for our Main Menu and load into game. 
 
 
 [VAR:Steps]
@@ -23,215 +23,98 @@ Description:The fourth step in the UMG Quick Start is to hook up the Ammo, Healt
 %Steps%
 
 ---
-[OBJECT:TopicList]
-	[PARAM:icon]
-		![](%ROOT%/build_icon.png)(convert:false)
-	[/PARAM]
-	[PARAM:title]
-		Creating the HUD elements
-	[/PARAM]
-	[PARAM:description]
-
-In this section, we will add the elements we want to display on the HUD. 
-
-1. Inside the **UMG** folder that you created, open the **GameHUD** Blueprint to open it in the **Widget Blueprint Editor** (shown below). 
-	
-		![](New_1.png)(w:720)
-
-1. Inside the **Palette** window under **Panel**, drag a **Horizontal Box** onto the **CanvasPanel** in the **Hierarchy** window. 
-
-		![](UMG_Update1.png)
-
-1. Also under **Panel**, drag two **Vertical Boxes** onto the **Horizontal Box** in the **Hierarchy** window. 
-
-		![](UMG_Update2.png)
-
-1. Under **Common**, drag two **Text Blocks** onto the first **Vertical Box** and two **Progress Bars** onto the second **Vertical Box**. 
-
-		![](UMG_Update3.png)
-
-1. With the **Horizontal Box** selected, re-size and arrange the box so it appears in the upper left corner of the window. 
-
-		![](UMG_Update4.png)
-
-		The progress bars will appear very small, however we will fix that in the next step.
-
-1. Click on one of the **Progress Bars**, then in the **Details** panel, click the **Fill** button under **Layout** / **Size**. 
-
-		![](UMG_Update5.png)
-
-1. Click the other **Progress Bar**, set it to **Fill** under **Size** as performed in the previous step. 
-
-1. Click the **Vertical Box** containing the **Progress Bars**, and set its **Size** to **Fill** as well. 
-
-		![](UMG_Update6.png)
-
-1. Once you have done so, re-size the box in the **Visual Designer** as seen below. 
-
-		![](UMG_Update7.png)
-
-1. Click on the top **Text Block** and in the **Details** panel, change the text to **Health :** and align it to the right. 
-
-		![](UMG_Update8.png)
-
-1. Do the same for the other **Text Block** however, for its text, enter **Energy :**. 
-
-1. For reference, select the **Horizontal Box** and in the **Details** panel, rename it to **PlayerCondition**. 
-
-		![](UMG_Update9.png)
-
-1. In the **Palette**, drag another **Horizontal Box** onto the **CanvasPanel** in the **Hierarchy**. 
-
-1. Drag 4 **Text Blocks** onto the new **Horizontal Box** and for each, set the text in order as follows: **Ammo :**, **25**, **/** and **25**. 
 
-		![](UMG_Update10.png)
+Time to get our hands dirty and script the functionality of our Main Menu. At the end of this step we will be able to launch our game to the Main Menu, start our game as well as change the resolution inside the Options Menu. 
 
-1. Under **Appearance** and **Font**, you can change the font size for each text (we set each to 32). 
+1.	Open the **MainMenu** Widget Blueprint and click the **Graph** tab. 
 
-		![](UMG_Update11.png)
+	![](Script1.png)
 
-1. Drag the **Horizontal Box** with the Ammo Display to the upper right portion of the screen and resize as desired. 
+1.	In the **MyBlueprint** window, select the **PlayButton** then in the **Details** panel click the plus sign next to **On Clicked**. 
 
-		![](UMG_Update12.png)
+	![](Script2.png)
 
-1. With the **Horizontal Box** selected, in the **Details** panel, click the **Anchors** button and choose the upper right anchor. 
+	This will add a node to the graph that will execute whenever the button is clicked on. 
 
-		![](UMG_Update13.png)(w:800)
+1.	Add **On Clicked** Events for each of your buttons. 
 
-		By **Anchoring** a Widget, you are indicating to the Widget that it is to appear in the location specified or relative to the position based on the screen size. You can verify this by clicking the **Resolution** drop-down menu in the **Visual Designer** and changing aspect ratios. You should see that the Widgets maintain their relative positions. 
+	![](Buttons.png)
 
-	[/PARAM]
-[/OBJECT]
-
-[OBJECT:TopicList]
-	[PARAM:icon]
-		![](%ROOT%/build_icon.png)(convert:false)
-	[/PARAM]
-	[PARAM:title]
-		Retrieving the Variables
-	[/PARAM]
-	[PARAM:description]
-
-In this section, we will hook up the variables from our **MyCharacter** Blueprint so they are passed to the HUD. 
-
-1. Click on the **Progress Bar** which represents Health. 
+	We have arranged our Buttons in the graph to mimic the layout of our menu for clarity. 
 
-1. In the **Details** panel under **Appearance**, click on **Bind** next to **Percent** option, then **Create Binding**. 
+1.	Off the **OnClicked(PlayButton)**, add a **OpenLevel** (set to **FirstPersonExampleMap**) node and a **Remove from Parent** node. 
 
-		![](UMG_Update14.png)
+	![](PlayButton.png)
 
-		Here we are going to Bind the Percent value of the Progress Bar to a value in our Character Blueprint. 
+	The **Level Name** is where you indicate the name of the level you want to load (in our case the First Person Example map). The **Remove from Parent** node will remove the _target_ Widget Blueprint from the viewport. It is set to target itself (which is the Main Menu Widget Blueprint that we want to remove from view after the level has loaded). 
 
-1. On the graph that opens, **Right-click** and search for and add a **Get Player Character** node. 
+1.	Hold **Ctrl** and drag in the **MainMenu** and **OptionsMenu** variables. 
 
-1. Off the **Get Player Character** add the **Cast to MyCharacter** node and connect as shown. 
+1.	Drag off **MainMenu** and **Set Visibility** (set to **Hidden**) then **Set Visibility** off **OptionsMenu** to **Visible**.  
 
-		![](UMG_Update15.png)
+	![](OptionsButton.png)(w:820)
 
-1. Off the **As MyCharacter C** pin, search for **Get P Health** and connect it as show below. 
+	Here we are turning off the Main Menu when the Options Button is clicked and turning on the Options Menu. 
 
-		![](UMG_Update16.png)(w:820)
+1.	Off each of the **Settings** Buttons add an **Execute Console Command** node. 
 
-		This will always get the value associated with our player's Health and pass it to the HUD automatically. 
+1.	For the **Command** for each, use **r.setRes XxY** where X and Y are the dimensions to use. 
 
-1. On the **Designer** tab, click the Energy Progress Bar and create a Binding for the Percent as done with the Health. 
+	![](Resolution.png)
 
-1. On the graph that opens, **Get Player Character** and **Cast To MyCharacter**, then get P Energy and connect as shown. 
+	Above we are using the following settings: **640x480**, **1280x720** and **1920x1080**. 
 
-		![](UMG_Update17.png)(w:820)
+1.	For the **Return Button**, repeat step 6 but reverse the settings so Main Menu is visible and Options Menu is hidden. 
 
-1. Return to the **Designer** tab and click on the first number in the Ammo display, then in the **Details** for the **Text**, **Create Binding**. 
+	![](ReturnButton.png)(w:820)
 
-		![](UMG_Update18.png)
+1.	Off the **OnClicked (QuitButton)** add an **Execute Console Command** set to the command **Quit**. 
 
-1. In the graph that opens, **Get Player Character** and **Cast To MyCharacter**, then get P Current Ammo and connect as shown. 
+	![](QuitButton.png)
 
-		![](UMG_Update19.png)(w:820)
+1.	**Compile** and **Save** then close the Widget Blueprint. 
 
-		When connecting the Ammo to the ReturnNode's Return Value, a conversion node will automatically be created. 
+1.	In the **Content Browser**, open the **Main** Level then open the **Level Blueprint**. 
 
-1. While still on the **Graph** tab, create a new **Integer** variable from the **MyBlueprint** window. 
+	![](OpenLevelBlueprint.png)
 
-1. In the **Details** panel for the variable, name it **H_MaxAmmo** and check the **Editable** and **Expose on Spawn** checkboxes. 
+1.	**Right-click** in the graph and add an **Event Begin Play** node and connect a **Create Widget** (set **Class** to **Main Menu**). 
 
-		![](UMG_Update20.png)
+	![](EventBeginPlay.png)
 
-		Since the Max Ammo Count value will not change for this example, we are going to initialize it when the HUD is created. 
+1.	Connect an **Add to Viewport** node and add a **Get Player Controller** node. 
 
-1. Return to the **Designer** tab. 
+1.	Off **Get Player Controller**, add the **Set Show Mouse Cursor** node set to **True**, then **Compile** and close the Blueprint. 
 
-1. Click the second number in the Ammo display, then in the **Details** for the **Text** click **Bind** then **Create Binding**. 
+	![](FinishedLevel.png)(w:820)
 
-1. When the graph opens, drag the **H_MaxAmmo** variable into the graph onto the **Return Value**. 
+1.	Open the **First Person Character** Blueprint inside the **Content/FirstPersonBP/Blueprints** folder. 
 
-		![](UMG_Update21.png)
+1.	Following the **Event Begin Play** script, add a **Get Player Controller** then **Set Input Mode Game Only** node. 
 
-		This will create a conversion node and connect it to the Return Value in pin. 
+	![](PlayerCharacterBlueprint.png)(w:820)
 
-1. Return to the **Designer** tab and click on the Health's progress bar again. 
+	This will turn off the cursor and set the input for the player to Game Only. 
 
-1. Under **Appearance** in the **Details** panel, select any color you wish to use as your Health Bar's color. 
+1.	**Compile** and close the Blueprint, then open the **World Settings**. 
 
-		![](UMG_Update22.png)
+	![](WorldSettings.png)
 
-1. Do the same for the **Energy** progress bar and set your Energy Bar's color. 
+1.	Under _Game Mode_, set the **GameMode Override** to **FirstPersonGameMode** and change the **Default Pawn Class** to **Character**. 
 
-1. You can verify the color of the bar by setting the percent value to 1.0 (just remember to set it back to 0.0 after selecting a color). 
+	![](Overrides.png)
 
-		![](UMG_Update23.png)
+	Here we are assigning a Game Mode for this Map and changing the player's character to use the Engine Default Character inside of the playable **FirstPersonCharacter** Blueprint as we do not what the player to be able to run around and shoot behind the Main Menu. 
 
-1. **Compile** and **Save** then close the **GameHUD** Blueprint. 
+1.	Click the **Play** Button from the Main Toolbar to play in the editor. 
 
-[REGION:note]
-Note: If you are not getting the colors you desire, there is an option under **Style** you may need to change. 
-
-![](UMG_Update24.png)
-
-As of Engine Version 4.5, the **Fill Image** defaults to Orange. By assigning a plain White texture as the Fill Image, you should fix any issues you have with incorrect colors appearing. This particular issue will be addressed in a future version update. 
-
-[/REGION]  
-	[/PARAM]
-[/OBJECT]
-		
-[OBJECT:TopicList]
-	[PARAM:icon]
-		![](%ROOT%/build_icon.png)(convert:false)
-	[/PARAM]
-	[PARAM:title]
-		Displaying the HUD
-	[/PARAM]
-	[PARAM:description]
+	[![](MenuImage.png)(w:720)](ExampleMainMenu.mp4)
 
-In this final section, we will call the HUD to be displayed from **MyCharacter** and assign the default **MaxAmmo** to be displayed. 
+	_Click the image above to play the video._
 
-1. Inside the **Game/Blueprints** folder, open the **MyCharacter** Blueprint. 
+	We have now created a Main Menu that is functional and allows the player to quit the game, change resolution options or play the game (which loads up the specified level and transitions the player into gameplay). 
 
-1. **Right-click** anywhere in the graph and add an **Event Begin Play** node. 
-
-1.	Drag off the out pin and add a **Create Widget** node, and set the **Class** to **GameHUD_C**. 
-
-		![](UMG_Update25.png)
-
-1.	Drag in the P Ammo Max variable and connect it to the H Max Ammo in pin. 
-
-1.	Drag off the **Return Value** of the Create Widget node and **Promote to Variable** called **GameHUDWidget**. 
-
-1.	Drag off the **Return Value** again and add the **Add to Viewport** node. 
-
-		![](UMG_Update26.png)
-
-		We are creating a **Variable** for the Game HUD Widget so that we can access it later. 
-		
-1. **Compile** and **Save** then play in the editor.
-
-		![](GIF1.gif)(convert:false)
-
-		You should see that when you press **F** you are "taking damage" and losing Health, jumping will deplete Energy and ammo count represents how many shots you have left (where you are unable to fire when you have 0 shots left).
-
-	[/PARAM]
-[/OBJECT]
-		
-In the next section, we will add an "Experience Bar" where shooting blocks increases your experience (you would probably want this to be tied to kill counts or some other meaningful gameplay situation; however this will show how you can set it up). When you "Level Up", you will unlock the ability to use an "Action Bar" button ability which will allow the character to perform a super jump. 
+In the next step we will add an in-game pause menu so that the player can pause gameplay or quit the game and return to this Main Menu. 
 
 ---
 
@@ -241,3 +124,4 @@ In the next section, we will add an "Experience Bar" where shooting blocks incre
 
 
 
+
diff --git a/Engine/Documentation/Source/Engine/UMG/QuickStart/5/UMGQS_5.INT.udn b/Engine/Documentation/Source/Engine/UMG/QuickStart/5/UMGQS_5.INT.udn
--- a/Engine/Documentation/Source/Engine/UMG/QuickStart/5/UMGQS_5.INT.udn
+++ b/Engine/Documentation/Source/Engine/UMG/QuickStart/5/UMGQS_5.INT.udn
@@ -1,13 +1,13 @@
 Availability: Public
-Title:5. Experience Bar and Hotbar Button
+Title:5. Creating / Scripting an In-Game Pause Menu
 Crumbs:%ROOT%, Engine, Engine/UMG, Engine/UMG/QuickStart
-Description:The fifth step in the UMG Quick Start is to create an experience bar and hotbar button. 
+Description:Here we will create an in-game pause menu which will allow us to pause or exit our game. 
 
 
 [VAR:Steps]
 [OBJECT:Navigation]
 	[PARAM:previous]
-		[Previous Step](Engine/UMG\QuickStart\4 "%Engine/UMG\QuickStart\4:title%") 
+		[Previous Step](Engine/UMG\QuickStart\5 "%Engine/UMG\QuickStart\5:title%") 
 	[/PARAM]
 	[PARAM:current]
 	[/PARAM]
@@ -15,7 +15,7 @@ Description:The fifth step in the UMG Quick Start is to create an experience bar
 		[Quick Start Home](Engine/UMG\QuickStart) 
 	[/PARAM]
 	[PARAM:next]
-		[Next Step](Engine/UMG\QuickStart\6 "%Engine/UMG\QuickStart\6:title%") 
+	[Next Step](Engine/UMG\QuickStart\6 "%Engine/UMG\QuickStart\6:title%") 
 	[/PARAM]
 [/OBJECT]
 [/VAR]
@@ -24,172 +24,158 @@ Description:The fifth step in the UMG Quick Start is to create an experience bar
 
 ---
 
-[OBJECT:TopicList]
-	[PARAM:icon]
-		![](%ROOT%/build_icon.png)(convert:false)
-	[/PARAM]
-	[PARAM:title]
-		Adding an Experience Bar and Hot Bar Button
-	[/PARAM]
-	[PARAM:description]
-	In this section, we will update our HUD to include an experience bar, level up text, and an action bar button. 
-	[/PARAM]
-[/OBJECT]
+The last thing we will do is create a simple Pause Menu which will allow the player to pause or exit the game. 
 
-1. From the **Content Browser**, select the **UMG** folder and open the **GameHUD** Blueprint. 
 
-1. Inside the **Palette** window under **Common**, drag a **Progress Bar** onto the **CanvasPanel** in the **Hierarchy** window. 
+## Visual: Pause Menu Layout
 
-1. Move and resize the bar as you would like (we placed ours at the bottom). 
+In this section we will create the visual layout of our basic Pause Menu. 
 
-	![](New_1.png)
+1.	Open your **Pause Menu** Widget Blueprint. 
 
-	In the **Details** panel, make sure to **Anchor** it where you would like (we Anchored it to the bottom center).
+1.	On the **Designer** tab, drag and drop widgets from the **Palette** onto the **Hierarchy** to achieve the setup below. 
 
-1. Add a **Text Block** to the **CanvasPanel**, this text will be used to let the user know they "Leveled Up."
+	![](Pause1.png)
 
-1. Place the **Text Block** where you would like, along with the text you would like to display.
+	Here we have a **Border** that contains a **Vertical Box** with **Text** and two **Buttons** each with **Text** on our **Canvas Panel**. 
 
-	![](New_2.png)
+1.	In the **Designer** window, resize the **Border** so that it fills the entire dotted window (filling the screen). 
 
-	 We used the text "Super Jump Unlocked!" and placed and Anchored it at the bottom.
+	![](Pause2.png)
 
-1. Next, add a **Button** to the **CanvasPanel**, move and resize as desired (we placed ours under our experience bar). 
+	You can also change the color from the **Brush Color** option inside the **Details** panel. 
 
-	![](New_4.png)
+	![](Pause3.png)
 
-1. In the **Details** panel for the button under **Behavior**, uncheck the **Is Enabled** button (as we will turn it on when the player levels up). 
+	You may also want to set the set the **Alpha** for your color to **0.5** to give it some translucency. When the Pause Menu is opened during gameplay, the Border will fill the screen but adding some translucency will allow the game to still be visible in the background. 
 
-	![](New_3.png)
+1.	Also for the **Border** in the **Details** panel, be sure to use the fill screen Anchor. 
 
-1. Next to the **IsEnabled** checkbox, click the **Bind** button then **Create Binding**. 
+	![](Pause4.png)
 
-1. In the graph, create a new **Bool** variable called **H_LevelUp** and connect it to the **Return Value**. 
+	This will ensure that the Border fills the screen regardless of screen size. 
 
-	![](New_5.png)
+1.	Select the **Vertical Box** in the **Hierarchy** then in the **Details** panel, center it both horizontally and vertically. 
 
-1. Return to the **Design** tab and click the experience bar. 
+	![](Pause5.png)
 
-1. In the **Details** panel, assign a color to the bar then click **Bind** next to **Percent** and **Create Binding**. 
+1.	Select the **Text** under the **Vertical Box** then in the **Details** panel for **Content**, enter **Pause Menu**. 
 
-1. In the graph, create a new **Float** variable called **H_Experience** and connect it to the **Return Value**. 
+	![](Pause6.png)
 
-	![](New_7.png)
+1.	Enter **Resume** and **Quit** for the other two **Text** blocks. 
 
-1. In the **MyBlueprint** window, **Double-click** the **EventGraph** to jump to it. 
+	![](Pause7.png)
 
-	![](New_15.png)
+1.	Hold **Ctrl** and select both **Buttons** in the **Hierarchy**, then in the **Details** panel under **Style**, assign a color for **Hovered**. 
 
-1. **Right-click** in the graph and add an **Event Tick** connected to a **Branch**. 
+	![](Pause8.png)
 
-1. For the condition, **Control** drag in the **H_Experience** variable and connect it to a **>=** node (set to **1.0**). 
+1.	Rename the **Buttons** to **Resume** and **Quit** respectively inside the **Details** panel. 
 
-1. Create a new **Text** variable called **H_UnlockedText**. 
+	![](Pause9.png)
 
-1. Next, **Alt** drag in the **H_UnlockedText** variable and enter the text you want to display when unlocking the ability. 
 
-	![](New_13.png)
+## Script: Pause Menu Functionality
 
-1.	Return to the **Designer** tab and select your unlock text, next to the **Text** field click **Bind** and select the **H_UnlockedText** variable.
+With the visuals set, next we will provide the scripted functionality for our Pause Menu. 
 
-	![](New_13b.png)
+1.	Click the **Graph** button in the upper-right corner of the Widget Blueprint Editor window. 
 
-	This will Bind the text displayed to whatever the variable is. 
+1.	On the **Graph** tab in the **My Blueprint** window, click the **QuitButton** then in the **Details** panel add an **OnClicked** Event. 
 
-1. Return to the **Graph** tab. 
+	![](PauseScript1.png)
 
-1. Add a **Delay** node following the unlock text (set to 1 second). 
+1.	Also add an **OnClicked** Event for the **ResumeButton**. 
 
-1. After the **Delay**, **Set H_Experience** to **0.0**, **Set H_UnlockedText** to **Blank** and finally set **H_LevelUp** to **True**. 
+1.	In the **Graph** window, **Right-click** and add a **Get Player Controller** node. 
 
-	[REGION:lightbox]
-	[![](New_14.png)(w:840)](New_14.png)
-	[/REGION]
-	
-	The resulting graph should look similar to above, click the image for full view. Here we are enabling the ability when enough experience is reached, displaying text that we unlocked the ability, then removing the text and resetting the experience bar.
-	
-1. **Compile** and **Save**, then close the Blueprint. 
+	![](PauseScript2.png)
 
-[OBJECT:TopicList]
-	[PARAM:icon]
-		![](%ROOT%/build_icon.png)(convert:false)
-	[/PARAM]
-	[PARAM:title]
-		Generating Experience and Finishing Up
-	[/PARAM]
-	[PARAM:description]
-	In this section, we will set it up so that when you shoot a box, it generates experience and updates the HUD. 
-	[/PARAM]
-[/OBJECT]
+1.	Off the **Get Player Controller**, use **Set Input Mode Game Only** then **Remove from Parent** and **Set Game Paused** as shown below. 
 
-1. In the **Content Browser** under **Game/Blueprints**, open the **MyProjectile** Blueprint. 
+	![](PauseScript3.png)
 
-1. In the graph after the **Destroy Actor** node, **Cast To MyCharacter** by adding a **Get Player Character** node. 
+	Here when the **Resume Button** is clicked, the input mode is set to game only and the cursor that is displayed when the Pause Menu is activated is removed. The **Remove from Parent** node is used to remove the actual Pause Menu widget from display before leaving the paused state by using the **Set Game Paused** node and leaving **Paused** unchecked. 
 
-1. Off the **As My Character C** pin add the **Get GameHUDWidget** (our Widget Blueprint) node. 
+1.	Create a new **Function** called **Remove HUD**. 
 
-	![](New_8.png)
-
-1. Off the **GameHUDWidget**, **Set H_Experience** to **H_Experience** + **0.10** as seen below. 
-
-	![](New_9.png)
-	
-	The ending sequence should look similar to above, click the image for a full view.
+	![](PauseScript4.png)
 
-1. **Compile** and **Save**, then close the Blueprint.
+	We will use this to remove the game HUD from the player's view when the Pause Menu is activated. 
 
-1. In the **Content Browser** under **Game/Blueprints**, open the **MyCharacter** Blueprint. 
+1.	In the Function, add a **Get Player Character** node and off the **Return Value**, add the **Cast ToFirstPersonCharacter** node. 
 
-1. **Right-click** and add a **1** Key Press Event, drag in the **GameHUDWidget** variable and off it, get **H_LevelUp**. 
+	![](PauseScript5.png)
 
-1. Off the **Pressed** pin of the **1** Key Press Event, add a **Branch** with the condition being **H_LevelUp**. 
+	This will allow us to access the Character Blueprint being used by the player and the game HUD that is being used by that player. 
 
-	![](New_10.png)
+1.	Off the **As First Person Character**, use **Get HUD Reference** then **Remove from Parent** as shown below. 
 
-1. Off the **True** of the **Branch**, add a **Launch Character** node and set the **Launch Velocity** X to **50** and Z to **1000.0**. 
+	![](PauseScript6.png)(w:840)
 
-1. **Alt** drag in the **P_Energy** variable, then **Control** drag it in and connect it to a **Float - Float** node set to **0.5**. 
+	This will access the HUD being used by the player character and remove it from the viewport when quitting the game. 
 
-	![](New_11.png)
+1.	On the **Event Graph**, off **Get Player Controller** use the **Set Show Mouse Cursor** node set to **True**. 
 
-	This will launch the character in the air and subtract energy from the character which will be displayed via the HUD.
+	![](PauseScript7.png)(w:840)
 
-1. Drag off the **GameHUDWidget** node in this sequence again and **Set H_LevelUp** to **False** (unchecked). 
+	When **Quit** is clicked, this will turn on the mouse cursor for menu navigation when returning to the Main Menu. 
 
-1. Add a **Delay** node (set to **2**) following the previous step, then **Set H_LevelUp** to **True** (checked). 
+1.	Next add the **Remove HUD** function created above followed by **Open Level** (set to **MainMenu**) and connect as shown. 
 
 	[REGION:lightbox]
-	[![](New_12.png)(w:840)](New_12.png)
+	[![](PauseScript8.png)(w:840)](PauseScript8.png)
 	[/REGION]
 
-	_Click image for full view._
+	_Click for full view._
 
-	When **1** is pressed, if the player has unlocked the ability, we launch them and remove some energy, then add a 2 second cooldown to the ability.
+	Both Buttons are now setup, the last thing that we need to do is allow the player to call the Pause Menu and pause the game.  
 
-1. Next, find the section for the **F** key press event, following the **Set P_Health** node, add a **Branch** with condition being **P_Health** <= **0.0**. 
+1.	In the **Content Browser** under **Content/FirstPersonBP/Blueprints**, open the **FirstPersonCharacter** Blueprint. 
 
-1. **Control** drag in the **GameHUDWidget** variable and place it under the **Branch**, then drag off it and add the **Remove from Viewport** node. 
+	![](PauseScript9.png)
 
-1. Connect the **Remove from Viewport** node to the **True** of the **Branch**, then off the exit pin, add a **Print String** with text **You're Dead!**. 
+1.	In the Graph, add a **M** Key Event and a **Create Widget** node (set to **Pause Menu**) and promote the **Return Value** to a variable. 
 
-	[REGION:lightbox]
-	[![](New_16.png)(w:800)](New_16.png)
-	[/REGION]
+	![](PauseScript10.png)
+
+	While we used the **M** Key as our input for bringing up the Pause Menu, you can use any key you wish. For the promoted variable, we called our variable **Pause Menu Reference** as this will store a reference to created Pause Menu. 
+
+1.	Hold **Ctrl** and drag in the variable from the **MyBlueprint** window and connect it to an **IsValid** node. 
+
+	![](PauseScript11.png)
+
+	We are checking to see if the Pause Menu widget has been created and stored in this variable yet. 
+
+1.	Connect the **Is Not Valid** pin to the **Create PauseMenu_C Widget**. 
+
+	![](PauseScript12.png)
 
-	The last few steps should result in the above, click for full view. This will remove the HUD when the player is out of Health (or "dead").
+1.	**Right-click** and add the **Get Player Controller** node then **Set Input Mode UI Only** and connect the nodes as shown below. 
 
-	[REGION:note]
-	Note: If you do not see the **Remove from Viewport** option, use the **Remove from Parent** option instead.
+	![](PauseScript13.png)
+
+	When **M** is pressed, if the Pause Menu has been accessed before it will not need to create it again and will instead access the variable. If this is the first time the Pause Menu has been accessed it will create it then store it as a variable so it can be access later. In either case, the input mode is set to UI only before we display the Pause Menu (which we will do next). 
+
+1.	Off the **Pause Menu Reference** variable, use the **Add to Viewport** node then **Set Game Paused** (set **Paused** to **True**). 
+
+	[REGION:lightbox]
+	[![](PauseScript14.png)(w:840)](PauseScript14.png)
 	[/REGION]
+	
+	_Click for full view._
 
-1. **Compile** and **Save** then close the Blueprint then play in the editor.
+	Our Pause Menu is setup, now it is time to test. 
 
-	![](Gif1.gif)(convert:false)
+1.	Click the **Compile** and **Save** Buttons, then the **Play** button to play in the editor. 
 
-	You should now see that when you hit 10 boxes, you level up and the text appears on screen. After a moment, you can press 1 to perform the ability and a cooldown will be applied to the button. You can also press the **F** key several times until you are out of health at which point the HUD should disappear.
+	[![](Engine\UMG\QuickStart\4\MenuImage.png)(w:720)](PauseMovie.mp4)
 
+	Once in-game, if you press the **M** Key (or whichever key you assigned), you will see that the game is paused and the Pause Menu appears. From here, you can resume gameplay with the **Resume** button or quit back to the Main Menu with the **Quit** Button. 
 	
-While this is a simple example and there is plenty more to do here, you should have an understanding of how to create and implement screens using **UMG** now as well as how to pass variables to and communicate with them using different methods. The final step of this guide will provide you with additional links to more documentation on using **UMG** and **Blueprints Visual Scripting** which you can use to take this example even further on your own!
+
+The next section will provide you with additional information on working with **Unreal Motion Graphics** and **Blueprints**. 
 
 ---
 
@@ -198,4 +184,3 @@ While this is a simple example and there is plenty more to do here, you should h
 
 
 
-
diff --git a/Engine/Documentation/Source/Engine/UMG/QuickStart/UMGQuickStart.INT.udn b/Engine/Documentation/Source/Engine/UMG/QuickStart/UMGQuickStart.INT.udn
--- a/Engine/Documentation/Source/Engine/UMG/QuickStart/UMGQuickStart.INT.udn
+++ b/Engine/Documentation/Source/Engine/UMG/QuickStart/UMGQuickStart.INT.udn
@@ -42,7 +42,7 @@ Description:Getting started with using Unreal Motion Graphics in Unreal Engine 4
 [/OBJECT]
 [/VAR]
 
-In this **Quick Start Guide**, you will learn how to implement some basic HUD elements using **Unreal Motion Graphics UI Designer (UMG)**. You will learn how to implement Health and Energy Bars, an Experience Bar, hiding/display text based on conditions, and an Action Bar style button that when pressed, activates a special character ability. 
+In this **Quick Start Guide**, you will learn how to implement some basic in-game HUD elements and a front end menu using **Unreal Motion Graphics UI Designer (UMG)**. You will learn how to implement Health and Energy Bars as well as a way to track and display the Ammo of a player. You will also learn how to create a Main Menu which can launch into a game, set resolution options or quit the game. From there you will learn how to add an in-game pause menu which can be used to pause the game or return to the Main Menu. 
 
 Over the course of this guide, links to additional documentation will be provided which will go deeper into each section's subject should you have additional questions. While this guide will cover the actions needed to proceed through each step, if you have not worked with the Unreal Editor before, it is recommended that you first look over the **[Level Designer Quick Start](Engine/QuickStart/)** guide to gain a better understanding of general editor usage, terms, and processes. 
 
