#UE4doc #docs: removed unnecessary ellipsis (used with parentheses)

diff --git a/Engine/Documentation/Source/Engine/Rendering/ParticleSystems/Reference/TypeData/Mesh/MeshTypeData.INT.udn b/Engine/Documentation/Source/Engine/Rendering/ParticleSystems/Reference/TypeData/Mesh/MeshTypeData.INT.udn
--- a/Engine/Documentation/Source/Engine/Rendering/ParticleSystems/Reference/TypeData/Mesh/MeshTypeData.INT.udn
+++ b/Engine/Documentation/Source/Engine/Rendering/ParticleSystems/Reference/TypeData/Mesh/MeshTypeData.INT.udn
@@ -75,7 +75,7 @@ The Mesh type data module has the following properties:
 [EXCERPT:CameraFacingCaveats]
 [REGION:note]All axis-locked camera facing options assume the AxisLockOption is set. EPAL_NONE will be treated as EPAL_X.[/REGION]
 
-[REGION:note]All velocity-aligned options do NOT require the ScreenAlignment be set to PSA_Velocity. Doing so will result in additional work being performed... (it will orient the mesh twice).[/REGION]
+[REGION:note]All velocity-aligned options do NOT require the ScreenAlignment be set to PSA_Velocity. Doing so will result in additional work being performed (it will orient the mesh twice).[/REGION]
 [/EXCERPT:CameraFacingCaveats]
 
 [EXCERPT:MeshAlignmentOption]
