singular subject  with verb: "changes," effect to "affect" (verb), duplicate word, "module"

diff --git a/Engine/Documentation/Source/Engine/Rendering/ParticleSystems/VectorFields/VectorFields.INT.udn b/Engine/Documentation/Source/Engine/Rendering/ParticleSystems/VectorFields/VectorFields.INT.udn
--- a/Engine/Documentation/Source/Engine/Rendering/ParticleSystems/VectorFields/VectorFields.INT.udn
+++ b/Engine/Documentation/Source/Engine/Rendering/ParticleSystems/VectorFields/VectorFields.INT.udn
@@ -11,7 +11,7 @@ The most interesting feature of GPU particles, aside from their efficiency, is v
 
 By default, vector fields impart a force on particles within them. Vector fields also have a "tightness" parameter. This parameter controls how directly particles follow the vectors in the field. When tightness is set to 1, particles read their velocity directly from the field and thus follow the field exactly.
 
-Static vector fields are those in which the grid of vectors never change. These fields can be exported from Maya and imported as a volume texture. Static fields are very cheap and can be used to add interesting motion to particles, especially by animating the motion of the field itself.
+Static vector fields are those in which the grid of vectors never changes. These fields can be exported from Maya and imported as a volume texture. Static fields are very cheap and can be used to add interesting motion to particles, especially by animating the motion of the field itself.
 
 Additionally, vector fields may be reconstructed from a 2D image. In this case, an image much like a normal map can be imported and used to reconstruct a volume texture by extruding or revolving it around a volume. A static vector field may be added on top of this reconstruction to introduce some noise and randomness. Further, the 2D images may be animated by storing individual frames in an atlas texture. Doing so allows you to perform a fluid simulation offline and reconstruct the motion in real-time at very little cost.
 
@@ -39,7 +39,7 @@ In this image, only the particles about halfway down the length of the emission
 
 ## Local Vector Fields
 
-Local Vector Fields, as opposed to their global counterparts, are vector fields that exist entirely within a particle system and are not placed in the world. This means that local vector fields can only effect the particle emitters to which they are assigned, unlike Global Vector Fields, which can affect any particle system with a [Global Vector Fields module](Engine/Rendering/ParticleSystems/Reference/Modules/VectorField) module.
+Local Vector Fields, as opposed to their global counterparts, are vector fields that exist entirely within a particle system and are not placed in the world. This means that local vector fields can only affect the particle emitters to which they are assigned, unlike Global Vector Fields, which can affect any particle system with a [Global Vector Fields module](Engine/Rendering/ParticleSystems/Reference/Modules/VectorField).
 
 Local Vector Fields are added via the [Local Vector Field module](Engine/Rendering/ParticleSystems/Reference/Modules/VectorField). 
 
