#UE4doc #docs: Changed title to "Transparency," minor punctuation edits

diff --git a/Engine/Documentation/Source/Engine/Rendering/Materials/HowTo/Transparency/Transparency.INT.udn b/Engine/Documentation/Source/Engine/Rendering/Materials/HowTo/Transparency/Transparency.INT.udn
--- a/Engine/Documentation/Source/Engine/Rendering/Materials/HowTo/Transparency/Transparency.INT.udn
+++ b/Engine/Documentation/Source/Engine/Rendering/Materials/HowTo/Transparency/Transparency.INT.udn
@@ -10,11 +10,11 @@ Related: Engine/Rendering/LightingAndShadows/LitTranslucency
 
 ## Overview
 
-When creating certain surface types such as water or glass, you need the ability to make the surface not only see through but also give the surface a sense of depth and color. 
+When creating certain surface types such as water or glass, you need the ability to make the surface not only see through, but also give the surface a sense of depth and color. 
 In the real world, these properties are often referred to as **Transparency** or **Opacity** and are often used interchangeably to describe the same thing. 
 However inside of Unreal Engine 4(UE4), Transparency is used to describe if a surface is see-through or not while Opacity is used to define how see-through the surface is. In the following How - To, you will learn everything you need to know about how to introduce Transparency to your Materials in Unreal Engine 4.
 
-## What is Transparency
+## Transparency
 
 Transparency is the term used to describe a surface's ability to block or allow the passage of light. 
 For example, a brick would be an object that has no Transparency while stained glass would be an object that has Transparency.
@@ -88,7 +88,7 @@ content page for information about how to move content between projects. This wa
 	![](Tras_Creat_New_MI.gif)(convert:false)
 
 1. Now we need to find an object to place the Material Instance on so we can see the Material in action. Inside of the **Content Browser**, under the Starter Content folder there is a **Shapes** folder.
-Locate the  **Shape_Sphere** and then select it with the **Left Mouse Button** and then drag it into the current level with the **Left Mouse Button** still held down . Once in a location you like, release the **Left Mouse Button** to place the mesh in the level.
+Locate the  **Shape_Sphere** and then select it with the **Left Mouse Button** and then drag it into the current level with the **Left Mouse Button** still held down. Once in a location you like, release the **Left Mouse Button** to place the mesh in the level.
 
 	![](Trans_Add_Test_Mesh.png)
 
