changed, "lets" to "let's"

diff --git a/Engine/Documentation/Source/Engine/Rendering/PostProcessEffects/AutomaticExposure/AutomaticExposure.INT.udn b/Engine/Documentation/Source/Engine/Rendering/PostProcessEffects/AutomaticExposure/AutomaticExposure.INT.udn
--- a/Engine/Documentation/Source/Engine/Rendering/PostProcessEffects/AutomaticExposure/AutomaticExposure.INT.udn
+++ b/Engine/Documentation/Source/Engine/Rendering/PostProcessEffects/AutomaticExposure/AutomaticExposure.INT.udn
@@ -37,7 +37,7 @@ Description:Automatic adjustment of scene exposure to simulate eye adaptation fr
 
 Eye adaptation is [histogram](#Histogram) based and uses the following method:
 
-* Lets say **EyeAdaptationLowPercent** is `80%` and **EyeAdaptationHighPercent** is `95%`.
+* Let's say **EyeAdaptationLowPercent** is `80%` and **EyeAdaptationHighPercent** is `95%`.
 * Now we search the histogram for two values:
     * 80% of the screen pixels are darker than luminance value A.
     * 95% of the screen pixels are darker than luminance value B.
