#UE4doc #docs: duplicate word "the,"

diff --git a/Engine/Documentation/Source/Engine/Rendering/PostProcessEffects/Bloom/Bloom.INT.udn b/Engine/Documentation/Source/Engine/Rendering/PostProcessEffects/Bloom/Bloom.INT.udn
--- a/Engine/Documentation/Source/Engine/Rendering/PostProcessEffects/Bloom/Bloom.INT.udn
+++ b/Engine/Documentation/Source/Engine/Rendering/PostProcessEffects/Bloom/Bloom.INT.udn
@@ -4,7 +4,7 @@ Crumbs: %ROOT%, Engine, Engine/Rendering/PostProcessEffects
 Description:Haloing to produce glow effects for bright objects such as lights.
 
 [EXCERPT:Bloom_Desc]
-**Bloom** is a real world light phenomena that can greatly add to the perceived realism of a rendered image at a moderate render performance cost. Bloom can be seen by the naked eye when looking at very bright objects that are on a much darker background. Even brighter objects also cause other effects (streaks, lens flares), but those are not covered by the classic bloom effect. Because our displays (e.g. TV, TFT, ...) usually do not support HDR (high dynamic range), we cannot really render very bright objects. Instead we simulate the effects that happen in the eye (retina subsurface scattering), when light hits the film (film subsurface scattering), or in front of the the camera (milky glass filter). The effect might not always be physically correct but it can help to hint the relative brightness of objects or add realism to the LDR (low dynamic range) image that is shown on the screen.
+**Bloom** is a real world light phenomena that can greatly add to the perceived realism of a rendered image at a moderate render performance cost. Bloom can be seen by the naked eye when looking at very bright objects that are on a much darker background. Even brighter objects also cause other effects (streaks, lens flares), but those are not covered by the classic bloom effect. Because our displays (e.g. TV, TFT, ...) usually do not support HDR (high dynamic range), we cannot really render very bright objects. Instead we simulate the effects that happen in the eye (retina subsurface scattering), when light hits the film (film subsurface scattering), or in front of the camera (milky glass filter). The effect might not always be physically correct but it can help to hint the relative brightness of objects or add realism to the LDR (low dynamic range) image that is shown on the screen.
 [/EXCERPT:Bloom_Desc]
 
 [REGION:fullwidth]
@@ -54,7 +54,7 @@ The **Bloom Dirt Mask** effect uses a texture to brighten up the bloom in some d
 
 | Property | Description |
 | -------- | ----------- |
-| **Dirt Mask** | The Texture2D to use as the mask. It is recommended that a non power of two texture is used as it fits the screen size better and it avoids mip creation and streaming. A low resolution texture is often sufficient as dirt on the camera is likely to be out of focus and blurry anyway. The texture can be created with the Photoshop Lens blur. Chromatic aberration (color fringe) can improve the look. Beware of compression artifacts. This property cannot be blended. We plan to make this more dynamic and programmable through custom shaders. |
+| **Dirt Mask** | The Texture2D to use as the mask. It is recommended that a non-power of two texture is used as it fits the screen size better and it avoids mip creation and streaming. A low resolution texture is often sufficient as dirt on the camera is likely to be out of focus and blurry anyway. The texture can be created with the Photoshop Lens blur. Chromatic aberration (color fringe) can improve the look. Beware of compression artifacts. This property cannot be blended. We plan to make this more dynamic and programmable through custom shaders. |
 | **Dirt Mask Tint Color** | Used to darken or color tint the dirt mask texture. This can be used to tweak the dirt mask brightness and color. |
 | **Dirt Mask Intensity** | Scales the color of the bloom dirt mask effect (linear). This property is used to tweak dirt mask brightness. |
 
