#UE4doc #docs: minor wording edits for clarification.

diff --git a/Engine/Documentation/Source/Engine/Rendering/PostProcessEffects/ColorGrading/ColorGrading.INT.udn b/Engine/Documentation/Source/Engine/Rendering/PostProcessEffects/ColorGrading/ColorGrading.INT.udn
--- a/Engine/Documentation/Source/Engine/Rendering/PostProcessEffects/ColorGrading/ColorGrading.INT.udn
+++ b/Engine/Documentation/Source/Engine/Rendering/PostProcessEffects/ColorGrading/ColorGrading.INT.udn
@@ -19,9 +19,9 @@ This is the used function:
 
 	GammaColor = LinearColor / (LinearColor + 0.187) * 1.035;
 
-Note that bright colors become gradually brighter but much less than dark ones. Black remains black and there is a mostly linear part of the curve that is a bit steeper than the non tone-mapped curve would be. That results in some contrast enhancing. It is normal and expected to get a different look when using a tone mapper and in order to get a good result, the source image needs to have more dynamic in the brightness (more HDR). This can pay off in a more realistic film like look.
+Note that bright colors become gradually brighter but much less than dark ones. Black remains black and there is a mostly linear part of the curve that is a bit steeper than the non tone-mapped curve would be. That results in some contrast enhancing. While it is normal to get a different look when using a tone mapper, in order to get good results the source image needs to have a higher dynamic range. This can pay off in a more realistic film like look. 
 
-The presented tone mapper formula already has two tweakable constants and more math would allow further tweaking, but finding a good function is not easy as quality/flexibility/performance needs to be considered. We decided to stick to a simple formula and modify the resulting LDR color with a simple color lookup. As we already mapped the HDR color into a limited range, we get a lot of expressibility for the dark colors and still can modify the bright colors. This method is simple to understand, offers a lot of flexibility with local control and almost constant performance. 
+The presented tone mapper formula already has two tweakable constants and while more math would allow further tweaking, finding a good function is not easy as quality, flexibility, and performance need to be considered. We decided to stick to a simple formula and modify the resulting LDR color with a simple color lookup. As we already mapped the HDR color into a limited range, we get a lot of expressibility for the dark colors and still can modify the bright colors. This method is simple to understand and offers a lot of flexibility with local control and consistent performance. 
 
 | Property | Description |
 | -------- | ----------- |
