added apostrophe (to show possesion)

diff --git a/Engine/Documentation/Source/Engine/Rendering/PostProcessEffects/DepthOfField/DepthOfField.INT.udn b/Engine/Documentation/Source/Engine/Rendering/PostProcessEffects/DepthOfField/DepthOfField.INT.udn
--- a/Engine/Documentation/Source/Engine/Rendering/PostProcessEffects/DepthOfField/DepthOfField.INT.udn
+++ b/Engine/Documentation/Source/Engine/Rendering/PostProcessEffects/DepthOfField/DepthOfField.INT.udn
@@ -8,7 +8,7 @@ Description:Simulating focus by blurring the scene depending on depth.
 [/REGION]
 
 [EXCERPT:DepthOfField_Desc]
-Depth of Field (DoF) applies a blur to the scene based on distance in front of or behind a focal point. This is to simulate what happens in real-world cameras. The effect can be used to draw the viewers attention and to
+Depth of Field (DoF) applies a blur to the scene based on distance in front of or behind a focal point. This is to simulate what happens in real-world cameras. The effect can be used to draw the viewer's attention and to
 make the rendering appear more like a photograph or like a movie.
 [/EXCERPT:DepthOfField_Desc]
 
@@ -155,7 +155,7 @@ Bokeh depth of field is rendered at quarter-resolution (half-res in each directi
 
 You can see that the characters being blurred in the background appear blocky when using only the downsampled DoF technique. Also, there is some artifacting around the horns of the character in the foreground. The artifacting is eliminated by using the adaptive DoF, and the background characters have a much smoother appearance.
 
-The adaptive DoF can be visualized using the **Adaptive Depth of Field** show flag ( under **Show > Visualize** ). When enabled, an overlay is displayed showing where the downsampled effect is being used (in green) and where the full-resolution effect is being used (in red). The regular scene color is shown where no blurring is being applied.
+The adaptive DoF can be visualized using the **Adaptive Depth of Field** show flag (under **Show > Visualize** ). When enabled, an overlay is displayed showing where the downsampled effect is being used (in green) and where the full-resolution effect is being used (in red). The regular scene color is shown where no blurring is being applied.
 
 [REGION:fullwidth]
 ![Adaptive DoF Visualization](dof_adaptive_visualize.png)
