#UE4doc #docs: changed all instances to "post process" for consistency, sp: translucency, from & multiple edits

diff --git a/Engine/Documentation/Source/Engine/Rendering/PostProcessEffects/PostProcessMaterials/PostProcessMaterials.INT.udn b/Engine/Documentation/Source/Engine/Rendering/PostProcessEffects/PostProcessMaterials/PostProcessMaterials.INT.udn
--- a/Engine/Documentation/Source/Engine/Rendering/PostProcessEffects/PostProcessMaterials/PostProcessMaterials.INT.udn
+++ b/Engine/Documentation/Source/Engine/Rendering/PostProcessEffects/PostProcessMaterials/PostProcessMaterials.INT.udn
@@ -12,9 +12,9 @@ Description:How to author and blend custom Post Process passes with the material
 ## Post Processing Graph
 
 [EXCERPT:PostProcessMaterials_Desc]
-The engine already features complex postprocessing based on a graph of post processing nodes. The **Post Processing Materials**
+The engine already features complex post processing based on a graph of post processing nodes. The **Post Processing Materials**
 can be additionally inserted in some specific position. See [FAQ](#FAQ) **r.CompositionGraphDebug** to get a dump of the full graph.
-The graph is actually not only doing postprocessing but also some parts of the lighting. Over time we plan to make more
+The graph is actually not only doing post processing but also some parts of the lighting. Over time we plan to make more
 parts be customizable through the material editor.
 [/EXCERPT:PostProcessMaterials_Desc]
 
@@ -30,10 +30,10 @@ Post Process Materials should be used sparingly and only when you really need th
 ## Using a Post Process Material
 
 Through post process settings (usually defined with post process volumes or camera settings), it is possible to blend (so called) blendable assets.
-At the moment, only **materials** and **material instances** are blendable assets. The engine provides a few postprocess materials but with this feature
-you can create your own **custom postprocessing** without any programmer help.
+At the moment, only **materials** and **material instances** are blendable assets. The engine provides a few post process materials but with this feature
+you can create your own **custom post processing** without any programmer help.
 
-Simply assign one or more post process materials to a postprocess volume in the **Misc** section to the **Blendables**. First press the + to add
+Simply assign one or more post process materials to a post process volume in the **Misc** section to the **Blendables**. First press the + to add
 new slots, select a material in the **Content Browser**, and press the left arrow to assign it. The order here is not important and unused slots are simply ignored.
 
 ![](PostProcessSettings.gif)
@@ -46,11 +46,11 @@ For an overview of creating a simple Post Process Material effect from scratch,
 
 ## Critical Settings for Post Process Materials
 
-A postprocessing material needs to specify the material domain **post process**:
+A post processing material needs to specify the material domain **post process**:
 
 ![](DomainPostProcess.png)
 
-The material should only use *EmissiveColor* to output the new color. Additionally, you can define where during the postprocessing this pass should be applied
+The material should only use *EmissiveColor* to output the new color. Additionally, you can define where during the post processing this pass should be applied
 and, if there are multiple, in which order they should be processed (Priority):
 
 ![](PostProcessMaterialProps.png)
@@ -59,10 +59,10 @@ and, if there are multiple, in which order they should be processed (Priority):
 | -------- | ----------- |
 | **Before Tonemapping** | PostProcessInput0 provides the scene color with all lighting in HDR. Use this to fix issues with temporal antialiasing and GBuffer lookups e.g. depth, normals. |
 | **After Tonemapping** | Preferred location for performance as the color is LDR and therefore requires less precision and bandwidth. This is after tone mapping and color grading. |
-| **Before Translucency** | This is even earlier in the pipeline than 'Before Tonemapping' before translucencty was combined with the scene color. Note that Separate transluceny and is composited later than normal transluceny. |
+| **Before Translucency** | This is even earlier in the pipeline than 'Before Tonemapping' before translucency was combined with the scene color. Note that SeparateTranslucency is composited later than normal translucency. |
 | **Replacing the Tonemapper** | PostProcessInput0 provides the HDR scene color, PostProcessInput1 has the SeparateTranslucency (Alpha is mask), PostprocessInput2 has the low resolution bloom input. |
 
-The typical postprocess input comes form the pass before. That color can be accessed through the SceneTexture material expression when using **PostProcessInput0**. Using SceneColor might not give you the right results.
+The typical postprocess input comes from the pass before. That color can be accessed through the SceneTexture material expression when using **PostProcessInput0**. Using SceneColor might not give you the right results.
 
 ## Blending Between Different Material Instances
 
@@ -70,7 +70,7 @@ With postprocess volumes, it is easy to setup soft transitions between multiple
 
 ![](BlendingAVolume.png)
 
-Each volume we specify a different material instance of the same material. The color is specified as a material parameter which allows to have a different settings for the two material instances.
+Each volume we specify a different material instance of the same material. The color is specified as a material parameter which allows different settings for the two material instances.
 
 Depending on the camera position the setting of one volume are used:
 
@@ -99,7 +99,7 @@ The node has an optional input and multiple outputs:
 
 ![](SceneTextureExpression.png)
 
-The **UV** input allows to specify where you want to make a texture lookup (only used for the Color output).
+The **UV** input allows you to specify where you want to make a texture lookup (only used for the Color output).
 The color output is a 4 channel output (actual channel assignment depends on the scene texture id). **Size** is a 2 component vector with the
 width and height of the texture. The inverse of that (1/width, 1/height) is available in the **InvSize** output. It is handy when referencing
 neighbor samples like in the example below:
@@ -124,7 +124,7 @@ This adds extra draw calls but no more materials. The rendering is fairly cheap
 
 ![](CustomDepth.png)
 
-In this scene, we activated the feature on two objects but without a postprocessing pass visualizing the content, this feature remains invisible:
+In this scene, we activated the feature on two objects but without a post processing pass visualizing the content, this feature remains invisible:
 
 ![](scene.png)
 
@@ -140,8 +140,8 @@ This is the material we used to visualize it:
 
 Temporal anti-aliasing is a unique UE4 feature that greatly improves image quality with a very moderate performance cost.
 
-By default, a post process material is inserted in the end of the post processing graph (after tone mapper). This means you get the final LDR color after tonemapping, color grading and after the temporal antialasing
-was applied. This is the best spot for many simple postprocessing effects - for performance and ease of use.
+By default, a post process material is inserted in the end of the post processing graph (after tone mapper). This means you get the final LDR color after tonemapping, color grading and after the temporal anti-alasing
+was applied. This is the best spot for many simple post processing effects - for performance and ease of use.
 
 Here you can see how we used the custom depth input to visualize a silhouette around specific objects:
 
@@ -149,7 +149,7 @@ Here you can see how we used the custom depth input to visualize a silhouette ar
 
 Notice that the former image has no anti-aliasing on the silhouette but in motion you also would see the silhouette jittering around 1 pixel.
 This is because temporal anti-aliasing moves the rendering of the whole scene by a sub pixel each frame. Multiple frames together are getting combined to produce the final anti-aliased image.
-We can however move the material to an earlier sport in the postprocessing graph to fix this issue.
+We can however move the material to an earlier sport in the post processing graph to fix this issue.
 
 This is the result:
 
@@ -161,15 +161,15 @@ but ideally this should not be.
 
 ## UV and ScreenPos
 
-Postprocess materials you can lookup into screen aliged buffers but you need to know the right UVs. 
+Post process materials you can look up into screen aligned buffers, but you need to know the right UVs. 
 The **ScreenPosition** material expression outputs the UV that you would expect (0,0 at the left top of the viewport and 1,1 at the bottom right).
 Using the **texture coordinate** material expression might give you different results. This is because the actual texture (more correctly it is a render target) is potentially larger than the viewport.
 It can be larger in editor because we share this texture for multiple viewports and the largets extent is use for all viewports.
-Even in game it can be larger in some cases (e.g. SceneCaptureActors might have a smaller viewport, Matienee black borders, Splitscreen, VR, ...).
+Even in game it can be larger in some cases (e.g. SceneCaptureActors might have a smaller viewport, Matinee black borders, Splitscreen, VR, ...).
 The texture coordinate material expression gives you the UV for this larger texture. If you only need a relative offset (e.g. pixel sized edge detection) you need to scale with the right size.
 The **SceneTexture** material expression has outputs for the size and the inverse of the size (efficient and useful for pixel offsets).
 If you want to get the viewport UV (e.g. to map texture to the viewport) you can use the ScreenPosition material expression or you make the needed computations manually (more control, can be slower).
-For this, you need the **ViewSize** material expression. To test all that, you can use the console variable **r.ViewPortTest** which allows to test various viewport configuations.
+For this, you need the **ViewSize** material expression. To test all that, you can use the console variable **r.ViewPortTest** which allows you to test various viewport configurations.
 
 ## Filtered texture lookups
 
@@ -182,7 +182,7 @@ It is possible to override the in engine tonemapper with your own one by using t
 
 ![](ReplacingTheTonemapper.png)
 
-We started exposing some Postprocess setting parameters to the tonemapper but that part is likely to change quite a bit.
+We started exposing some Post process setting parameters to the tonemapper but that part is likely to change quite a bit.
 The values are exposed as material parameters and need to have the exact name. 
 
 Vector parameter:
@@ -194,11 +194,11 @@ Scalar parameter:
 	
 [EXCERPT:PostMatWarning]
 [REGION:warning]
-In order to get the parameters you need to create a material instance from the postprocess material!
+In order to get the parameters you need to create a material instance from the post process material!
 [/REGION]
 [/EXCERPT:PostMatWarning]
 
-You still can use your own parameters and get them blended like other postprocess material settings.
+You still can use your own parameters and get them blended like other post process material settings.
 
 ## Known Issues
 The following issues still needs to be fixed:
@@ -206,14 +206,14 @@ The following issues still needs to be fixed:
   * **Material expression SceneTexture**
     * SeparateTranslucency does not work.
     * Certain lookups do not work in certain passes (e.g. some we might not fix because they would cost too much performance).
-    * MaterialFunction might report an error but it still can be used in a Material with the PostProcess domain.
+    * MaterialFunction might report an error but it still can be used in a Material with the Post Process domain.
   * **Material**
     * UV in the PostProcesMaterial might not be in 0..1 range (e.g. in editor when reducing viewport size), this aligns with the lookup but makes it hard to implement something like an vignette effect.
-    * The asset thumbnail of a postprocess material does not look right.
+    * The asset thumbnail of a post process material does not look right.
     * Outputting alpha is not yet supported (should go through Opacity).
     * The preview material in the material editor does not look right.
-    * There has been a case where a material change was not resulting in a change in the postprocess. Workaround: restart editor.
-	* The **Content Browser** should allow easy filtering for postprocess materials.
+    * There has been a case where a material change was not resulting in a change in the post process. Workaround: restart editor.
+	* The **Content Browser** should allow easy filtering for post process materials.
   * **Blending**
     * When blending two post process volumes with a blend radius, it is possible to see a non soft transition. This can be prevented by using an unbound volume with a material instance setting that represents the default.
 
@@ -228,13 +228,13 @@ The following issues still needs to be fixed:
 
     When SceneColor is used we create a lower quality copy of the scene to allow a lookup into a texture that we currently
 	write to (General case is rendering any mesh where this is not possible).
-	In PostProcessing, you should use PostProcessInput0.
+	In Post Processing, you should use PostProcessInput0.
 
-  * **How much memory do I pay for a postprocess?**
+  * **How much memory do I pay for a post process?**
 
     Memory cost depends on screen resolution. Before tonemapping we use HDR (8bytes per pixel), after we use LDR (4 bytes per pixel).
 
-  * **How can I lower the render cost for postprocessing?**
+  * **How can I lower the render cost for post processing?**
 
     Measure on the target platform, keep the texture lookup count low, do less math operations and reduce dependent texture
 	lookups, avoid randomized texture lookup (can be slower due to texture cache misses).
@@ -244,9 +244,9 @@ The following issues still needs to be fixed:
     Each pass adds to the performance cost. Try to combine passes and activate passes only if needed. General game features
 	e.g. noise could be added to the engine passes for better performance.
 
-  * **How much CPU performance do I pay per postprocess and for blending?**
+  * **How much CPU performance do I pay per post process and for blending?**
 
-    Blending materials is very cheap. All material instances properties are getting blended and only one postprocess material pass with those settings gets rendered.
+    Blending materials is very cheap. All material instances properties are getting blended and only one post process material pass with those settings gets rendered.
 
   * **I need to use "Before Tonemapper" to get proper TemporalAA. When I use a color, it is tone mapped and therefore looks different. How can I prevent that?**
     
