UMG Tutorial

diff --git a/Engine/Documentation/Source/Programming/Tutorials/ProgrammingTutorialsHome.INT.udn b/Engine/Documentation/Source/Programming/Tutorials/ProgrammingTutorialsHome.INT.udn
--- a/Engine/Documentation/Source/Programming/Tutorials/ProgrammingTutorialsHome.INT.udn
+++ b/Engine/Documentation/Source/Programming/Tutorials/ProgrammingTutorialsHome.INT.udn
@@ -184,4 +184,23 @@ programming
 		components
 	[/PARAMLITERAL]
 [/OBJECT]
+
+[OBJECT:Role]
+	[PARAM:icon]
+		![](%ROOT%/start_icon.png)(convert:false)
+	[/PARAM]
+	[PARAM:title]
+		%Programming/Tutorials/UMG:title%
+	[/PARAM]
+	[PARAM:path]
+		[RELATIVE:Programming/Tutorials/UMG]
+	[/PARAM]
+	[PARAM:description]
+		%Programming/Tutorials/UMG:description%
+	[/PARAM]
+	[PARAMLITERAL:id]
+		umg
+	[/PARAMLITERAL]
+[/OBJECT]
+
 [/REGION]
diff --git a/Engine/Documentation/Source/Programming/Tutorials/UMG/1/UMG_1.INT.udn b/Engine/Documentation/Source/Programming/Tutorials/UMG/1/UMG_1.INT.udn
new file mode 100644
--- /dev/null
+++ b/Engine/Documentation/Source/Programming/Tutorials/UMG/1/UMG_1.INT.udn
@@ -0,0 +1,95 @@
+Availability: Public
+Title:1. Set Up Module Dependencies For UMG
+Crumbs:%ROOT%, Programming, Programming/Tutorials, Programming/Tutorials/UMG
+Description:Modifying our project to include the modules needed to build menus.
+
+[VAR:Steps]
+[OBJECT:Navigation]
+	[PARAM:previous]
+	[/PARAM]
+	[PARAM:current]
+	[/PARAM]
+	[PARAM:home]
+		[Programming Tutorials Home](Programming\Tutorials\)
+	[/PARAM]
+	[PARAM:next]
+		[Next Step](Programming\Tutorials\UMG\2 "%Programming\Tutorials\UMG\2:title%")
+	[/PARAM]
+[/OBJECT]
+[/VAR]
+
+%Steps%
+
+[REGION:tip]If you are new to **Unreal Engine 4**, you might want to read our [%Programming\QuickStart:title% tutorial](Programming\QuickStart) first. For this tutorial, we will assume you are familiar with creating a project, adding C++ code to it, and compiling your code. We will also expose functions and properties to **Blueprint**. If you would like to learn more about that, the [%Programming\Tutorials\VariablesTimersEvents:title% tutorial](Programming\Tutorials\VariablesTimersEvents) is a good place to start.[/REGION]
+
+1. We will begin by creating a new project:
+	* Select the C++ tab.
+	* Select Basic Code.
+	* Leave the default values of Desktop/Console and With Starter Content.
+	* We will use the name "HowTo_UMG".
+
+	[![](UMG_ProjectSetup.png)](UMG_ProjectSetup.png)
+
+	[REGION:note]Since we're going to write code that uses **Unreal Motion Graphics (UMG)**, we'll need to go into **Visual Studio** include a few **modules** that aren't available by default.[/REGION]
+
+1. We can get to our project's code by using the **Open Visual Studio** command from the **File** dropdown menu on the main editor screen.
+
+	[![](OpenVisualStudio.png)](OpenVisualStudio.png)
+
+1. **UMG** is dependent on a few modules, so we'll need to add those modules to "HowTo_UMG.Build.cs".
+
+	[![](BuildCS.png)](BuildCS.png)
+
+	In "HowTo_UMG.Build.CS", we need to add "UMG" to the list of included public modules, and "Slate" and "SlateCore" to the list of included private modules. We'll need to modify the first line in HowTo_UMG's constructor to read as follows:
+		PublicDependencyModuleNames.AddRange(new string[] { "Core", "CoreUObject", "Engine", "InputCore", "UMG" });
+
+	And then we can uncomment the following line:
+		PrivateDependencyModuleNames.AddRange(new string[] { "Slate", "SlateCore" });
+
+	As usual, full source code is provided below so you can double-check your code.
+
+
+----
+
+With **UMG** set up, we can add code to our project's custom **Game Mode** to create and display our game's menus.
+
+[OBJECT:Section]
+    [PARAMLITERAL:id]
+    code
+    [/PARAMLITERAL]
+    [PARAM:heading]
+    Work-In-Progress Code
+    [/PARAM]
+    [PARAM:content]
+		**HowTo_UMG.Build.cs**
+			// Copyright 1998-2015 Epic Games, Inc. All Rights Reserved.
+
+			using UnrealBuildTool;
+
+			public class HowTo_UMG : ModuleRules
+			{
+				public HowTo_UMG(TargetInfo Target)
+				{
+					PublicDependencyModuleNames.AddRange(new string[] { "Core", "CoreUObject", "Engine", "InputCore", "UMG" });
+
+					//PrivateDependencyModuleNames.AddRange(new string[] {  });
+
+					// Uncomment if you are using Slate UI
+					PrivateDependencyModuleNames.AddRange(new string[] { "Slate", "SlateCore" });
+		
+					// Uncomment if you are using online features
+					// PrivateDependencyModuleNames.Add("OnlineSubsystem");
+					// if ((Target.Platform == UnrealTargetPlatform.Win32) || (Target.Platform == UnrealTargetPlatform.Win64))
+					// {
+					//		if (UEBuildConfiguration.bCompileSteamOSS == true)
+					//		{
+					//			DynamicallyLoadedModuleNames.Add("OnlineSubsystemSteam");
+					//		}
+					// }
+				}
+			}
+
+    [/PARAM]
+[/OBJECT]
+
+%Steps%
diff --git a/Engine/Documentation/Source/Programming/Tutorials/UMG/2/UMG_2.INT.udn b/Engine/Documentation/Source/Programming/Tutorials/UMG/2/UMG_2.INT.udn
new file mode 100644
--- /dev/null
+++ b/Engine/Documentation/Source/Programming/Tutorials/UMG/2/UMG_2.INT.udn
@@ -0,0 +1,186 @@
+Availability: Public
+Title:2. Extend Our Game Mode
+Crumbs:%ROOT%, Programming, Programming/Tutorials, Programming/Tutorials/UMG
+Description:Extending the project's Game Mode to add and remove menus from the screen.
+
+[VAR:Steps]
+[OBJECT:Navigation]
+	[PARAM:previous]
+		[Previous Step](Programming\Tutorials\UMG\1 "%Programming\Tutorials\UMG\1:title%")
+	[/PARAM]
+	[PARAM:current]
+	[/PARAM]
+	[PARAM:home]
+		[Programming Tutorials Home](Programming\Tutorials\)
+	[/PARAM]
+	[PARAM:next]
+		[Next Step](Programming\Tutorials\UMG\3 "%Programming\Tutorials\UMG\3:title%")
+	[/PARAM]
+[/OBJECT]
+[/VAR]
+
+%Steps%
+
+
+1. The menus we create will be made from **User Widgets**. We'll need to write a function that creates and displays a new **User Widget**, and then call that function when the game starts. We'll also need to keep track of what we have created so that we can remove it later. Since each project already comes with a custom **Game Mode** class, we can simply open ours, which is defined in "HowTo_UMGGameMode.h". The following functions and properties will need to be added to the bottom of the class:
+
+		public:
+			/** Called when the game starts. */
+			virtual void BeginPlay() override;
+
+			/** Remove the current menu widget and create a new one from the specified class, if provided. */
+			UFUNCTION(BlueprintCallable, Category = "UMG Game")
+			void ChangeMenuWidget(TSubclassOf<UUserWidget> NewWidgetClass);
+
+		protected:
+			/** The widget class we will use as our menu when the game starts. */
+			UPROPERTY(EditAnywhere, BlueprintReadOnly, Category = "UMG Game")
+			TSubclassOf<UUserWidget> StartingWidgetClass;
+
+			/** The widget instance that we are using as our menu. */
+			UPROPERTY()
+			UUserWidget* CurrentWidget;
+
+1. In order to use **User Widgets** in our code,  and add the following line to the top of the "#include" section:
+
+		#include "Blueprint/UserWidget.h"
+
+1. Moving to "HowTo_UMGGameMode.cpp" now, we need to fill out the bodies of the two functions we declared. We'll start with overriding BeginPlay():
+
+		void AHowTo_UMGGameMode::BeginPlay()
+		{
+			Super::BeginPlay();
+			ChangeMenuWidget(StartingWidgetClass);
+		}
+
+	[REGION:note]When overriding functions from a parent class (referenced by the word Super), as we do here with _BeginPlay_, it is often important to call the parent class' version of that function. Since our version of the function is only meant to add one step to the end of the existing procedure, we call _Super::BeginPlay_ in the first line of the function.[/REGION]
+
+1. Next, still in "HowTo_UMGGameMode.cpp", we need to define how we change between menus. We will need to remove whatever **User Widget** we have active from the viewport, if any, first. After that, we can create and a new **User Widget** and add it to the viewport.
+
+		void AHowTo_UMGGameMode::ChangeMenuWidget(TSubclassOf<UUserWidget> NewWidgetClass)
+		{
+			if (CurrentWidget != nullptr)
+			{
+				CurrentWidget->RemoveFromViewport();
+				CurrentWidget = nullptr;
+			}
+			if (NewWidgetClass != nullptr)
+			{
+				CurrentWidget = CreateWidget<UUserWidget>(GetWorld(), NewWidgetClass);
+				if (CurrentWidget != nullptr)
+				{
+					CurrentWidget->AddToViewport();
+				}
+			}
+		}
+
+	[REGION:note]This code will create instances of any **Widgets** we provide and put them on the screen. It also removes them, so that only one is active at a time, although **Unreal Engine** can handle displaying and interacting with many **Widgets** at once. We never need to destroy a **Widget** directly, because removing it from the viewport and clearing (or changing) all variables that reference it will cause it to be cleaned up by **Unreal Engine's** garbage collection system.[/REGION]
+
+
+----
+
+We have built the code framework to create and display menus, and remove them when they are no longer needed. We're ready to return to the **Unreal Editor** and design our menu assets!
+
+[OBJECT:Section]
+    [PARAMLITERAL:id]
+    code
+    [/PARAMLITERAL]
+    [PARAM:heading]
+    Finished Code
+    [/PARAM]
+    [PARAM:content]
+		**HowTo_UMG.Build.cs**
+			// Copyright 1998-2015 Epic Games, Inc. All Rights Reserved.
+
+			using UnrealBuildTool;
+
+			public class HowTo_UMG : ModuleRules
+			{
+				public HowTo_UMG(TargetInfo Target)
+				{
+					PublicDependencyModuleNames.AddRange(new string[] { "Core", "CoreUObject", "Engine", "InputCore", "UMG" });
+
+					//PrivateDependencyModuleNames.AddRange(new string[] {  });
+
+					// Uncomment if you are using Slate UI
+					PrivateDependencyModuleNames.AddRange(new string[] { "Slate", "SlateCore" });
+		
+					// Uncomment if you are using online features
+					// PrivateDependencyModuleNames.Add("OnlineSubsystem");
+					// if ((Target.Platform == UnrealTargetPlatform.Win32) || (Target.Platform == UnrealTargetPlatform.Win64))
+					// {
+					//		if (UEBuildConfiguration.bCompileSteamOSS == true)
+					//		{
+					//			DynamicallyLoadedModuleNames.Add("OnlineSubsystemSteam");
+					//		}
+					// }
+				}
+			}
+
+		**HowTo_UMGGameMode.h**
+			// Copyright 1998-2015 Epic Games, Inc. All Rights Reserved.
+
+			#pragma once
+
+			#include "Blueprint/UserWidget.h"
+			#include "GameFramework/GameMode.h"
+			#include "HowTo_UMGGameMode.generated.h"
+
+			/**
+			 * 
+			 */
+			UCLASS()
+			class HOWTO_UMG_API AHowTo_UMGGameMode : public AGameMode
+			{
+				GENERATED_BODY()
+
+			public:
+				/** Called when the game starts. */
+				virtual void BeginPlay() override;
+
+				/** Remove the current menu widget and create a new one from the specified class, if provided. */
+				UFUNCTION(BlueprintCallable, Category = "UMG Game")
+				void ChangeMenuWidget(TSubclassOf<UUserWidget> NewWidgetClass);
+
+			protected:
+				/** The widget class we will use as our menu when the game starts. */
+				UPROPERTY(EditAnywhere, BlueprintReadOnly, Category = "UMG Game")
+				TSubclassOf<UUserWidget> StartingWidgetClass;
+
+				/** The widget instance that we are using as our menu. */
+				UPROPERTY()
+				UUserWidget* CurrentWidget;
+			};
+
+		**HowTo_UMGGameMode.cpp**
+			// Copyright 1998-2015 Epic Games, Inc. All Rights Reserved.
+
+			#include "HowTo_UMG.h"
+			#include "HowTo_UMGGameMode.h"
+
+			void AHowTo_UMGGameMode::BeginPlay()
+			{
+				Super::BeginPlay();
+				ChangeMenuWidget(StartingWidgetClass);
+			}
+
+			void AHowTo_UMGGameMode::ChangeMenuWidget(TSubclassOf<UUserWidget> NewWidgetClass)
+			{
+				if (CurrentWidget != nullptr)
+				{
+					CurrentWidget->RemoveFromViewport();
+					CurrentWidget = nullptr;
+				}
+				if (NewWidgetClass != nullptr)
+				{
+					CurrentWidget = CreateWidget<UUserWidget>(GetWorld(), NewWidgetClass);
+					if (CurrentWidget != nullptr)
+					{
+						CurrentWidget->AddToViewport();
+					}
+				}
+			}
+    [/PARAM]
+[/OBJECT]
+
+%Steps%
diff --git a/Engine/Documentation/Source/Programming/Tutorials/UMG/3/UMG_3.INT.udn b/Engine/Documentation/Source/Programming/Tutorials/UMG/3/UMG_3.INT.udn
new file mode 100644
--- /dev/null
+++ b/Engine/Documentation/Source/Programming/Tutorials/UMG/3/UMG_3.INT.udn
@@ -0,0 +1,179 @@
+Availability: Public
+Title:3. Create Menu Widget Blueprints
+Crumbs:%ROOT%, Programming, Programming/Tutorials, Programming/Tutorials/UMG
+Description:Create widget blueprints to be used as menu screens.
+
+[VAR:Steps]
+[OBJECT:Navigation]
+	[PARAM:previous]
+		[Previous Step](Programming\Tutorials\UMG\2 "%Programming\Tutorials\UMG\2:title%")
+	[/PARAM]
+	[PARAM:current]
+	[/PARAM]
+	[PARAM:home]
+		[Programming Tutorials Home](Programming\Tutorials\)
+	[/PARAM]
+	[PARAM:next]
+		[Next Step](Programming\Tutorials\UMG\4 "%Programming\Tutorials\UMG\4:title%")
+	[/PARAM]
+[/OBJECT]
+[/VAR]
+
+%Steps%
+
+
+1. In **Unreal Editor**, we can press the **Compile** button to build our new code. This will allow us to use **User Widgets** as menus.
+
+	[![](CompileFromEditor.png)](CompileFromEditor.png)
+
+1. We will now create the **User Widgets** that our **Game Mode** will use as menus. This is done with the "Add New" button in the **Content Browser**. The **Widget Blueprint** class is found in the "User Interface" category. We need to create two of these, one named "MainMenu" and one named "NewGameMenu". Our game will begin at the Main Menu, and will have an option to proceed to the New Game Menu.
+
+	[![](CreateWidgetBP.png)](CreateWidgetBP.png)
+
+1. Double-clicking the "MainMenu" **Widget** we've just created will take us to the **Blueprint Designer**, where we can create our menu layout.
+
+1. Let's drag a **Button** from the "Common" section of the **Palette Panel** to the **Graph**. This **Button** will eventually be used to open the New Game Menu.
+
+	[![](AddButtonAndTextBlock.png)](AddButton.png)
+
+1. The first step in making our layout look right is to adjust the location and size of our **Button**. We should make the following changes:
+	* Set the size to 200x200.
+	* Set the position to (200, 100).
+	* Rename it "NewGameButton", just to make it easier to recoginze when we hook up functionality to it later on.
+
+	[![](ButtonDetails.png)](ButtonDetails.png)
+
+1. Since we're not drawing a custom set of images for the **Button**, we can label it by dropping a **Text Block** onto it and making the following changes:
+	* Set the text to say "New Game".
+	* Change the Visibility to Hit Test Invisible. This will prevent the **Text Block** from intercepting mouse-clicks that are intended for the **Button** underneath.
+	* Set the name to "NewGameText". This is not needed, but it's a good habit.
+
+	[![](TextBlockDetails.png)](TextBlockDetails.png)
+
+1. Next, we'll want to make a "Quit" feature with a second **Button** and **Text Block**. Set those up in the same way as the "New Game" **Button** and **Text Block**.
+
+1. After that, we can add **Events** to our **Buttons** so that we can run code when a **Button** is clicked. This is done by locating and pressing the "+" next to the appropriate **Event** name in the **Details Panel**. In this case, "OnClicked" is the only **Event** available. Create this event for both the NewGameButton amd QuitButton **Widgets**.
+
+	[![](CreateOnClickedEvent.png)](CreateOnClickedEvent.png)
+
+	[REGION:note]Designers can build functionality with **Blueprint** scripting here, or C++ programmers can connect nodes that call exposed functions.[/REGION]
+
+1. For the **Event** called "OnClicked(NewGameButton)", we'll want to:
+	* Connect a _ChangeMenuWidget_ node to use the function we added to our **Game Mode** earlier.
+	* Set the New Widget Class field on the _ChangeMenuWidget_ node to the "NewGameMenu" asset.
+
+	[![](OnClicked_NewGame.png)](OnClicked_NewGame.png)
+
+1. For the "OnClicked(QuitButton)" **Event**, we'll want to:
+	* Connect a "Quit Game" node.
+
+	[![](OnClicked_Quit.png)](OnClicked_Quit.png)
+
+
+----
+
+With our main menu built, we can set up a **Game Mode** asset that will load it as soon as the level starts.
+
+[OBJECT:Section]
+    [PARAMLITERAL:id]
+    code
+    [/PARAMLITERAL]
+    [PARAM:heading]
+    Finished Code
+    [/PARAM]
+    [PARAM:content]
+		**HowTo_UMG.Build.cs**
+			// Copyright 1998-2015 Epic Games, Inc. All Rights Reserved.
+
+			using UnrealBuildTool;
+
+			public class HowTo_UMG : ModuleRules
+			{
+				public HowTo_UMG(TargetInfo Target)
+				{
+					PublicDependencyModuleNames.AddRange(new string[] { "Core", "CoreUObject", "Engine", "InputCore", "UMG" });
+
+					//PrivateDependencyModuleNames.AddRange(new string[] {  });
+
+					// Uncomment if you are using Slate UI
+					PrivateDependencyModuleNames.AddRange(new string[] { "Slate", "SlateCore" });
+		
+					// Uncomment if you are using online features
+					// PrivateDependencyModuleNames.Add("OnlineSubsystem");
+					// if ((Target.Platform == UnrealTargetPlatform.Win32) || (Target.Platform == UnrealTargetPlatform.Win64))
+					// {
+					//		if (UEBuildConfiguration.bCompileSteamOSS == true)
+					//		{
+					//			DynamicallyLoadedModuleNames.Add("OnlineSubsystemSteam");
+					//		}
+					// }
+				}
+			}
+
+		**HowTo_UMGGameMode.h**
+			// Copyright 1998-2015 Epic Games, Inc. All Rights Reserved.
+
+			#pragma once
+
+			#include "Blueprint/UserWidget.h"
+			#include "GameFramework/GameMode.h"
+			#include "HowTo_UMGGameMode.generated.h"
+
+			/**
+			 * 
+			 */
+			UCLASS()
+			class HOWTO_UMG_API AHowTo_UMGGameMode : public AGameMode
+			{
+				GENERATED_BODY()
+
+			public:
+				/** Called when the game starts. */
+				virtual void BeginPlay() override;
+
+				/** Remove the current menu widget and create a new one from the specified class, if provided. */
+				UFUNCTION(BlueprintCallable, Category = "UMG Game")
+				void ChangeMenuWidget(TSubclassOf<UUserWidget> NewWidgetClass);
+
+			protected:
+				/** The widget class we will use as our menu when the game starts. */
+				UPROPERTY(EditAnywhere, BlueprintReadOnly, Category = "UMG Game")
+				TSubclassOf<UUserWidget> StartingWidgetClass;
+
+				/** The widget instance that we are using as our menu. */
+				UPROPERTY()
+				UUserWidget* CurrentWidget;
+			};
+
+		**HowTo_UMGGameMode.cpp**
+			// Copyright 1998-2015 Epic Games, Inc. All Rights Reserved.
+
+			#include "HowTo_UMG.h"
+			#include "HowTo_UMGGameMode.h"
+
+			void AHowTo_UMGGameMode::BeginPlay()
+			{
+				Super::BeginPlay();
+				ChangeMenuWidget(StartingWidgetClass);
+			}
+
+			void AHowTo_UMGGameMode::ChangeMenuWidget(TSubclassOf<UUserWidget> NewWidgetClass)
+			{
+				if (CurrentWidget != nullptr)
+				{
+					CurrentWidget->RemoveFromViewport();
+					CurrentWidget = nullptr;
+				}
+				if (NewWidgetClass != nullptr)
+				{
+					CurrentWidget = CreateWidget<UUserWidget>(GetWorld(), NewWidgetClass);
+					if (CurrentWidget != nullptr)
+					{
+						CurrentWidget->AddToViewport();
+					}
+				}
+			}
+    [/PARAM]
+[/OBJECT]
+
+%Steps%
diff --git a/Engine/Documentation/Source/Programming/Tutorials/UMG/4/UMG_4.INT.udn b/Engine/Documentation/Source/Programming/Tutorials/UMG/4/UMG_4.INT.udn
new file mode 100644
--- /dev/null
+++ b/Engine/Documentation/Source/Programming/Tutorials/UMG/4/UMG_4.INT.udn
@@ -0,0 +1,162 @@
+Availability: Public
+Title:4. Configure Our Game Mode
+Crumbs:%ROOT%, Programming, Programming/Tutorials, Programming/Tutorials/UMG
+Description:Set up our Game Mode to use our Main Menu Widget.
+
+[VAR:Steps]
+[OBJECT:Navigation]
+	[PARAM:previous]
+		[Previous Step](Programming\Tutorials\UMG\3 "%Programming\Tutorials\UMG\3:title%")
+	[/PARAM]
+	[PARAM:current]
+	[/PARAM]
+	[PARAM:home]
+		[Programming Tutorials Home](Programming\Tutorials\)
+	[/PARAM]
+	[PARAM:next]
+		[Next Step](Programming\Tutorials\UMG\5 "%Programming\Tutorials\UMG\5:title%")
+	[/PARAM]
+[/OBJECT]
+[/VAR]
+
+%Steps%
+
+
+1. In the **Content Browser**, we will add a **Blueprint Class** based on our project's **Game Mode**. This makes it possible to set the initial values of exposed variables to whatever we want. To do this:
+
+	* Click the **Add** button in the **Content Browser**.
+
+	[![](AddNewBPClass.png)](AddNewBPClass.png)
+
+	* Pick HowTo_UMGGameMode as the parent class. It will be listed in the "All Classes" section.
+
+	[![](PickParentClassForGameModeBP.png)](PickParentClassForGameModeBP.png)
+
+1. After naming the resulting **Blueprint**, which we will call "MenuGameMode" for this example, right-click it and choose "Edit".
+
+	[![](EditBP.png)](EditBP.png)
+
+1. We will configure the **Blueprint** of our **Game Mode** with two changes:
+	* The Starting Widget Class must be set to MainMenu in order to bring up our menu when the game begins.
+	* The Default Pawn Class should be set to **Pawn** instead of **DefaultPawn** so that our player will not be able to fly around while in the menu.
+
+	[![](ConfigureGameMode.png)](ConfigureGameMode.png)
+
+1. In order for our **Blueprint** to be used, we must return to the **Level Editor** window and change the **World Settings** for our current **Level** via the **Settings** button.
+
+	[![](WorldSettingsBar.png)](WorldSettingsBar.png)
+
+	[REGION:note]It is also possible to set the default **Game Mode** in the **Project Settings** menu, under the **Maps and Modes** section. If you use this method, all of your levels will default to the **Game Mode** you choose, unless overridden individually. Which method you use depends on how you prefer to set up your project.[/REGION]
+
+1. The **World Settings Panel** will open up. By default, it will be docked with the **Details Panel**, but it can be moved elsewhere. We need to set the Game Mode Override field to our "MenuGameMode" asset.
+
+	[![](WorldSettings.png)](WorldSettings.png)
+
+
+----
+
+Our custom **Game Mode** asset is now in effect on our level. If we run the game now, our Quit button will work as expected, and our New Game button will take us to an empty menu screen. Our next step will be to set up the New Game menu.
+
+[OBJECT:Section]
+    [PARAMLITERAL:id]
+    code
+    [/PARAMLITERAL]
+    [PARAM:heading]
+    Finished Code
+    [/PARAM]
+    [PARAM:content]
+		**HowTo_UMG.Build.cs**
+			// Copyright 1998-2015 Epic Games, Inc. All Rights Reserved.
+
+			using UnrealBuildTool;
+
+			public class HowTo_UMG : ModuleRules
+			{
+				public HowTo_UMG(TargetInfo Target)
+				{
+					PublicDependencyModuleNames.AddRange(new string[] { "Core", "CoreUObject", "Engine", "InputCore", "UMG" });
+
+					//PrivateDependencyModuleNames.AddRange(new string[] {  });
+
+					// Uncomment if you are using Slate UI
+					PrivateDependencyModuleNames.AddRange(new string[] { "Slate", "SlateCore" });
+		
+					// Uncomment if you are using online features
+					// PrivateDependencyModuleNames.Add("OnlineSubsystem");
+					// if ((Target.Platform == UnrealTargetPlatform.Win32) || (Target.Platform == UnrealTargetPlatform.Win64))
+					// {
+					//		if (UEBuildConfiguration.bCompileSteamOSS == true)
+					//		{
+					//			DynamicallyLoadedModuleNames.Add("OnlineSubsystemSteam");
+					//		}
+					// }
+				}
+			}
+
+		**HowTo_UMGGameMode.h**
+			// Copyright 1998-2015 Epic Games, Inc. All Rights Reserved.
+
+			#pragma once
+
+			#include "Blueprint/UserWidget.h"
+			#include "GameFramework/GameMode.h"
+			#include "HowTo_UMGGameMode.generated.h"
+
+			/**
+			 * 
+			 */
+			UCLASS()
+			class HOWTO_UMG_API AHowTo_UMGGameMode : public AGameMode
+			{
+				GENERATED_BODY()
+
+			public:
+				/** Called when the game starts. */
+				virtual void BeginPlay() override;
+
+				/** Remove the current menu widget and create a new one from the specified class, if provided. */
+				UFUNCTION(BlueprintCallable, Category = "UMG Game")
+				void ChangeMenuWidget(TSubclassOf<UUserWidget> NewWidgetClass);
+
+			protected:
+				/** The widget class we will use as our menu when the game starts. */
+				UPROPERTY(EditAnywhere, BlueprintReadOnly, Category = "UMG Game")
+				TSubclassOf<UUserWidget> StartingWidgetClass;
+
+				/** The widget instance that we are using as our menu. */
+				UPROPERTY()
+				UUserWidget* CurrentWidget;
+			};
+
+		**HowTo_UMGGameMode.cpp**
+			// Copyright 1998-2015 Epic Games, Inc. All Rights Reserved.
+
+			#include "HowTo_UMG.h"
+			#include "HowTo_UMGGameMode.h"
+
+			void AHowTo_UMGGameMode::BeginPlay()
+			{
+				Super::BeginPlay();
+				ChangeMenuWidget(StartingWidgetClass);
+			}
+
+			void AHowTo_UMGGameMode::ChangeMenuWidget(TSubclassOf<UUserWidget> NewWidgetClass)
+			{
+				if (CurrentWidget != nullptr)
+				{
+					CurrentWidget->RemoveFromViewport();
+					CurrentWidget = nullptr;
+				}
+				if (NewWidgetClass != nullptr)
+				{
+					CurrentWidget = CreateWidget<UUserWidget>(GetWorld(), NewWidgetClass);
+					if (CurrentWidget != nullptr)
+					{
+						CurrentWidget->AddToViewport();
+					}
+				}
+			}
+    [/PARAM]
+[/OBJECT]
+
+%Steps%
diff --git a/Engine/Documentation/Source/Programming/Tutorials/UMG/5/UMG_5.INT.udn b/Engine/Documentation/Source/Programming/Tutorials/UMG/5/UMG_5.INT.udn
new file mode 100644
--- /dev/null
+++ b/Engine/Documentation/Source/Programming/Tutorials/UMG/5/UMG_5.INT.udn
@@ -0,0 +1,167 @@
+Availability: Public
+Title:5. Build A Second Menu
+Crumbs:%ROOT%, Programming, Programming/Tutorials, Programming/Tutorials/UMG
+Description:Build a second menu with new functionality.
+
+[VAR:Steps]
+[OBJECT:Navigation]
+	[PARAM:previous]
+		[Previous Step](Programming\Tutorials\UMG\4 "%Programming\Tutorials\UMG\4:title%")
+	[/PARAM]
+	[PARAM:current]
+	[/PARAM]
+	[PARAM:home]
+		[Programming Tutorials Home](Programming\Tutorials\)
+	[/PARAM]
+	[PARAM:next]
+		[Next Step](Programming\Tutorials\UMG\6 "%Programming\Tutorials\UMG\6:title%")
+	[/PARAM]
+[/OBJECT]
+[/VAR]
+
+%Steps%
+
+
+1. In the **Content Browser**, find and open the "NewGameMenu" asset we created earlier. This menu will contain a name-entry **Text Box**, a **Button** to play the game that cannot be pressed until a name is entered, and a **Button** to return to the main menu.
+
+1. To create the name entry box, we'll drag a **Text Box** (not a **Text Block**) into the layout.
+
+	[![](CreateTextEntryBox.png)](CreateTextEntryBox.png)
+	
+1.  The **Text Box** should be configured with the following values:
+	* Position is (325, 200). This leaves room for a **Text Block** placed to the left of the **Text Box**.
+	* Size is 250x40.
+	* Font Size (under the "Style" heading) is 9.
+	
+	[![](TextBoxDetails.png)](TextBoxDetails.png)
+
+1. We can create the Play Game **Button** with a **Text Block** label the same way we created the **Buttons** on the previous menu.
+
+1. The Play Game **Button** will have a special feature - it will be enabled only when the name entered in the **Text Box** is not empty. We can use UMG's "bind" feature to create a new function for the "Is Enabled" field (under the Behavior section).
+
+	[![](PlayGameButtonDetails.png)](PlayGameButtonDetails.png)
+
+	[REGION:note]If there were complex rules for determining what constitutes a valid player name in our game, or if we needed to save the name to a C++ variable, we might expose a **UFUNCTION** on our **Game Mode** or as a static function somewhere within our project. However, since we only care that the name string is not empty, we can script it right here in our **Widget**.[/REGION]
+
+1. To ensure that the **Button** is enabled if and only if the **Text Box** is not empty, we can convert the text from the **Text Box** to a string and then check that its length is greater than zero. Here is how that logic would appear:
+
+	[![](PlayGameButtonEnableFunction.png)](PlayGameButtonEnableFunction.png)
+
+1. Let's add one more **Button** so we can back out and get to our Main Menu from here. This will be just like the Play Game **Button** from our Main Menu, but it will be positioned relative to the lower-right corner instead of the upper-left. To accomplish this, click the "Anchors" dropdown in the **Details Panel** for the **Button** and find the appropriate graphical representation in the pop-up menu.
+
+	[![](SelectAnchor.png)](SelectAnchor.png)
+
+	[REGION:tip]Positioning our anchor at the lower-right corner does not change how size and position values work, so we'll need to make our position values negative in order to be on the screen. Size values remain positive, but we will want to lower our position values to compensate.[/REGION]
+
+1. We will now add scripting to our new **Buttons**. The Main Menu **Button** will simply reload the Main Menu **Widget**, while the Play Game **Button** will deactivate our menu entirely by having no new **Widget** provided in the call to our _ChangeMenuWidget_ function. This is shown by the phrase "Select Class" being displayed instead of the name of an actual class or asset.
+
+	[![](NewGameButtonBPs.png)](NewGameButtonBPs.png)
+
+	[REGION:note]After deactivating the menu with the Play Game **Button**, we will be unable to do anything further in our game. This is the point at which we would normally load the first level, play an introductory cutscene, or spawn and possess a **Pawn**.[/REGION]
+
+
+----
+
+[OBJECT:Section]
+    [PARAMLITERAL:id]
+    code
+    [/PARAMLITERAL]
+    [PARAM:heading]
+    Finished Code
+    [/PARAM]
+    [PARAM:content]
+		**HowTo_UMG.Build.cs**
+			// Copyright 1998-2015 Epic Games, Inc. All Rights Reserved.
+
+			using UnrealBuildTool;
+
+			public class HowTo_UMG : ModuleRules
+			{
+				public HowTo_UMG(TargetInfo Target)
+				{
+					PublicDependencyModuleNames.AddRange(new string[] { "Core", "CoreUObject", "Engine", "InputCore", "UMG" });
+
+					//PrivateDependencyModuleNames.AddRange(new string[] {  });
+
+					// Uncomment if you are using Slate UI
+					PrivateDependencyModuleNames.AddRange(new string[] { "Slate", "SlateCore" });
+		
+					// Uncomment if you are using online features
+					// PrivateDependencyModuleNames.Add("OnlineSubsystem");
+					// if ((Target.Platform == UnrealTargetPlatform.Win32) || (Target.Platform == UnrealTargetPlatform.Win64))
+					// {
+					//		if (UEBuildConfiguration.bCompileSteamOSS == true)
+					//		{
+					//			DynamicallyLoadedModuleNames.Add("OnlineSubsystemSteam");
+					//		}
+					// }
+				}
+			}
+
+		**HowTo_UMGGameMode.h**
+			// Copyright 1998-2015 Epic Games, Inc. All Rights Reserved.
+
+			#pragma once
+
+			#include "Blueprint/UserWidget.h"
+			#include "GameFramework/GameMode.h"
+			#include "HowTo_UMGGameMode.generated.h"
+
+			/**
+			 * 
+			 */
+			UCLASS()
+			class HOWTO_UMG_API AHowTo_UMGGameMode : public AGameMode
+			{
+				GENERATED_BODY()
+
+			public:
+				/** Called when the game starts. */
+				virtual void BeginPlay() override;
+
+				/** Remove the current menu widget and create a new one from the specified class, if provided. */
+				UFUNCTION(BlueprintCallable, Category = "UMG Game")
+				void ChangeMenuWidget(TSubclassOf<UUserWidget> NewWidgetClass);
+
+			protected:
+				/** The widget class we will use as our menu when the game starts. */
+				UPROPERTY(EditAnywhere, BlueprintReadOnly, Category = "UMG Game")
+				TSubclassOf<UUserWidget> StartingWidgetClass;
+
+				/** The widget instance that we are using as our menu. */
+				UPROPERTY()
+				UUserWidget* CurrentWidget;
+			};
+
+		**HowTo_UMGGameMode.cpp**
+			// Copyright 1998-2015 Epic Games, Inc. All Rights Reserved.
+
+			#include "HowTo_UMG.h"
+			#include "HowTo_UMGGameMode.h"
+
+			void AHowTo_UMGGameMode::BeginPlay()
+			{
+				Super::BeginPlay();
+				ChangeMenuWidget(StartingWidgetClass);
+			}
+
+			void AHowTo_UMGGameMode::ChangeMenuWidget(TSubclassOf<UUserWidget> NewWidgetClass)
+			{
+				if (CurrentWidget != nullptr)
+				{
+					CurrentWidget->RemoveFromViewport();
+					CurrentWidget = nullptr;
+				}
+				if (NewWidgetClass != nullptr)
+				{
+					CurrentWidget = CreateWidget<UUserWidget>(GetWorld(), NewWidgetClass);
+					if (CurrentWidget != nullptr)
+					{
+						CurrentWidget->AddToViewport();
+					}
+				}
+			}
+    [/PARAM]
+[/OBJECT]
+
+%Steps%
diff --git a/Engine/Documentation/Source/Programming/Tutorials/UMG/6/UMG_6.INT.udn b/Engine/Documentation/Source/Programming/Tutorials/UMG/6/UMG_6.INT.udn
new file mode 100644
--- /dev/null
+++ b/Engine/Documentation/Source/Programming/Tutorials/UMG/6/UMG_6.INT.udn
@@ -0,0 +1,40 @@
+Availability: Public
+Title:6. On Your Own!
+Crumbs:%ROOT%, Programming, Programming/Tutorials, Programming/Tutorials/UMG
+Description:Ideas for further exploration!
+
+[VAR:Steps]
+[OBJECT:Navigation]
+	[PARAM:previous]
+		[Previous Step](Programming\Tutorials\UMG\5 "%Programming\Tutorials\UMG\5:title%")
+	[/PARAM]
+	[PARAM:current]
+	[/PARAM]
+	[PARAM:home]
+		[Programming Tutorials Home](Programming\Tutorials\)
+	[/PARAM]
+	[PARAM:next]
+	[/PARAM]
+[/OBJECT]
+[/VAR]
+
+%Steps%
+
+
+Using what you have learned, try to do the following:
+
+* Have multiple menus open at the same time.
+* Slide or fade **Widgets** onto the screen instead of just having them appear.
+* Look into the different sound, style, and color properties that **Widgets** support.
+* Set variables in your **Game Mode** such as difficulty level, character class, and so on from within menus.
+* Create an in-game menu that pops up and pauses the game when the player presses a key, then resumes gameplay when it is closed.
+* Link any functionality you've built or learned from a previous tutorial to a menu built in **UMG**.
+
+As for the specifics covered in this tutorial:
+
+* To learn about designing gameplay screens in UMG and using the many **Widgets** already available, read the [%Engine/UMG/UserGuide:title%](Engine/UMG/UserGuide) or the [%Engine/UMG/UserGuide/WidgetTypeReference:title%](Engine/UMG/UserGuide/WidgetTypeReference) page.
+* For more information about exposing variables and functions, and about using **Timers** and **Events** in your menu, try the [%Programming/Tutorials/VariablesTimersEvents:title%](Programming/Tutorials/VariablesTimersEvents) page.
+* For further tutorials, see the [%Programming/Tutorials:title%](Programming/Tutorials) page.
+
+
+%Steps%
diff --git a/Engine/Documentation/Source/Programming/Tutorials/UMG/UMG.INT.udn b/Engine/Documentation/Source/Programming/Tutorials/UMG/UMG.INT.udn
new file mode 100644
--- /dev/null
+++ b/Engine/Documentation/Source/Programming/Tutorials/UMG/UMG.INT.udn
@@ -0,0 +1,74 @@
+Availability: Public
+Title:User Interface With UMG
+Crumbs:%ROOT%, GettingStarted
+Description:Create a simple menu system using UMG.
+
+[VAR:Topic]
+[OBJECT:Topic]
+	[PARAM:image]
+		![](newclass_topic.png)
+	[/PARAM]
+	[PARAM:title]
+		%Programming\Tutorials\UMG:title%
+	[/PARAM]
+	[PARAM:path]
+		[RELATIVE:Programming\Tutorials\UMG]
+	[/PARAM]
+	[PARAM:description]
+		%Programming\Tutorials\UMG:description%
+	[/PARAM]
+[/OBJECT]
+[/VAR]
+
+[VAR:TopicCompact]
+[OBJECT:TopicCompact]
+	[PARAM:image]
+		![](newclass_topic.png)
+	[/PARAM]
+	[PARAM:title]
+		%Programming\Tutorials\UMG:title%
+	[/PARAM]
+	[PARAM:path]
+		[RELATIVE:Programming\Tutorials\UMG]
+	[/PARAM]
+	[PARAM:description]
+		%Programming\Tutorials\UMG:description%
+	[/PARAM]
+[/OBJECT]
+[/VAR]
+
+Over the course of this tutorial, you will learn to use Unreal Motion Graphics (UMG) to build a basic menu system with multiple screens and buttons.
+
+[OBJECT:EmbeddedVideo]
+	[PARAMLITERAL:width]
+	640
+	[/PARAMLITERAL]
+	[PARAMLITERAL:height]
+	360
+	[/PARAMLITERAL]
+	[PARAMLITERAL:videoid]
+	nTshVrOaNdY
+	[/PARAMLITERAL]
+[/OBJECT]
+
+[OBJECT:TopicButtonList]
+	[PARAM:icon]
+		![](%ROOT%/start_icon.png)(convert:false)
+	[/PARAM]
+	[PARAM:title]
+		Steps
+	[/PARAM]
+	[PARAM:description]
+	[/PARAM]
+	[PARAM:links]
+	* [](Programming\Tutorials\UMG\1 "%Programming\Tutorials\UMG\1:description%")
+	* [](Programming\Tutorials\UMG\2 "%Programming\Tutorials\UMG\2:description%")
+	* [](Programming\Tutorials\UMG\3 "%Programming\Tutorials\UMG\3:description%")
+	* [](Programming\Tutorials\UMG\4 "%Programming\Tutorials\UMG\4:description%")
+	* [](Programming\Tutorials\UMG\5 "%Programming\Tutorials\UMG\5:description%")
+	* [](Programming\Tutorials\UMG\6 "%Programming\Tutorials\UMG\6:description%")
+	[/PARAM]
+[/OBJECT]
+
+
+
