#UE4doc #docs: corrected "set up" text , removed some redundant phrasing

diff --git a/Engine/Documentation/Source/Engine/Rendering/Materials/ParameterCollections/ParemeterCollections.INT.udn b/Engine/Documentation/Source/Engine/Rendering/Materials/ParameterCollections/ParemeterCollections.INT.udn
--- a/Engine/Documentation/Source/Engine/Rendering/Materials/ParameterCollections/ParemeterCollections.INT.udn
+++ b/Engine/Documentation/Source/Engine/Rendering/Materials/ParameterCollections/ParemeterCollections.INT.udn
@@ -148,9 +148,9 @@ This is an extremely powerful way to manipulate Material Collection Parameter as
 	When completed, your Blueprint Graph should look something like this. 
 	![](MPC_Added_Nodes_To_BP.png)
 
-1. With the correct nodes added to the Blueprint Graph, it is now time to setup them up. 
-	The first thing that we are going to setup is going to be the  ** Set Scalar Parameter Value**. 
-	To set this up all we have to do is first define what Material Parameter Collection we are going to use. 
+1. With the correct nodes added to the Blueprint Graph, it is now time to set them up. 
+	The first thing that we are going to set up is the  ** Set Scalar Parameter Value**. 
+	To set this up we have to first define what Material Parameter Collection we are going to use. 
 	To do this, use your **Left Mouse Button** to click on the **Collection** input and then select which Material Parameter Collection you would wish to use.
 	Next, select the **Parameter Name** you wish to manipulate, in this case it will be **Global_Desaturation**. When completed, your **Set Scalar Parameter Value** node should look like this. 
 
