#UE4doc #docs: capitalized "Main Menu Bar"

diff --git a/Engine/Documentation/Source/Engine/Rendering/PostProcessEffects/PostProcessMaterials/PostProcessDemoMaterialOverview/PostProcessDemoMaterialOverview.INT.udn b/Engine/Documentation/Source/Engine/Rendering/PostProcessEffects/PostProcessMaterials/PostProcessDemoMaterialOverview/PostProcessDemoMaterialOverview.INT.udn
--- a/Engine/Documentation/Source/Engine/Rendering/PostProcessEffects/PostProcessMaterials/PostProcessDemoMaterialOverview/PostProcessDemoMaterialOverview.INT.udn
+++ b/Engine/Documentation/Source/Engine/Rendering/PostProcessEffects/PostProcessMaterials/PostProcessDemoMaterialOverview/PostProcessDemoMaterialOverview.INT.udn
@@ -55,7 +55,7 @@ w:450
 
 Below are some basic instructions on how to set up a simple Post Process Material:
 
-1. Create a new Level by choosing **File > New Level** from the main menubar.
+1. Create a new Level by choosing **File > New Level** from the Main Menu Bar.
 
 1. Create a new Material by clicking the ![](CB_button_NewAsset.png) button in the **Content Browser** and choosing **Material**. Provide a name for the Material.
 
@@ -67,16 +67,16 @@ Below are some basic instructions on how to set up a simple Post Process Materia
 
 	![](PostMaterialProperties.png)
 
-1. In the **Post Process Material** category, be sure to set the **Blendable Location** property to _Before Tonemapping_. This makes for more of a performance hit, but prevents ghosting of the scanlines while the view is in motion. For more information about the Blendable Location property, see the [Critical Settings for Post Process Materials](Engine/Rendering\PostProcessEffects\PostProcessMaterials#CriticalSettingsforPostProcessMaterials) section below. 
+1. In the **Post Process Material** category, be sure to set the **Blendable Location** property to _Before Tonemapping_. This makes for more of a performance hit, but prevents ghosting of the scan lines while the view is in motion. For more information about the Blendable Location property, see the [Critical Settings for Post Process Materials](Engine/Rendering\PostProcessEffects\PostProcessMaterials#CriticalSettingsforPostProcessMaterials) section below. 
 
 	![](PostProcessMaterialsCategory.png)
 
-1. At this point, you would create some Material expression network that defined the look of your post process. As a test for this example, we will create a video scanline overlay effect. You can really make just about anything you like, however. If you would like an example of a cool video scanline post process material, see the [Material Setup](#MaterialSetup) section.
+1. At this point, you would create some Material expression network that defined the look of your post process. As a test for this example, we will create a video scan line overlay effect. You can really make just about anything you like, however. If you would like an example of a cool video scan line post process material, see the [Material Setup](#MaterialSetup) section.
 
 
 ## Material Setup
 
-This document is an overview and breakdown of simple video scanline Post Process Material.
+This document is an overview and breakdown of simple video scan line Post Process Material.
 
 [REGION:fullwidth]
 ![](VideoScanlineEffect.png)
@@ -106,7 +106,7 @@ To facilitate putting this Material together, we have broken the key elements in
 
 1. We start with the Scene Texture expression, which is set to bring in the Scene Color. This is how we will manipulate the scene via post process.
 1. We perturb the UVs of the Scene Color using a simple panning texture (ScanDistort.png). This is done by multiplying the result of the texture by 0.03 and then combining it only with the red channel of the texture coordinates, effectively sliding the pixels across the screen. This creates a distortion effect.
-1. We bring in our 4x4 pixel scanline texture and tile it very heavily in the vertical direction (v:128). We also pan it downward very slowly.
+1. We bring in our 4x4 pixel scan line texture and tile it very heavily in the vertical direction (v:128). We also pan it downward very slowly.
 1. We reuse the same scanline texture, this time tiling it only by 1.28 and panning it upward, but a little more quickly. We also clamp it between 0.1 and 0.5. This will be used to create some variation in the scanlines.
 1. Create a time-driven sine wave that runs between 0.6 and 1.0, at a frequency of 10 Hz. This will be used to flicker the scanlines.
 1. Take the noise texture and tile and pan it at 2 different amounts and speeds for variation.
