#UE4doc #docs: sp: noticeable

diff --git a/Engine/Documentation/Source/Engine/Rendering/PostProcessEffects/SceneFringe/SceneFringe.INT.udn b/Engine/Documentation/Source/Engine/Rendering/PostProcessEffects/SceneFringe/SceneFringe.INT.udn
--- a/Engine/Documentation/Source/Engine/Rendering/PostProcessEffects/SceneFringe/SceneFringe.INT.udn
+++ b/Engine/Documentation/Source/Engine/Rendering/PostProcessEffects/SceneFringe/SceneFringe.INT.udn
@@ -4,7 +4,7 @@ Crumbs: %ROOT%, Engine, Engine/Rendering/PostProcessEffects
 Description:Chromatic aberration effect that simulates the color shifts near the edges of real-world camera lenses.
 
 [EXCERPT:SceneFringe_Desc]
-**Scene Fringe** is a chromatic aberration effect that simulates the color shifts in real-world camera lenses. The effect is mostly noticable near the edges of the image.
+**Scene Fringe** is a chromatic aberration effect that simulates the color shifts in real-world camera lenses. The effect is mostly noticeable near the edges of the image.
 [/EXCERPT:SceneFringe_Desc]
 
 [REGION:imagetable]
