#UE4doc #docs: sp: Quality

diff --git a/Engine/Documentation/Source/Engine/Rendering/PostProcessEffects/ScreenSpaceReflection/ScreenSpaceReflection.INT.udn b/Engine/Documentation/Source/Engine/Rendering/PostProcessEffects/ScreenSpaceReflection/ScreenSpaceReflection.INT.udn
--- a/Engine/Documentation/Source/Engine/Rendering/PostProcessEffects/ScreenSpaceReflection/ScreenSpaceReflection.INT.udn
+++ b/Engine/Documentation/Source/Engine/Rendering/PostProcessEffects/ScreenSpaceReflection/ScreenSpaceReflection.INT.udn
@@ -13,5 +13,5 @@ Description:Effect which alters the reflection that appear on the surface of mat
 | Property | Description |
 | -------- | ----------- |
 | **Intensity** | Enable/Fade/Disable the Screen Space Reflection feature by percentage (avoid numbers between 0 and 1 for consistency). |
-| **Quality** | 0=Lowest Quailty while 100=Maximum Quaility (50 is the default to provide better performance). |
+| **Quality** | 0=Lowest Quality while 100=Maximum Quality (50 is the default to provide better performance). |
 | **Max Roughness** | Used to determine what roughness we fade the Screen Space Reflections (0.8 works well, smaller can run faster). |
