#UE4doc #docs: sp: analyze, designed

diff --git a/Engine/Documentation/Source/Engine/Tools/CrashReporter/CrashReporter.INT.udn b/Engine/Documentation/Source/Engine/Tools/CrashReporter/CrashReporter.INT.udn
--- a/Engine/Documentation/Source/Engine/Tools/CrashReporter/CrashReporter.INT.udn
+++ b/Engine/Documentation/Source/Engine/Tools/CrashReporter/CrashReporter.INT.udn
@@ -5,7 +5,7 @@ Description:Application for collecting information from crashes in Unreal Engine
 
 [TOC (start:2 end:3)]
 
-The Crash Reporter suite of programs is desgined to capture and analyse application crashes both internal to the company, and externally from end users.
+The Crash Reporter suite of programs is designed to capture and analyze application crashes both internal to the company, and externally from end users.
 
 ## Requirements
 
