#UE4doc #docs: sp: designers, palette, Example

diff --git a/Engine/Documentation/Source/Engine/UI/ColorPicker/ColorPicker.INT.udn b/Engine/Documentation/Source/Engine/UI/ColorPicker/ColorPicker.INT.udn
--- a/Engine/Documentation/Source/Engine/UI/ColorPicker/ColorPicker.INT.udn
+++ b/Engine/Documentation/Source/Engine/UI/ColorPicker/ColorPicker.INT.udn
@@ -72,9 +72,9 @@ The **Color Picker** allows you to easily adjust a color value for a color prope
 
 ## Themes
 
-**Themes** are collections of colors that can be easily reused. These can be used to make commonly used colors available or to restrict the colors used by artists and deisgners to a specific pallette.
+**Themes** are collections of colors that can be easily reused. These can be used to make commonly used colors available or to restrict the colors used by artists and designers to a specific palette.
 
-![Exmaple Theme](theme_spring.png)
+![Example Theme](theme_spring.png)
 
 ### Creating Themes
 
