#UE4doc #docs: removed unneeded apostrophe, sp: cannot

diff --git a/Engine/Documentation/Source/Engine/UI/ClassViewer/ClassViewer.INT.udn b/Engine/Documentation/Source/Engine/UI/ClassViewer/ClassViewer.INT.udn
--- a/Engine/Documentation/Source/Engine/UI/ClassViewer/ClassViewer.INT.udn
+++ b/Engine/Documentation/Source/Engine/UI/ClassViewer/ClassViewer.INT.udn
@@ -105,7 +105,7 @@ The context menu is opened by clicking on ![Context Menu Button](ContextMenuButt
 |Class Type|Options|Description|
 |----------|-------|-----------|
 |**Blueprint Class**|**Create Blueprint...**<br>**Edit Blueprint...**<br>**Find In Content Browser**|Creates a new Blueprint parented to the selected Blueprint.<br>Opens the selected Blueprint in the [Blueprint Editor](Engine/Blueprints).<br>Finds the Blueprint Actor in the [Content Browser](Engine/Content/Browser).|
-|**Code Class**|**Create Blueprint...**<br>**Open C++ Header File...**<br>**Create New C++ Class...**|Creates a new Blueprint parented to the selected Blueprint.<br>Opens up the class' header file in Visual Studio.<br>Opens the [](Programming\Development\ManagingGameCode\CppClassWizard) to create a new class with the selected class as the Parent class.|
+|**Code Class**|**Create Blueprint...**<br>**Open C++ Header File...**<br>**Create New C++ Class...**|Creates a new Blueprint parented to the selected Blueprint.<br>Opens up the class header file in Visual Studio.<br>Opens the [](Programming\Development\ManagingGameCode\CppClassWizard) to create a new class with the selected class as the Parent class.|
 
 
 ## Drag & Dropping
@@ -124,7 +124,7 @@ Classes that have not been loaded will not appear in the combo box. Dropping a c
 The Class Picker is a mode that the Class Viewer can be switched to using only code. It is used for displaying a list of available classes such as for converting Static Meshes or selecting a parent for a new Blueprint. There are a number of ways to configure a Class Picker to show the desired results.
 
 [REGION:tip]
-When putting a Class Viewer or Class Picker into a location, make sure its height is limited in some manner. Thousands of classes may be in the Viewer/Picker and unconstrained it will stretch off the screen trying to render all of them. It will automatically display a scrollbar if all items can not be displayed in the available area.
+When putting a Class Viewer or Class Picker into a location, make sure its height is limited in some manner. Thousands of classes may be in the Viewer/Picker and unconstrained it will stretch off the screen trying to render all of them. It will automatically display a scrollbar if all items cannot be displayed in the available area.
 [/REGION]
 
 ### Class Picker Options
