#UE4doc #docs: duplicate word: the

diff --git a/Engine/Documentation/Source/Engine/UI/GlobalAssetPicker/GlobalAssetPicker.INT.udn b/Engine/Documentation/Source/Engine/UI/GlobalAssetPicker/GlobalAssetPicker.INT.udn
--- a/Engine/Documentation/Source/Engine/UI/GlobalAssetPicker/GlobalAssetPicker.INT.udn
+++ b/Engine/Documentation/Source/Engine/UI/GlobalAssetPicker/GlobalAssetPicker.INT.udn
@@ -48,7 +48,7 @@ The **Global Asset Picker** (**Ctrl-P**) is a fast-access means of editing asset
 
 While mousing over any asset, all of its relevant data will appear as a **tool tip** at the bottom of the window.
 
-**View Options** at the bottom of the window is exactly the same as the one on the **Content Browser**, allowing you to change the the way it displays and scales the assets. The thumbnail editor does not work in this window.
+**View Options** at the bottom of the window is exactly the same as the one on the **Content Browser**, allowing you to change the way it displays and scales the assets. The thumbnail editor does not work in this window.
 
 Currently, you cannot pull up the **Right-Click Menu** on anything in here, but you can **Click and Drag** things into your scene or you can **double-click** assets to bring up their associated editors.
 
