#UE4doc #docs: removed unneeded commas

diff --git a/Engine/Documentation/Source/Engine/UI/LevelEditor/Details/Customizations/Materials/MaterialsCategory.INT.udn b/Engine/Documentation/Source/Engine/UI/LevelEditor/Details/Customizations/Materials/MaterialsCategory.INT.udn
--- a/Engine/Documentation/Source/Engine/UI/LevelEditor/Details/Customizations/Materials/MaterialsCategory.INT.udn
+++ b/Engine/Documentation/Source/Engine/UI/LevelEditor/Details/Customizations/Materials/MaterialsCategory.INT.udn
@@ -10,7 +10,7 @@ For example, a Static Mesh with 3 elements would display 3 Materials when select
 
 | Control | Description |
 | ------- | ----------- |
-| ![](Details_Mat_ObjRef.png)| Standard [Object reference](Engine/UI/LevelEditor/Details/Properties/ObjectReference) property controls for setting, finding, and resetting the assigned Material back to the base Material. Clicking the drop-down list that contains the Material's name, and then clicking **Edit** on the context menu, opens the Material in the Material Instance Editor for modification. |
+| ![](Details_Mat_ObjRef.png)| Standard [Object reference](Engine/UI/LevelEditor/Details/Properties/ObjectReference) property controls for setting, finding, and resetting the assigned Material back to the base Material. Clicking the drop-down list that contains the Material's name and then clicking **Edit** on the context menu opens the Material in the Material Instance Editor for modification. |
 | ![Material Thumbnail](button_Details_Mat_Image.png)| Thumbnail of the Material, which can be **double-clicked** to open the Material Instance in the Material Instance Editor for modification. |
 | ![Find Texture](button_textures.png) | List of Textures used within the Material. Selecting a Texture from the list finds that Texture in the **Content Browser**. |
 
