#UE4doc #docs: duplicate word: can

diff --git a/Engine/Documentation/Source/Engine/UI/LevelEditor/LevelEditor.INT.udn b/Engine/Documentation/Source/Engine/UI/LevelEditor/LevelEditor.INT.udn
--- a/Engine/Documentation/Source/Engine/UI/LevelEditor/LevelEditor.INT.udn
+++ b/Engine/Documentation/Source/Engine/UI/LevelEditor/LevelEditor.INT.udn
@@ -181,7 +181,7 @@ The **Viewport** panel is your window into the worlds you create in UnrealEd.
 
 ![Viewport Panel](Engine/UI/LevelEditor/Viewports/viewports_LE.png)(w:720)
 
-This panel contains a set of viewports, each of which can can be maximized to fill the entire panel and offer the ability to display the world from one of three orthographic views (Top, Side, Front) or a perspective view giving you complete control over what you see as well as how you see it.
+This panel contains a set of viewports, each of which can be maximized to fill the entire panel and offer the ability to display the world from one of three orthographic views (Top, Side, Front) or a perspective view giving you complete control over what you see as well as how you see it.
 
 See [Viewports](Engine/UI/LevelEditor\Viewports) for more information on working with viewports.
 
