#UE4doc #docs: sp: editor's

diff --git a/Engine/Documentation/Source/Engine/UI/InterfaceOverview/InterfaceOverview.INT.udn b/Engine/Documentation/Source/Engine/UI/InterfaceOverview/InterfaceOverview.INT.udn
--- a/Engine/Documentation/Source/Engine/UI/InterfaceOverview/InterfaceOverview.INT.udn
+++ b/Engine/Documentation/Source/Engine/UI/InterfaceOverview/InterfaceOverview.INT.udn
@@ -72,7 +72,7 @@ but only a single instance can be editing a specific asset at any time.
 General editors are opened through the **Window** menu, while asset editors are opened by choosing to edit a particular 
 type of asset. The **Level Editor** is a special case and is always present while Unreal Editor is open.
 
-Editors can be closed by clicking the ![Close App Button](button_close_app.png) button on the editors's tab or by **right-
+Editors can be closed by clicking the ![Close App Button](button_close_app.png) button on the editor's tab or by **right-
 clicking** and choose **Close**.
 
 ![](Tabs.png)
