#UE4doc #docs: duplicate word: the

diff --git a/Engine/Documentation/Source/Engine/UI/LevelEditor/Details/Customizations/Transform/TransformCategory.INT.udn b/Engine/Documentation/Source/Engine/UI/LevelEditor/Details/Customizations/Transform/TransformCategory.INT.udn
--- a/Engine/Documentation/Source/Engine/UI/LevelEditor/Details/Customizations/Transform/TransformCategory.INT.udn
+++ b/Engine/Documentation/Source/Engine/UI/LevelEditor/Details/Customizations/Transform/TransformCategory.INT.udn
@@ -20,7 +20,7 @@ the scale values to each other are maintained when any of the individual values
 
 The transform properties default to _relative_ transformation. This means the transform is relative to the component's parent. 
 Each of the property labels are hyperlinks that can be clicked to toggle between _absolute_ and _relative_ transforms. When 
-using _absolute_ transformations, the transform is relative to the world instead of the the parent.
+using _absolute_ transformations, the transform is relative to the world instead of the parent.
 
 
 
