Updated project file generator documentation
- Added "-2015" command-line argument to docs

diff --git a/Engine/Documentation/Source/Programming/UnrealBuildSystem/ProjectFileGenerator/ProjectFileGenerator.INT.udn b/Engine/Documentation/Source/Programming/UnrealBuildSystem/ProjectFileGenerator/ProjectFileGenerator.INT.udn
--- a/Engine/Documentation/Source/Programming/UnrealBuildSystem/ProjectFileGenerator/ProjectFileGenerator.INT.udn
+++ b/Engine/Documentation/Source/Programming/UnrealBuildSystem/ProjectFileGenerator/ProjectFileGenerator.INT.udn
@@ -42,6 +42,8 @@ The project generator has some optional command-line arguments that you can use
  ----------------- |  ------------ |
 -CurrentPlatform   |  Generates buildable projects for only the current desktop platform (Windows or Mac), instead of all detected available target platforms. |
 -2012              |  Generates projects in a format native to Visual Studio 2012.  By default, Visual Studio 2013 project files are generated. |
+-2013              |  Generates projects in a format native to Visual Studio 2013.  This is the default. |
+-2015              |  Generates projects in a format native to Visual Studio 2015.  The Visual C++ 2015 compiler will be used to compile when this option is set. |
 -ThirdParty        |  Adds headers and other files from third-party libraries to the project.  This could be useful if you like to see symbols and files in Visual Assist for PhysX, Steamworks or Bink for example.  However, it will make the project files take longer to load! |
 -Game `GameName`   |	 Tells the project generator to generate projects that only include code and content for the specified project name, excluding all other discovered projects.  Make sure to specify a project name as well.  For example `GenerateProjectFiles.bat ShooterGame.uproject -Game` would generate projects that only have source and targets for ShooterGame.  If you want to include Engine source code and Programs alongside a single game project, you can also pass the `-Engine` parameter. |
 -Engine            |  When used with `-Game`, causes Engine code, content, and programs to also be included in the generated solution.  This is useful if you think you will need to work with engine source a lot while working on your game project. |
