#UE4doc #docs: sp: compiled, Opens, visibility

diff --git a/Engine/Documentation/Source/Engine/UI/LevelEditor/MenuBar/MainMenuBar.INT.udn b/Engine/Documentation/Source/Engine/UI/LevelEditor/MenuBar/MainMenuBar.INT.udn
--- a/Engine/Documentation/Source/Engine/UI/LevelEditor/MenuBar/MainMenuBar.INT.udn
+++ b/Engine/Documentation/Source/Engine/UI/LevelEditor/MenuBar/MainMenuBar.INT.udn
@@ -26,7 +26,7 @@ The **Menu Bar** in editor should be familiar to anyone who has used Windows app
 | [REGION:tablesection]**Project**[/REGION] ||
 | **New Project** | Opens a dialog to create a new game project. |
 | **Open Project** | Opens a dialog to choose a game project to open. |
-| **Add Code To Project** | Adds C++ code to the project. The code can only be compilied if you have Visual Studio installed. |
+| **Add Code To Project** | Adds C++ code to the project. The code can only be compiled if you have Visual Studio installed. |
 | **Cook Content** | Cooks your project's content for a specific platform. This is useful if you run a code-based project from Visual Studio. |
 | **Package Project** | Compile, cook and package your project and its content for distribution. |
 | **Open Visual Studio** | Opens your C++ code in Visual Studio. |
@@ -74,7 +74,7 @@ The **Menu Bar** in editor should be familiar to anyone who has used Windows app
 |**Details**| Shows the **Details** panel, which provides access to various properties for selected objects. You may show up to 4 **Details** panels simultaneously.|
 |**Viewports**| Shows the Viewport, which is your view into the 3D world of your game. You may show up to 4 Viewports simultaneously.|
 |**Build and Submit**| Opens the **Build and Submit** panel, which allows you to build lighting for your level, save the level, then submit it to source control.|
-| **Layers** | Opens the **Layers** panel, where you can divide your level into a series of visiblity layers.|
+| **Layers** | Opens the **Layers** panel, where you can divide your level into a series of visibility layers.|
 | **Levels** | Opens the **Levels** panel which displays the persistent and all streaming levels. |
 | **Modes** | Shows the **Modes** panel, which accesses special level editing tools such as Landscape and Foliage. |
 | **World Outliner** | Opens the **World Outliner**, a hierarchical breakdown of all the Actors currently in the level.|
@@ -93,7 +93,7 @@ The **Menu Bar** in editor should be familiar to anyone who has used Windows app
 | **Output Log** | Opens a new **Output Log**. Only one **Output Log** may exist at any given time. |
 | **Plugins** | Open the Plugins Browser tab. |
 | [REGION:tablesection]**Unreal Frontend**[/REGION] ||
-| **Device Manager** | Opoen the **Device Manager** tab. |
+| **Device Manager** | Opens the **Device Manager** tab. |
 | **Session Frontend** | Opens the [](Engine/Deployment/UnrealFrontend) tool which provides access to the [Automation System](Programming/Automation), Unreal Console, and more. |
 | [REGION:tablesection]**Layout**[/REGION] ||
 | **Reset Layout** | Makes a backup of your user settings and resets the layout customizations. |
