#UE4doc #docs: sp: Placing

diff --git a/Engine/Documentation/Source/Engine/UI/LevelEditor/Modes/PlaceMode/PlacementBrowser.INT.udn b/Engine/Documentation/Source/Engine/UI/LevelEditor/Modes/PlaceMode/PlacementBrowser.INT.udn
--- a/Engine/Documentation/Source/Engine/UI/LevelEditor/Modes/PlaceMode/PlacementBrowser.INT.udn
+++ b/Engine/Documentation/Source/Engine/UI/LevelEditor/Modes/PlaceMode/PlacementBrowser.INT.udn
@@ -86,7 +86,7 @@ These buttons allow you to quickly filter down the asset type that you need.
 
 The Asset View shows icons for each of the various assets currently available for placement, after all filters and search criteria have been applied. The area is scrollable by right-mouse dragging or using the scrollbar on the right of the panel.
 
-## Placeing Actors in Place Mode
+## Placing Actors in Place Mode
 
 [INCLUDE:Engine/Actors/Placement#placemode]
 
