#UE4 Docs: Fixing a smart quote.

diff --git a/Engine/Documentation/Source/Engine/Rendering/Materials/HowTo/MaterialHowTo.INT.udn b/Engine/Documentation/Source/Engine/Rendering/Materials/HowTo/MaterialHowTo.INT.udn
--- a/Engine/Documentation/Source/Engine/Rendering/Materials/HowTo/MaterialHowTo.INT.udn
+++ b/Engine/Documentation/Source/Engine/Rendering/Materials/HowTo/MaterialHowTo.INT.udn
@@ -31,7 +31,7 @@ Related: Engine/Rendering/Materials/MaterialProperties
 ## Overview
 
 The following How-To tutorials will demonstrate how you can utilize the full potential of the Material Editor in Unreal Engine 4 (UE4).
-From tips on How - To correctly set up Human Skin to How - To use Layered Materials, the How - To tutorials are a must read for anyone who wants to know more about UE4’s Material Editor.
+From tips on How - To correctly set up Human Skin to How - To use Layered Materials, the How - To tutorials are a must read for anyone who wants to know more about UE4's Material Editor.
 
 
 ## Creating a New Project
