TypoCorrection

diff --git a/Engine/Documentation/Source/Engine/UMG/QuickStart/4/UMGQS_4.INT.udn b/Engine/Documentation/Source/Engine/UMG/QuickStart/4/UMGQS_4.INT.udn
--- a/Engine/Documentation/Source/Engine/UMG/QuickStart/4/UMGQS_4.INT.udn
+++ b/Engine/Documentation/Source/Engine/UMG/QuickStart/4/UMGQS_4.INT.udn
@@ -1,4 +1,4 @@
-Availability: Public
+﻿Availability: Public
 Title:4. Scripting the Main Menu
 Crumbs:%ROOT%, Engine, Engine/UMG, Engine/UMG/QuickStart
 Description:In this step we will script the functionality for our Main Menu and load into game. 
@@ -104,7 +104,7 @@ Time to get our hands dirty and script the functionality of our Main Menu. At th
 
 	![](Overrides.png)
 
-	Here we are assigning a Game Mode for this Map and changing the player's character to use the Engine Default Character inside of the playable **FirstPersonCharacter** Blueprint as we do not what the player to be able to run around and shoot behind the Main Menu. 
+	Here we are assigning a Game Mode for this Map and changing the player's character to use the Engine Default Character inside of the playable **FirstPersonCharacter** Blueprint as we do not want the player to be able to run around and shoot behind the Main Menu. 
 
 1.	Click the **Play** Button from the Main Toolbar to play in the editor. 
 
