TypoCorrection

diff --git a/Engine/Documentation/Source/Engine/UMG/QuickStart/5/UMGQS_5.INT.udn b/Engine/Documentation/Source/Engine/UMG/QuickStart/5/UMGQS_5.INT.udn
--- a/Engine/Documentation/Source/Engine/UMG/QuickStart/5/UMGQS_5.INT.udn
+++ b/Engine/Documentation/Source/Engine/UMG/QuickStart/5/UMGQS_5.INT.udn
@@ -1,4 +1,4 @@
-Availability: Public
+﻿Availability: Public
 Title:5. Creating / Scripting an In-Game Pause Menu
 Crumbs:%ROOT%, Engine, Engine/UMG, Engine/UMG/QuickStart
 Description:Here we will create an in-game pause menu which will allow us to pause or exit our game. 
@@ -7,7 +7,7 @@ Description:Here we will create an in-game pause menu which will allow us to pau
 [VAR:Steps]
 [OBJECT:Navigation]
 	[PARAM:previous]
-		[Previous Step](Engine/UMG\QuickStart\5 "%Engine/UMG\QuickStart\5:title%") 
+		[Previous Step](Engine/UMG\QuickStart\4 "%Engine/UMG\QuickStart\4:title%") 
 	[/PARAM]
 	[PARAM:current]
 	[/PARAM]
