#UE4doc #docs: sp: devices

diff --git a/Engine/Documentation/Source/Engine/UI/LevelEditor/Toolbar/LevelEditorToolbar.INT.udn b/Engine/Documentation/Source/Engine/UI/LevelEditor/Toolbar/LevelEditorToolbar.INT.udn
--- a/Engine/Documentation/Source/Engine/UI/LevelEditor/Toolbar/LevelEditorToolbar.INT.udn
+++ b/Engine/Documentation/Source/Engine/UI/LevelEditor/Toolbar/LevelEditorToolbar.INT.udn
@@ -39,7 +39,7 @@ The **Toolbar** panel, like in most applications, is a group of commands providi
 | ![](button_Toolbar_matinee.png)| Enables you to create a new [Matinee](Engine/Matinee) sequence or edit any existing **Matinee** sequences in the level. |
 | ![](button_toolbar_build.png) | Performs a the _build_ operation on all levels (persistent and streaming) open in the editor. Building precalculates as much data as possible relating to various aspects of the level. For instance, static lighting - lightmaps, shadows, global illumination - and geometry are calculated during this process. Clicking the arrow displays the **Build Options** menu. |
 | ![](button_Toolbar_PlayIn.png) | [INCLUDE:#playin] |
-| ![](button_Toolbar_launch.png) | Launches the current map on any of the supported platforms. Connected deviecs are listed in the menu accessed by clicking the arrow. |
+| ![](button_Toolbar_launch.png) | Launches the current map on any of the supported platforms. Connected devices are listed in the menu accessed by clicking the arrow. |
 
 <!--
 [EXCERPT:simulate]
