Updates to UMG tutorial.

diff --git a/Engine/Documentation/Source/Programming/Tutorials/UMG/4/UMG_4.INT.udn b/Engine/Documentation/Source/Programming/Tutorials/UMG/4/UMG_4.INT.udn
--- a/Engine/Documentation/Source/Programming/Tutorials/UMG/4/UMG_4.INT.udn
+++ b/Engine/Documentation/Source/Programming/Tutorials/UMG/4/UMG_4.INT.udn
@@ -22,7 +22,7 @@ Description:Set up our Game Mode to use our Main Menu Widget.
 %Steps%
 
 
-1. In the **Content Browser**, we will add a **Blueprint Class** based on our project's **Game Mode**. This makes it possible to set the initial values of exposed variables to whatever we want. To do this:
+1. In the **Content Browser**, we will add a **Blueprint Class** based on our project's **Game Mode**. This makes it possible to set the exposed variables on those two classes to whatever values we want. To do this:
 
 	* Click the **Add** button in the **Content Browser**.
 
@@ -32,14 +32,30 @@ Description:Set up our Game Mode to use our Main Menu Widget.
 
 	[![](PickParentClassForGameModeBP.png)](PickParentClassForGameModeBP.png)
 
-1. After naming the resulting **Blueprint**, which we will call "MenuGameMode" for this example, right-click it and choose "Edit".
+	* Name the resulting **Blueprint** asset "MenuGameMode".
 
-	[![](EditBP.png)](EditBP.png)
+1. In order to see our mouse cursor in-game, we'll need to create a **Blueprint** of the **Player Controller** as we did with our **Game Mode**.
+
+	* Click the **Add** button in the **Content Browser** again.
+	
+	* Select **Player Controller** from the "Common Classes" section.
+
+	* Name the **Blueprint** "MenuPlayerController".
+
+1. Edit "MenuPlayerController".
+
+	* Check the "Show Mouse Cursor" box.
+
+	[![](GamePlayerController.png)](GamePlayerController.png)
+
+1. Edit "MenuGameMode".
+
+	* The Starting Widget Class must be set to the "MainMenu" asset in order to bring up our menu when the game begins.
 
-1. We will configure the **Blueprint** of our **Game Mode** with two changes:
-	* The Starting Widget Class must be set to MainMenu in order to bring up our menu when the game begins.
 	* The Default Pawn Class should be set to **Pawn** instead of **DefaultPawn** so that our player will not be able to fly around while in the menu.
 
+	* The Player Controller Class should be set to the "MenuPlayerController" asset we created so that the mouse cursor will be shown in-game.
+
 	[![](ConfigureGameMode.png)](ConfigureGameMode.png)
 
 1. In order for our **Blueprint** to be used, we must return to the **Level Editor** window and change the **World Settings** for our current **Level** via the **Settings** button.
@@ -55,7 +71,7 @@ Description:Set up our Game Mode to use our Main Menu Widget.
 
 ----
 
-Our custom **Game Mode** asset is now in effect on our level. If we run the game now, our Quit button will work as expected, and our New Game button will take us to an empty menu screen. Our next step will be to set up the New Game menu.
+Our custom **Game Mode** asset is now in effect on our level, configured to load our Main Menu, and use our **Player Controller** that shows the mouse cursor. If we run the game now, our Quit button will work as expected, and our New Game button will take us to an empty menu screen. Our next step will be to set up the New Game menu.
 
 [OBJECT:Section]
     [PARAMLITERAL:id]
diff --git a/Engine/Documentation/Source/Programming/Tutorials/UMG/5/UMG_5.INT.udn b/Engine/Documentation/Source/Programming/Tutorials/UMG/5/UMG_5.INT.udn
--- a/Engine/Documentation/Source/Programming/Tutorials/UMG/5/UMG_5.INT.udn
+++ b/Engine/Documentation/Source/Programming/Tutorials/UMG/5/UMG_5.INT.udn
@@ -51,7 +51,7 @@ Description:Build a second menu with new functionality.
 
 	[![](SelectAnchor.png)](SelectAnchor.png)
 
-	[REGION:tip]Positioning our anchor at the lower-right corner does not change how size and position values work, so we'll need to make our position values negative in order to be on the screen. Size values remain positive, but we will want to lower our position values to compensate.[/REGION]
+	[REGION:tip]Positioning our anchor at the lower-right corner does not change how size and position values work, so we'll need to make our position values negative in order to be on the screen. Size values will remain positive.[/REGION]
 
 1. We will now add scripting to our new **Buttons**. The Main Menu **Button** will simply reload the Main Menu **Widget**, while the Play Game **Button** will deactivate our menu entirely by having no new **Widget** provided in the call to our _ChangeMenuWidget_ function. This is shown by the phrase "Select Class" being displayed instead of the name of an actual class or asset.
 
diff --git a/Engine/Documentation/Source/Programming/Tutorials/UMG/UMG.INT.udn b/Engine/Documentation/Source/Programming/Tutorials/UMG/UMG.INT.udn
--- a/Engine/Documentation/Source/Programming/Tutorials/UMG/UMG.INT.udn
+++ b/Engine/Documentation/Source/Programming/Tutorials/UMG/UMG.INT.udn
@@ -47,7 +47,7 @@ Over the course of this tutorial, you will learn to use Unreal Motion Graphics (
 	360
 	[/PARAMLITERAL]
 	[PARAMLITERAL:videoid]
-	nTshVrOaNdY
+	eEwZ8QVaYCQ
 	[/PARAMLITERAL]
 [/OBJECT]
 
