#UE4doc #docs: sp: causes

diff --git a/Engine/Documentation/Source/Engine/UI/LevelEditor/Viewports/ViewportControls/ViewportControls.INT.udn b/Engine/Documentation/Source/Engine/UI/LevelEditor/Viewports/ViewportControls/ViewportControls.INT.udn
--- a/Engine/Documentation/Source/Engine/UI/LevelEditor/Viewports/ViewportControls/ViewportControls.INT.udn
+++ b/Engine/Documentation/Source/Engine/UI/LevelEditor/Viewports/ViewportControls/ViewportControls.INT.udn
@@ -226,7 +226,7 @@ These controls affect how the level is displayed in the viewport.
 
 | Control | Action |
 | ------- | ------ |
-| **G** | Toggles **Game Mode**, which casues the viewport to render only what would be seen in-game. |
+| **G** | Toggles **Game Mode**, which causes the viewport to render only what would be seen in-game. |
 | **Ctrl + R** | Toggles real-time playback in the active viewport. |
 | **F11** | Toggles immersive mode, which puts the viewport into full-screen. |
 
