#UE4doc #docs: made plural: sections

diff --git a/Engine/Documentation/Source/Engine/UI/ProjectSettings/ProjectSettings.INT.udn b/Engine/Documentation/Source/Engine/UI/ProjectSettings/ProjectSettings.INT.udn
--- a/Engine/Documentation/Source/Engine/UI/ProjectSettings/ProjectSettings.INT.udn
+++ b/Engine/Documentation/Source/Engine/UI/ProjectSettings/ProjectSettings.INT.udn
@@ -51,7 +51,7 @@ The **Project Settings** editor can be opened from the **Edit** menu:
 ## Categories and Sections
 
 The **Project Settings** editor is divided up into various categories and sections of related options. Categories 
-are displayed as headings, while each of these section is displayed as a hyperlink that opens the options for that 
+are displayed as headings, while each of these sections is displayed as a hyperlink that opens the options for that 
 section in the editor.
 
 While viewing a particular section, you can navigate back to the main view by using the **< Back** link at the top of the left sidebar:
