#UE4doc #docs: changed to "provides" to match singular subject

diff --git a/Engine/Documentation/Source/Engine/UMG/QuickStart/1/UMGQS_1.INT.udn b/Engine/Documentation/Source/Engine/UMG/QuickStart/1/UMGQS_1.INT.udn
--- a/Engine/Documentation/Source/Engine/UMG/QuickStart/1/UMGQS_1.INT.udn
+++ b/Engine/Documentation/Source/Engine/UMG/QuickStart/1/UMGQS_1.INT.udn
@@ -29,7 +29,7 @@ For this tutorial we are using the **Blueprint First Person** template with **St
 
 	![](UMGQS1.png)
 
-	All User Interface elements (HUD, Menus, etc.) will be created and housed inside a **Widget Blueprint**. The Widget Blueprint allows you to visually layout your UI elements as well as provide scripted functionality for those elements. We are creating this now so that we can tell it to be displayed when our player character spawns in the world, however we will set it up later. 
+	All User Interface elements (HUD, Menus, etc.) will be created and housed inside a **Widget Blueprint**. The Widget Blueprint allows you to visually layout your UI elements as well as provides scripted functionality for those elements. We are creating this now so that we can tell it to be displayed when our player character spawns in the world, however we will set it up later. 
 	
 	[REGION:note]
 	For more information on **Widget Blueprints** see the [Widget Blueprint](Engine/UMG/UserGuide/WidgetBlueprints) documentation. 
