#UE4doc #docs: minor edits made for clarification

diff --git a/Engine/Documentation/Source/Engine/UMG/QuickStart/2/UMGQS_2.INT.udn b/Engine/Documentation/Source/Engine/UMG/QuickStart/2/UMGQS_2.INT.udn
--- a/Engine/Documentation/Source/Engine/UMG/QuickStart/2/UMGQS_2.INT.udn
+++ b/Engine/Documentation/Source/Engine/UMG/QuickStart/2/UMGQS_2.INT.udn
@@ -54,7 +54,7 @@ Here we will begin setting up the visual aspects of our HUD.
 
 	![](Step2_10.png)
 
-	The progress bars will appear very small, but don't worry we will fix that next. 
+	The progress bars will appear very small but don't worry, we will fix that next. 
 
 1.	Select both **Progress Bars**, then in the **Details** panel set both to **Fill** for the **Size**. 
 
@@ -132,7 +132,7 @@ Now that we have our visual layout, we can add the hooks to provide the function
 
 	![](Step2_23.png)
 
-	This binds the value of the Progress Bar to our **Health** variable inside our **First Person Character** Blueprint. Now whenever our Health value changes in our Character Blueprint it will automatically be updated in our HUD as well. You may also notice that the Progress Bar does not change after binding it, that is because we've stated this to occur upon **Event Construct** (or when the HUD is constructed, which occurs when the game is started).
+	This binds the value of the Progress Bar to our **Health** variable inside our **First Person Character** Blueprint. Now whenever our Health value changes in our Character Blueprint it will automatically be updated in our HUD as well. You may also notice that the Progress Bar does not change after binding it. That is because we've stated this to occur upon **Event Construct** (or when the HUD is constructed, which occurs when the game is started).
 
 1.	Select the **Progress Bar** next to Energy and repeat the process above, binding **Percent** to **Energy** from **MyCharacter**. 
 
@@ -142,7 +142,7 @@ Now that we have our visual layout, we can add the hooks to provide the function
 
 	![](Step2_26.png)
 
-	Previously we used **Sub-Object Properties** to bind properties of the same type, we can also create our own custom bindings. Here we are going to bind our **Text** property to an **Integer** property from our **First Person Character** Blueprint giving us a display of our current Ammo.  
+	Now, like the way we previously used **Sub-Object Properties** to bind properties of the same type, we can also create our own custom bindings. Here we are going to bind our **Text** property to an **Integer** property from our **First Person Character** Blueprint giving us a display of our current Ammo.  
 
 1.	Inside the function that is created and automatically opened, hold **Ctrl** and drag the **MyCharacter** variable into the graph. 
 
