#UE4 Docs: fixed bad path

diff --git a/Engine/Documentation/Source/Engine/Rendering/Materials/HowTo/MaterialHowTo.INT.udn b/Engine/Documentation/Source/Engine/Rendering/Materials/HowTo/MaterialHowTo.INT.udn
--- a/Engine/Documentation/Source/Engine/Rendering/Materials/HowTo/MaterialHowTo.INT.udn
+++ b/Engine/Documentation/Source/Engine/Rendering/Materials/HowTo/MaterialHowTo.INT.udn
@@ -80,7 +80,7 @@ Read over each description to get an understanding of what the How-To will cover
 | **[Add Detail Texturing](Engine\Rendering\Materials\HowTo\DetailTexturing)** |Adding extra detail to a surface so that it will look nice when viewed very closely can be achieved through Detail mapping. In this How-To, you will be shown how you can add Detail Mapping to your Materials so they will always look good no matter what distance they are viewed from.|
 | **[Enable and use Refraction](Engine\Rendering\Materials\HowTo\Refraction)** |When light hits an object like glass, some of the light is changed because of refraction. In this How-To, we will cover how you can use Refraction in your Materials.|
 | **[Enable and use Transparency](Engine\Rendering\Materials\HowTo\Transparency)** |The ability to simulate Transparency is key to being able to make believable water or glass. In this How-To, we will cover how you can utilize Transparency in your Materials.|
-| **[Use Layered Materials](Engine/Rendering/Materials/HowTo/CreatingLayeredMaterials)** | %Engine/Rendering/Materials/How-To/CreatingLayeredMaterials:description% |
+| **[Use Layered Materials](Engine/Rendering/Materials/HowTo/CreatingLayeredMaterials)** | %Engine/Rendering/Materials/HowTo/CreatingLayeredMaterials:description% |
 | **[Make Materials Shiny](Engine/Rendering/Materials/HowTo/ShinyMaterials)** | Making Materials that are shiny is something can be accomplished in just a few clicks of the mouse. In this How-To, we will cover how you can turn any Material into a Material that shines.|
 | **[Colored Translucent Shadows](Engine/Rendering/Materials/HowTo/ColoredTransluscentShadows)** |Getting shadows to look like they are casting colored shadows into the world can be achieved by using Colored Translucent Shadows. In this How-To, we will cover how you can create Materials that make use of Colored Translucent Shadows. |
 
