#UE4doc #docs: Minor text edits made for clarification.

diff --git a/Engine/Documentation/Source/Engine/UMG/UserGuide/Slots/Slots.INT.udn b/Engine/Documentation/Source/Engine/UMG/UserGuide/Slots/Slots.INT.udn
--- a/Engine/Documentation/Source/Engine/UMG/UserGuide/Slots/Slots.INT.udn
+++ b/Engine/Documentation/Source/Engine/UMG/UserGuide/Slots/Slots.INT.udn
@@ -7,9 +7,10 @@ Description:An overview of how slots are used in UMG.
 
 ## Overview
 
-**Slots** are the invisible glue that binds Widgets together. In Slate they are a lot more explicit, you would first create a **Slot** then choose what control to place inside of it. For **UMG**, you instead have **Panel** Widgets that automatically use the right kind of **Slot** when Child Widgets are added to them. 
+**Slots** are the invisible glue that bind Widgets together. In Slate they are a lot more explicit in that you must first create a Slot then
+choose what control to place inside of it. In **UMG**, however, you have **Panel** Widgets that automatically use the right kind of **Slot** when Child Widgets are added to them. 
 
-Each **Slot** is different, for example, if you were to place a control on a Grid, you would expect to be able to set things like Row and Column. But these properties have no business being on a Widget that was placed on a Canvas. That is where **Slots** come in, a **Canvas Slot** understands how to layout content absolutely and through Anchors, while a **Grid Slot** only understands Row and Column.
+Moreover, each **Slot** is different. For example, if you were to place a control on a Grid, you would expect to be able to set things like Row and Column. But these properties have no business being on a Widget that was placed on a Canvas. That is where **Slots** come in. A **Canvas Slot** understands how to layout content absolutely and through Anchors, while a **Grid Slot** only understands Rows and Columns.
 
 ## Accessing Slots 
 
