Jira 1297 - Blueprint Communication Updates (1 of 2)

Check in for updates to Direct Blueprint Communication / Blueprint Casting How To's.

diff --git a/Engine/Documentation/Source/Engine/Blueprints/BP_HowTo/BlueprintHowTo.INT.udn b/Engine/Documentation/Source/Engine/Blueprints/BP_HowTo/BlueprintHowTo.INT.udn
--- a/Engine/Documentation/Source/Engine/Blueprints/BP_HowTo/BlueprintHowTo.INT.udn
+++ b/Engine/Documentation/Source/Engine/Blueprints/BP_HowTo/BlueprintHowTo.INT.udn
@@ -25,6 +25,8 @@ Refer to the table below for more information and click each link to view its co
 | **[](Engine/Blueprints/BP_HowTo/MakingFunctions)** | %Engine/Blueprints/BP_HowTo/MakingFunctions:description% |
 | **[](Engine/Blueprints/BP_HowTo/WorkingWithArrays)** | %Engine/Blueprints/BP_HowTo/WorkingWithArrays:description% |
 | **[](Engine/Blueprints/UserGuide/Variables)** | The entire user guide for Variables. |
+| **[](Engine\Blueprints\BP_HowTo\DirectBlueprintComs)** | %Engine\Blueprints\BP_HowTo\DirectBlueprintComs:description% |
+| **[](Engine\Blueprints\BP_HowTo\Casting)** | %Engine\Blueprints\BP_HowTo\Casting:description% |
 | **[](Engine/Blueprints/BP_HowTo/MakingMacros)** | %Engine/Blueprints/BP_HowTo/MakingMacros:description% |
 | **[](Engine/Blueprints/BP_HowTo/Debugging)** | %Engine/Blueprints/BP_HowTo/Debugging:description% |
 
@@ -34,3 +36,8 @@ Refer to the table below for more information and click each link to view its co
 
 
 
+
+
+
+
+
diff --git a/Engine/Documentation/Source/Engine/Blueprints/BP_HowTo/Casting/Casting.INT.udn b/Engine/Documentation/Source/Engine/Blueprints/BP_HowTo/Casting/Casting.INT.udn
new file mode 100644
--- /dev/null
+++ b/Engine/Documentation/Source/Engine/Blueprints/BP_HowTo/Casting/Casting.INT.udn
@@ -0,0 +1,92 @@
+Availability: Public
+Title:Blueprint Casting
+Crumbs: %ROOT%, Engine, Engine\Blueprints, Engine\Blueprints\BP_HowTo
+Description:How Blueprints can work together and pass information between each other through Blueprint Casting. 
+Related: Engine/Blueprints/UserGuide/CastNodes
+Related: Engine/Blueprints/UserGuide/BlueprintCommsUsage
+Related: Engine/Blueprints/UserGuide
+
+This page covers how to directly communicate between two Blueprints and pass information between them. 
+
+For this example we will allow the player to pass through a box and communicate with and activate a particle effect. 
+
+[REGION:note]
+For this guide we are using the **Blueprint First Person** template with **Starter Content** enabled. 
+[/REGION]
+
+1.  Inside the **Viewport** select one of the **LargeCube** meshes, then in the **Details** panel click the **Blueprint/Add Script** button. 
+
+	![](Casting1.png)
+
+1.  In the **Content Browser** under **Content/FirstPersonBP/Blueprints**, open the **FirstPeronProjectile** Blueprint. 
+
+	![](Casting2.png)
+
+	This is the Blueprint for the projectile that is launched when the player fires the weapon during gameplay. 
+
+1.  At the confirmation window that appears, click **Ok** to use the default name and create the Blueprint. 
+
+1.  Locate the **Event Hit** node and drag off the **Other** pin then search for and add the **Cast To LargeCube_Blueprint** node. 
+
+	![](Casting3.png)(w:800)
+
+	On the **Event Hit** node, the **Other** pin refers to the other Actor that was hit. By using the **Cast To** node, we are checking to see if it is indeed the **LaregCube_Blueprint** that was the other Actor that was hit. If it is, we can access it and execute any Events or Functions within the Blueprint, change Variable values inside of it, or change the properties of the Components that are inside the Blueprint. 
+
+1.	Off the **As Large Cube Blueprint** pin, use **Get Static Mesh Component** and then **Set World Scale 3D** and connect as shown.
+
+	![](Casting4.png)
+
+	After we **Add Impulse at Location** we are then casting to the Large Cube Blueprint and getting the Static Mesh Component inside it (the actual Cube) and setting a new World Scale 3D value for it. Currently the New Scale is set to 0 however we will reduce the size of the Large Cube Blueprint each time it is hit. 
+
+1.	Off **Static Mesh Component** use **Get World Scale** and connect to a **Vector * Float** pin (set to **0.5**) then plug into **New Scale**. 
+
+	![](Casting5.png)
+
+	This will get the current World Scale of the cube and multiply it by .5 (or half its current size) each time it is hit, reducing its size. 
+
+1.	**Right-click** in the graph and add a **Get Player Character** node, then off **Return Value**, **Cast To FirstPersonCharacter**. 
+
+	![](Casting6.png)
+
+	Here we are checking if the player character is using the **FirstPersonCharacter** Blueprint, which by default is the case with this template. By doing so we can now access the Events, Functions, Variables or Properties of the player character and modify them. 
+
+1.	Connect the **Set World 3D** node to the **Cast To FirstPersonCharacter** node. 
+
+1.	Off the **As First Person Character** pin, add the **Get Character Movement** node. 
+
+	![](Casting7.png)
+
+	We are going to adjust the character's speed whenever they hit the Large Cube, making them run faster. 
+
+1.	Off the **Character Movement** node, add a **Set Max Walk Speed** and **Get Max Walk Speed** node. 
+
+1.	Off the **Get Max Walk Speed**, use a **Float + Float** node (set to **200**) then connect that to the **Set Max Walk Speed** pin. 
+
+	![](Casting8.png)(w:800)
+
+	Once you have connected the nodes as shown above, this will allow the player character to move faster each time the cube is hit. 
+
+1.	For both of the **Cast To** nodes, connect the **Cast Failed** pins to the **Destroy Actor** node. 
+
+	![](Casting9.png)
+
+	In the event that the player hits something that is not the Large Cube, we still want to call the **Destroy Actor** node to destroy the projectile. When hitting the Large Cube, in the event that it is not the **First Person Character** that hit the cube, we also still want to destroy the projectile. 
+
+1.	The ending sequence for the Projectile Blueprint should look similar to below. 
+
+	[REGION:lightbox]
+	[![](EndingSequence.png)(w:800)](EndingSequence.png)
+	[/REGION:lightbox]
+
+	_Click image for full view._
+
+	We used **Reroute Nodes** to adjust the flow of our wires for clarity, but your connections should be the same. 
+
+1.	Click **Compile** and **Save**, then the **Play** button to play in the editor. 
+
+	![](CastingGIF.gif)(convert:false)
+
+	Each time you shoot the Large Cube Blueprint, it will reduce in size and your character will move faster each time. 
+
+	For more information on **Casting** and usage cases, refer to the **Related Topics** section below. 
+	
diff --git a/Engine/Documentation/Source/Engine/Blueprints/BP_HowTo/DirectBlueprintComs/DirectBlueprintComs.INT.udn b/Engine/Documentation/Source/Engine/Blueprints/BP_HowTo/DirectBlueprintComs/DirectBlueprintComs.INT.udn
new file mode 100644
--- /dev/null
+++ b/Engine/Documentation/Source/Engine/Blueprints/BP_HowTo/DirectBlueprintComs/DirectBlueprintComs.INT.udn
@@ -0,0 +1,96 @@
+Availability: Public
+Title:Direct Blueprint Communication
+Crumbs: %ROOT%, Engine, Engine\Blueprints, Engine\Blueprints\BP_HowTo
+Description:A demonstration of passing information directly from one Blueprint to another. 
+Related: Engine/Blueprints/UserGuide/BlueprintComms
+Related: Engine/Blueprints/UserGuide/BlueprintCommsUsage
+Related: Engine/Blueprints/UserGuide
+
+This page covers how to directly communicate between two Blueprints and pass information between them. 
+
+For this example we will allow the player to pass through a box and communicate with and activate a particle effect. 
+
+[REGION:note]
+For this guide we are using the **Blueprint Third Person** template with **Starter Content** enabled. 
+[/REGION]
+
+1.  In the **Content Browser**, open the **Content/StarterContent/Shapes** folder. 
+
+1.  **Right-click** on the **Shape_Cube**, then under _Asset Actions_, choose **Create Blueprint Using This...**
+
+	![](DirectCom1.png)
+
+	Click **Ok** at the confirmation box that appears to create the Blueprint. 
+
+1. Inside the Shape_Cube Blueprint, select the **Static Mesh** in the **Components** window in the upper-left. 
+
+1. In the **Details** panel, change the **Collision Presets** to **OverlapOnlyPawn**. 
+
+	![](DirectCom2.png)
+
+	This will ensure that only the **Pawns** and by extension **Charaters** can pass through the box. 
+
+1. In the **MyBlueprint** window, click the **Add Variable** button and name the variable **TargetBlueprint**. 
+
+	![](DirectCom3.png)
+
+	This will be the other Blueprint that we want to directly communicate with. 
+
+1. In the **Details** panel for the variable, set **Variable Type** to **Blueprint_Effect_Fire** and make it **Editable**. 
+
+	![](DirectCom4.png)
+
+	Here we set the Blueprint that we want to communicate with and made the variable publicly editable (editable outside this Blueprint). 
+
+1. Select the **Static Mesh** in the **My Blueprint** window, then in the **Details** panel, add an **On Component Begin Overlap** Event. 
+
+	![](DirectCom8.png)
+
+	This will add a node and open the **Event Graph** where we can provide the script that occurs when entering the cube. 
+
+1. Drag off the **Target Blueprint** and **Get P Fire** and **Get Fire Audio**, then connect as shown to an **Activate** node. 
+
+	![](DirectCom10.png)(w:720)
+
+	Here we are accessing the Particle Effect and Audio components inside our Target Blueprint and Activating them. 
+
+1. In the **Content Browser** under **Blueprints**, drag the **Shape_Cube** Blueprint into the level. 
+
+	![](DirectCom5.png)(w:720)
+
+1. In the **Content Browser** under **Content/StarterContent/Blueprints**, open the **Blueprint_Effect_Fire** Blueprint. 
+
+1. In the **Components** window, select **P_Fire** then in the **Details** panel un-check **Auto Activate**. 
+
+	![](DirectCom6.png)
+
+	We do not want this effect to auto activate as we will tell it to activate when entering the box in our level. 
+
+1. Drag the **Blueprint_Effect_Fire** Blueprint into the level. 
+
+1. Select the **Shape_Cube** Blueprint then in the **Details** panel, use the drop-down to select the **Blueprint_Effect_Fire** Blueprint. 
+
+	![](DirectCom7.png)
+
+	Here we are specifying the instance of the Blueprint_Effect_Fire Bluperint in the level that we want to communicate with. If we have placed several instances of the Blueprint_Effect_Fire Blueprint in the level, each of the instances would be displayed inside the drop-down menu allowing us to indicate which instance is the **Target Blueprint** to communicate with. 
+
+	[REGION:tip]
+	Clicking the **Eye Dropper** icon allows you to select the instance of your **Target Bluperint** inside the level instead of using the drop-down menu. Remember, you can only select Blueprints to target based on the variable type you have specified. 
+	[/REGION]
+
+1. Click the **Play** button to play in the editor and test running through the box with the charcter. 
+
+You will see that once the character enters the box, the fire effect will be activated inside the level. You can use **Direct Bluperint Communcation*** to not only change the properties of a Blueprint from inside another Blueprint, but you can also pass the values of Variables, call Functions or Events, or anything else that the Blueprint would otherwise normally be able to do. 
+
+
+
+
+
+
+
+
+
+
+
+	
+
