#UE4doc #docs lots of image updates to mouse interface content example page and info

diff --git a/Engine/Documentation/Source/Resources/ContentExamples/MouseInterface/MouseControlSetup/MouseControlSetup.INT.udn b/Engine/Documentation/Source/Resources/ContentExamples/MouseInterface/MouseControlSetup/MouseControlSetup.INT.udn
--- a/Engine/Documentation/Source/Resources/ContentExamples/MouseInterface/MouseControlSetup/MouseControlSetup.INT.udn
+++ b/Engine/Documentation/Source/Resources/ContentExamples/MouseInterface/MouseControlSetup/MouseControlSetup.INT.udn
@@ -25,7 +25,7 @@ By default, the PlayerController class does not have mouse interfacing enabled.
 	Be aware that these are just the basic steps for enabling mouse control. In most cases, there will be others steps involved to customize your PlayerController and GameMode for the specific type of game that you are creating.
 [/REGION]
 
-1. In the **Content Browser** (for sake of ease, use your Developer folder), **right-click** and choose **Blueprint** from the context menu.
+1. In the **Content Browser** (for sake of ease, use your Developer folder), **right-click** and choose **Blueprint Class** from the context menu.
 
 	[REGION:raw]![](CreateNewBlueprint.png)[/REGION]
 
@@ -37,7 +37,7 @@ By default, the PlayerController class does not have mouse interfacing enabled.
 
 	![](MousePlayerController_CB.png)
 
-1. **Right-click** on your new Blueprint in the **Content Browser** and choose **Edit Defaults**. This will open up the **Class Defaults** window.
+1. **Right-click** on your new Blueprint in the **Content Browser** and choose **Edit...**. This will open up the **Class Defaults** window.
 
 	![](EditDefaultsBlueprint.png)
 
@@ -52,7 +52,7 @@ By default, the PlayerController class does not have mouse interfacing enabled.
 
 Now that we have a PlayerController that has Mouse Input enabled, our next step is to create a custom GameMode that utilizes our new PlayerController. Note that the steps described here will _not work_ if you have not yet created the custom Player Controller class as described above.
 
-1. In the **Content Browser** (for sake of ease, use your Developer folder), **right-click** and choose **Blueprint** from the context menu.
+1. In the **Content Browser** (for sake of ease, use your Developer folder), **right-click** and choose **Blueprint Class** from the context menu.
 
 	[REGION:raw]![](CreateNewBlueprint.png)[/REGION]
 
@@ -64,11 +64,11 @@ Now that we have a PlayerController that has Mouse Input enabled, our next step
 
 	![](Mouse_GameInfoBlueprint.png)
 
-1. **Right-click** on your new Blueprint in the **Content Browser** and choose **Edit Defaults**. This will open up the **Class Defaults** window.
+1. **Right-click** on your new Blueprint in the **Content Browser** and choose **Edit...**. This will open up the **Class Defaults** window.
 
-	![](EditDefaultsBlueprint.png)
+	![](EditDefaultsBlueprint2.png)
 
-1. In the Game Mode category, click on the dropdown next to **Player Controller Class**. You should see the **Mouse_PlayerController** class created previously. Choose it from the list.
+1. In the Classes category, click on the dropdown next to **Player Controller Class**. You should see the **Mouse_PlayerController** class created previously. Choose it from the list.
 
 	![](SelectingPlayerControllerClass.png)
 
@@ -79,7 +79,7 @@ Now that we have a PlayerController that has Mouse Input enabled, our next step
 
 Now that we have our custom PlayerController and GameMode Blueprints completed, we just need to tell our World Settings to use them.
 
-1. Click the **World Settings** button located in the toolbar. This opens up the World Settings in the **Details** panel.
+1. Click the **Settings** button located in the toolbar. Then from the drop down select **World Settings**. This opens up the World Settings in the **Details** panel.
 
 	![](WorldSettings.png)
 
