Jira 1297 - Blueprint Communication Updates (1.5 of 2)

Added Event Dispatchers How-To to Blueprints HowTo's.

diff --git a/Engine/Documentation/Source/Engine/Blueprints/BP_HowTo/BlueprintHowTo.INT.udn b/Engine/Documentation/Source/Engine/Blueprints/BP_HowTo/BlueprintHowTo.INT.udn
--- a/Engine/Documentation/Source/Engine/Blueprints/BP_HowTo/BlueprintHowTo.INT.udn
+++ b/Engine/Documentation/Source/Engine/Blueprints/BP_HowTo/BlueprintHowTo.INT.udn
@@ -27,6 +27,7 @@ Refer to the table below for more information and click each link to view its co
 | **[](Engine/Blueprints/UserGuide/Variables)** | The entire user guide for Variables. |
 | **[](Engine\Blueprints\BP_HowTo\DirectBlueprintComs)** | %Engine\Blueprints\BP_HowTo\DirectBlueprintComs:description% |
 | **[](Engine\Blueprints\BP_HowTo\Casting)** | %Engine\Blueprints\BP_HowTo\Casting:description% |
+| **[](Engine\Blueprints\BP_HowTo\EventDispatcher)** | %Engine\Blueprints\BP_HowTo\EventDispatcher:description% |
 | **[](Engine/Blueprints/BP_HowTo/MakingMacros)** | %Engine/Blueprints/BP_HowTo/MakingMacros:description% |
 | **[](Engine/Blueprints/BP_HowTo/Debugging)** | %Engine/Blueprints/BP_HowTo/Debugging:description% |
 
diff --git a/Engine/Documentation/Source/Engine/Blueprints/BP_HowTo/EventDispatcher/EventDispatcher.INT.udn b/Engine/Documentation/Source/Engine/Blueprints/BP_HowTo/EventDispatcher/EventDispatcher.INT.udn
new file mode 100644
--- /dev/null
+++ b/Engine/Documentation/Source/Engine/Blueprints/BP_HowTo/EventDispatcher/EventDispatcher.INT.udn
@@ -0,0 +1,95 @@
+Availability: Public
+Title:Event Dispatchers
+Crumbs: %ROOT%, Engine, Engine\Blueprints, Engine\Blueprints\BP_HowTo
+Description:A sample of how an event called in one Blueprint can execute events in another. 
+Related: Engine/Blueprints/UserGuide/EventDispatcher
+Related: Engine/Blueprints/UserGuide/BlueprintCommsUsage
+Related: Engine/Blueprints/UserGuide
+
+When an **Event Dispatcher** is called, any Blueprints that implement and have Events bound to the Event Dispatcher will also execute when the Event Dispatcher is called. In other words, whenever the Event Dispatcher is called, any Events that are bound to it are also called allowing you to execute multiple Events at once from a single source. 
+
+This page will show you how to set up an Event Dispatcher and Bind an Event to the Event Dispatcher. 
+
+For this example we will allow the player to press a key which will cause a fire particle effect to spawn inside of a bush. 
+
+[REGION:note]
+For this guide we are using the **Blueprint Third Person** template with **Starter Content** enabled. 
+[/REGION]
+
+1.  In the **Content Browser**, under the **Content/ThirdPersonBP/Blueprints** folder, open the **ThirdPersonCharacter** Blueprint. 
+
+	![](EventD_1.png)
+
+1.	Under the **My Blueprint** window, click the **Add Event Dispatcher** button and name it **StartFire**. 
+
+	![](EventD_2.png)
+
+1.	**Right-click** in the graph window and add an **F** Key Event.
+
+1.	Drag the **StartFire** Event Dispatcher into the graph and select **Call**. 
+
+	![](EventD_3.png)
+
+1.	Connect the **F** Key Event to the **Call StartFire** Event Dispatcher. 
+
+	![](EventD_4.png)
+
+	Whenever the **F** key is pressed, any Blueprints that have an Event bound to this Event Dispatcher will also be executed. 
+
+1.	**Compile** and **Save** then close the Blueprint. 
+
+1.	In the **Content Browser** under **Content/StarterContent/Props**, drag the **SM_Bush** into the level. 
+
+	![](EventD_5.png)(w:820)
+
+1.	From the Main Toolbar, click the **Blueprints** button then select **Open Level Blueprint**. 
+
+	![](EventD_6.png)
+
+1.	**Right-click** inside the graph and select the **Create a Reference to SM_Bush** node. 
+
+	![](EventD_7.png)
+
+	Here we are getting a reference to the SM_Bush inside our level. 
+
+	[REGION:note]
+	Don't see the node listed above? Select the **SM_Bush** inside your level then try again. 
+	[/REGION]
+
+1.	**Right-click** in the graph and add an **Event Begin Play** and **Get Player Character** node. 
+
+1.	Off the **Get Player Character** node, search for and add the **Cast To ThirdPersonCharacter** node. 
+
+	![](EventD_8.png)
+
+	This will allow us to access the Events, Functions, Variables or Properties inside the **ThirdPersonCharacter** Blueprint. 
+
+	[REGION:note]
+	Please see the [](Engine/Blueprints/UserGuide/CastNodes) documentation for more information regarding Blueprint Cast Nodes. 
+	[/REGION]
+
+1.	Off the **As Third Person Character** pin, search for and use the **Assign Start Fire** node. 
+
+	![](EventD_9.png)
+
+	This will create two new nodes, a **Bind Event** node and a **Custom Event**. Whenever the Event Dispatcher is called inside our Character Blueprint, since we are binding that event, the Custom Event and anything connected to it will also be executed. Any Blueprints that have Events bound to the Event Dispatcher will also be executed. 
+
+1.	Off the **Custom Event**, add a **Spawn Actor from Class** node (with **Class** set to **Blueprint_Effect_Fire**). 
+
+1.	Drag off the **SM_Bush** to **Get Actor Transform** and use that as the **Spawn Transform** for the Spawn Actor node. 
+
+	![](EventD_10.png)
+
+	Once all nodes are connected, your graph should look similar to above. 
+
+1.	Click **Compile**, close the Level Blueprint, then click **Play** to play in the editor. 
+
+	![](EventD_11.png)(w:640)
+
+	When you press the **F** key, the fire effect will spawn inside the bush. 
+
+For more information on working with **Event Dispatcher** and sample use cases, refer to the **Related Topics** below. 
+
+
+
+	
