#UE4doc #docs: "A Blueprint Class allows..." sp: cannot

diff --git a/Engine/Documentation/Source/Gameplay/ClassCreation/BlueprintOnly/BlueprintOnly.INT.udn b/Engine/Documentation/Source/Gameplay/ClassCreation/BlueprintOnly/BlueprintOnly.INT.udn
--- a/Engine/Documentation/Source/Gameplay/ClassCreation/BlueprintOnly/BlueprintOnly.INT.udn
+++ b/Engine/Documentation/Source/Gameplay/ClassCreation/BlueprintOnly/BlueprintOnly.INT.udn
@@ -5,7 +5,7 @@ Description:Introductory information for gameplay programmers getting started wi
 
 ## Overview
 
-[Blueprints Class](Engine/Blueprints/UserGuide/Types/ClassBlueprint) allow you to set up new classes using the [](Engine/Blueprints) system.  After you [create a new Blueprint Class](Engine/Blueprints/UserGuide/Types/ClassBlueprint/Creation),
+A [Blueprints Class](Engine/Blueprints/UserGuide/Types/ClassBlueprint) allows you to set up new classes using the [](Engine/Blueprints) system.  After you [create a new Blueprint Class](Engine/Blueprints/UserGuide/Types/ClassBlueprint/Creation),
 you can add components, set up functions and other gameplay or design behavior with visual scripting, and set default values for class variables. The LightSwitch class set up using only Blueprints is named
 **LightSwitch_BPOnly**, and is explained below.
 
@@ -48,11 +48,11 @@ and [Graph Editor](Engine/Blueprints/Editor/UIComponents/GraphEditor) documentat
 ![](BP_Only_EventGraph.png)
 
 [REGION:note]
-If you can not find the Toggle Visibility from **right-click** menu un-check context sensitive or try searching for it in the Find a Node Menu. You can also get the same results by simply dragging
+If you cannot find the Toggle Visibility from **right-click** menu un-check context sensitive or try searching for it in the Find a Node Menu. You can also get the same results by simply dragging
 out a line from the Point light variable node and then searching for Toggle Visibility. 
 [/REGION]
 [REGION:note]
-If you can not find the OnComponentEndOverlap for Sphere1 make sure that you have the Sphere variable selected and then use the add event from the Event->Add Event or 
+If you cannot find the OnComponentEndOverlap for Sphere1 make sure that you have the Sphere variable selected and then use the add event from the Event->Add Event or 
 from the **right-click** menu Add Event->Collision->View **OnComponentBeginOverlap / OnComponentEndOverlap**. 
 [/REGION]
 
