#UE4doc #docs: sp: transferring

diff --git a/Engine/Documentation/Source/Gameplay/Framework/GameMode/GameMode.INT.udn b/Engine/Documentation/Source/Gameplay/Framework/GameMode/GameMode.INT.udn
--- a/Engine/Documentation/Source/Gameplay/Framework/GameMode/GameMode.INT.udn
+++ b/Engine/Documentation/Source/Gameplay/Framework/GameMode/GameMode.INT.udn
@@ -35,7 +35,7 @@ The GameMode contains a state machine that tracks the state of the Match, or the
 * WaitingToStart is the next state, and HandleMatchIsWaitingToStart is called when entering. Actors are ticking, but players have not yet spawned in. It transitions to the next state if ReadyToStartMatch returns _true_, or if someone calls StartMatch.
 * InProgress is the next state, and HandleMatchHasStarted is called when entering, which calls BeginPlay on all Actors. Normal gameplay is in progress. It transitions to the next state if ReadyToEndMatch returns _true_, or if someone calls EndMatch.
 * WaitingPostMatch is the next state, and HandleMatchHasEnded is called when entering. Actors are still ticking but new players are not accepted. It transitions to the next state when a map transfer starts.
-* LeavingMap is the last state, and HandleLeavingMap is called when entering. The match stays in this state while transfering to a new map.
+* LeavingMap is the last state, and HandleLeavingMap is called when entering. The match stays in this state while transferring to a new map.
 * Aborted is the failure state, started from AbortMatch. This is set when there is an unrecoverable error.
 
 By default, the Match state will almost always be InProgress. But, individual games can override this behavior to build a multiplayer game with more complicated rules.
