#UE4doc #docs: sp: assigning

diff --git a/Engine/Documentation/Source/Gameplay/HowTo/CharacterMovement/Blueprints/AnimBlueprint_Walk/Walk_ABP.INT.udn b/Engine/Documentation/Source/Gameplay/HowTo/CharacterMovement/Blueprints/AnimBlueprint_Walk/Walk_ABP.INT.udn
--- a/Engine/Documentation/Source/Gameplay/HowTo/CharacterMovement/Blueprints/AnimBlueprint_Walk/Walk_ABP.INT.udn
+++ b/Engine/Documentation/Source/Gameplay/HowTo/CharacterMovement/Blueprints/AnimBlueprint_Walk/Walk_ABP.INT.udn
@@ -47,7 +47,7 @@ The steps below will show you how to set up the Animation Blueprint:
 	
 	_Click image for full view._
 
-	In the graph above, we are getting the direction of the player character and their speed and assiging them to variables that will be used inside the Animation Blueprint. We are also getting variables from our **MyCharacter** Blueprint (by dragging off the **Try Get Pawn Owner** node and **Casting** to our Character Blueprint) which corresponds to actions taken by the player. Those variables are then passed to variables inside the Animation Blueprint that will determine when to transition to/from Animation States. 
+	In the graph above, we are getting the direction of the player character and their speed and assigning them to variables that will be used inside the Animation Blueprint. We are also getting variables from our **MyCharacter** Blueprint (by dragging off the **Try Get Pawn Owner** node and **Casting** to our Character Blueprint) which corresponds to actions taken by the player. Those variables are then passed to variables inside the Animation Blueprint that will determine when to transition to/from Animation States. 
 
 1.	Click the **Anim Graph** tab, **Right-click** in the graph and search for and add a new **State Machine**. 
 
