#UE4doc #docs: "its" to "it is"

diff --git a/Engine/Documentation/Source/Gameplay/HowTo/CharacterMovement/Blueprints/Setup_2/Setup_2.INT.udn b/Engine/Documentation/Source/Gameplay/HowTo/CharacterMovement/Blueprints/Setup_2/Setup_2.INT.udn
--- a/Engine/Documentation/Source/Gameplay/HowTo/CharacterMovement/Blueprints/Setup_2/Setup_2.INT.udn
+++ b/Engine/Documentation/Source/Gameplay/HowTo/CharacterMovement/Blueprints/Setup_2/Setup_2.INT.udn
@@ -107,7 +107,7 @@ In this section, we will finish setting up the Character by implementing the scr
 
 	_Click image for full view._
 
-	For the Jump script, we are first checking if we are either Prone or Crouched. If neither are _true_, we are then using the **CharacterMovement** component to determine our current velocity and if its greater than 0, which would mean that we are moving. 
+	For the Jump script, we are first checking if we are either Prone or Crouched. If neither are _true_, we are then using the **CharacterMovement** component to determine our current velocity and if it is greater than 0, which would mean that we are moving. 
 	
 	We then set how high the character can jump based on whether or not they are moving by setting the **JumpZVelocity** to **365** (for a moving character) or **340** while we are stationary. In either case, we set **Jump Pressed** to _True_ and set **Jump Pressed** to _False_ when Jump is released. 
 
