#UE4doc #docs: sp: Classes

diff --git a/Engine/Documentation/Source/Gameplay/HowTo/CharacterMovement/Blueprints/Input/Input.INT.udn b/Engine/Documentation/Source/Gameplay/HowTo/CharacterMovement/Blueprints/Input/Input.INT.udn
--- a/Engine/Documentation/Source/Gameplay/HowTo/CharacterMovement/Blueprints/Input/Input.INT.udn
+++ b/Engine/Documentation/Source/Gameplay/HowTo/CharacterMovement/Blueprints/Input/Input.INT.udn
@@ -29,7 +29,7 @@ Here we will set up the **Input** (Game Controls) to control the character and c
 
 	![](SelectClassDefaults.png)
 
-1.	In the **Details** panel under **Classess**, set the **Default Pawn Class** to use the **MyCharacter** Blueprint. 
+1.	In the **Details** panel under **Classes**, set the **Default Pawn Class** to use the **MyCharacter** Blueprint. 
 
 	![](GameMode_1.png)
 
