#UE4doc #docs: "a" to "an"

diff --git a/Engine/Documentation/Source/Gameplay/HowTo/FindingActors/Blueprints/FindingActors_BP.INT.udn b/Engine/Documentation/Source/Gameplay/HowTo/FindingActors/Blueprints/FindingActors_BP.INT.udn
--- a/Engine/Documentation/Source/Gameplay/HowTo/FindingActors/Blueprints/FindingActors_BP.INT.udn
+++ b/Engine/Documentation/Source/Gameplay/HowTo/FindingActors/Blueprints/FindingActors_BP.INT.udn
@@ -103,7 +103,7 @@ One method you can take is by using **Tags** on the Actor, for example:
 
 1.	**Right-click** in the graph and add a **Branch** node. 
 
-1.	Drag off the **Condition** of the **Branch** and add a **Actor Has Tag** node. 
+1.	Drag off the **Condition** of the **Branch** and add an **Actor Has Tag** node. 
 
 1.	Connect the out pin from the **Get** node to the **Tag** pin on the **Actor Has Tag** node. 
 
