Fixing filenames in C++ Tutorials to meet new style guidelines.

diff --git a/Engine/Documentation/Source/Programming/QuickStart/3/QuickStart_3.INT.udn b/Engine/Documentation/Source/Programming/QuickStart/3/QuickStart_3.INT.udn
--- a/Engine/Documentation/Source/Programming/QuickStart/3/QuickStart_3.INT.udn
+++ b/Engine/Documentation/Source/Programming/QuickStart/3/QuickStart_3.INT.udn
@@ -21,17 +21,17 @@ Description:Writing a C++ class and compiling it with Unreal Engine
 
 %Steps%
 
-1. In **Visual Studio**, we'll use the **Solution Explorer** pane to find our newly-created C++ files. In our example, they will be named "FloatingActor.cpp" and "FloatingActor.h" and will be inside the "QuickStart" project.
+1. In **Visual Studio**, we'll use the **Solution Explorer** pane to find our newly-created C++ files. In our example, they will be named `FloatingActor.cpp` and `FloatingActor.h` and will be inside the "QuickStart" project.
 
 	[![](SolutionExplorer_Files.png)](SolutionExplorer_Files.png)
 
 	[REGION:note]We're about to add code. All code used in this tutorial will be found at bottom of the page in its current state as of completing that page's instructions.[/REGION]
 
-2. In "FloatingActor.h", we'll add the following code just before the closing brace and semicolon at the end of the file:
+2. In `FloatingActor.h`, we'll add the following code just before the closing brace and semicolon at the end of the file:
 
 		float RunningTime;
 
-3. Switching to "FloatingActor.cpp", we'll add the following code just before the closing brace at the bottom of **AFloatingActor::Tick**:
+3. Switching to `FloatingActor.cpp`, we'll add the following code just before the closing brace at the bottom of **AFloatingActor::Tick**:
 
 		FVector NewLocation = GetActorLocation();
 		float DeltaHeight = (FMath::Sin(RunningTime + DeltaTime) - FMath::Sin(RunningTime));
diff --git a/Engine/Documentation/Source/Programming/Tutorials/AutoCamera/2/AutoCamera_2.INT.udn b/Engine/Documentation/Source/Programming/Tutorials/AutoCamera/2/AutoCamera_2.INT.udn
--- a/Engine/Documentation/Source/Programming/Tutorials/AutoCamera/2/AutoCamera_2.INT.udn
+++ b/Engine/Documentation/Source/Programming/Tutorials/AutoCamera/2/AutoCamera_2.INT.udn
@@ -25,7 +25,7 @@ Description:Coding a class to change camera views
 
 	[![](NameActorClass.png)](NameActorClass.png)
 
-1. In "CameraDirector.h", add the following code to the bottom of the **ACameraDirector** class definition:
+1. In `CameraDirector.h`, add the following code to the bottom of the **ACameraDirector** class definition:
 
 		UPROPERTY(EditAnywhere)
 		AActor* CameraOne;
@@ -37,7 +37,7 @@ Description:Coding a class to change camera views
 
 	[REGION:note]The **UPROPERTY** macro makes our variables visible to **Unreal Engine**. This way, values set in these variables will not be reset when we launch the game or reload our level or project in a future work session. We have also added the **EditAnywhere** keyword, which will allow us to set CameraOne and CameraTwo in the **Unreal Editor**.[/REGION]
 
-1. In "CameraDirector.cpp", add the following line of code to the top of the file, right underneath the other "#include" lines:
+1. In `CameraDirector.cpp`, add the following line of code to the top of the file, right underneath the other "#include" lines:
 
 		#include "Kismet/GameplayStatics.h"
 
diff --git a/Engine/Documentation/Source/Programming/Tutorials/Components/1/Components_1.INT.udn b/Engine/Documentation/Source/Programming/Tutorials/Components/1/Components_1.INT.udn
--- a/Engine/Documentation/Source/Programming/Tutorials/Components/1/Components_1.INT.udn
+++ b/Engine/Documentation/Source/Programming/Tutorials/Components/1/Components_1.INT.udn
@@ -28,13 +28,13 @@ Description:Creating a variety of components and attaching them to a customized
 
 	[![](NamePawnClass.png)](NamePawnClass.png)
 
-1. In **Visual Studio**, we should open "CollidingPawn.h" and add the following code to the bottom of our class definition:
+1. In **Visual Studio**, we should open `CollidingPawn.h` and add the following code to the bottom of our class definition:
 
 		UParticleSystemComponent *OurParticleSystem;
 
 	We will use this variable to keep track of a **Particle System Component** that we will create later. We can create **Components** without making variables to track them, but if we want to use those **Components** in our code, we should store them in class member variables like this.
 
-1. We can now open "CollidingPawn.cpp" and edit the constructor function, **ACollidingPawn::ACollidingPawn**, by adding code that will spawn a variety of useful **Components** and arrange them in a hierarchy. We will create a **Sphere Component** to interact with the physical world, a **Static Mesh Component** to represent our collision shape visually, a **Particle System Component** that we can turn on or off at will, and a **Spring Arm Component** that we can use to attach a **Camera Component** to control our in-game perspective.
+1. We can now open `CollidingPawn.cpp` and edit the constructor function, **ACollidingPawn::ACollidingPawn**, by adding code that will spawn a variety of useful **Components** and arrange them in a hierarchy. We will create a **Sphere Component** to interact with the physical world, a **Static Mesh Component** to represent our collision shape visually, a **Particle System Component** that we can turn on or off at will, and a **Spring Arm Component** that we can use to attach a **Camera Component** to control our in-game perspective.
 
 1. First up is deciding which **Component** will be the root of our hierarchy. In this tutorial, we'll go with the **Sphere Component**, since it is the physical presence that can interact and collide with our game world. Note that an **Actor** can have multiple physics-enabled **Components** in its hierarchy, but for this tutorial, we only need one.
 			// Our root component will be a sphere that reacts to physics
diff --git a/Engine/Documentation/Source/Programming/Tutorials/Components/3/Components_3.INT.udn b/Engine/Documentation/Source/Programming/Tutorials/Components/3/Components_3.INT.udn
--- a/Engine/Documentation/Source/Programming/Tutorials/Components/3/Components_3.INT.udn
+++ b/Engine/Documentation/Source/Programming/Tutorials/Components/3/Components_3.INT.udn
@@ -21,12 +21,12 @@ Description:Writing the code that will define how our Pawn Movement Component wi
 
 %Steps%
 
-1. We can go back to **Visual Studio** now and program our customized **Pawn Movement Component**. All we really need to write is a TickComponent function (similar to an **Actor's** Tick function) to tell us how to move each frame. In "CollidingPawnMovementComponent.h", we'll need to override TickComponent in the class definition:
+1. We can go back to **Visual Studio** now and program our customized **Pawn Movement Component**. All we really need to write is a TickComponent function (similar to an **Actor's** Tick function) to tell us how to move each frame. In `CollidingPawnMovementComponent.h`, we'll need to override TickComponent in the class definition:
 
 		public:
 			virtual void TickComponent(float DeltaTime, enum ELevelTick TickType, FActorComponentTickFunction *ThisTickFunction) override;	
 
-	We'll define this function in "CollidingPawnMovementComponent.cpp":
+	We'll define this function in `CollidingPawnMovementComponent.cpp`:
 
 		void UCollidingPawnMovementComponent::TickComponent(float DeltaTime, enum ELevelTick TickType, FActorComponentTickFunction *ThisTickFunction)
 		{
diff --git a/Engine/Documentation/Source/Programming/Tutorials/Components/4/Components_4.INT.udn b/Engine/Documentation/Source/Programming/Tutorials/Components/4/Components_4.INT.udn
--- a/Engine/Documentation/Source/Programming/Tutorials/Components/4/Components_4.INT.udn
+++ b/Engine/Documentation/Source/Programming/Tutorials/Components/4/Components_4.INT.udn
@@ -21,11 +21,11 @@ Description:Tying the functionality of the Pawn, the input configuration, and th
 
 %Steps%
 
-1. In order to use our custom **Pawn Movement Component**, we will first need to add a variable to our **Pawn** class to keep track of it. At the bottom of our class definition in "CollidingPawn.h", near where we added the "OurParticleSystem" variable, we should add:
+1. In order to use our custom **Pawn Movement Component**, we will first need to add a variable to our **Pawn** class to keep track of it. At the bottom of our class definition in `CollidingPawn.h`, near where we added the "OurParticleSystem" variable, we should add:
 
 		class UCollidingPawnMovementComponent* OurMovementComponent;
 
-1. Once we have a place to keep track of it, we will need to create a **Colliding Pawn Movement Component** to store in our new variable, so let's open "CollidingPawn.cpp" and add the following to the top of the file, underneath the "#include" lines, so that our code will be able to reference our new class:
+1. Once we have a place to keep track of it, we will need to create a **Colliding Pawn Movement Component** to store in our new variable, so let's open `CollidingPawn.cpp` and add the following to the top of the file, underneath the "#include" lines, so that our code will be able to reference our new class:
 
 		#include "CollidingPawnMovementComponent.h"
 
@@ -37,25 +37,25 @@ Description:Tying the functionality of the Pawn, the input configuration, and th
 
 	[REGION:note]Unlike the other **Components** we've seen so far, we don't need to attach this **Component** to our own **Component** hierarchy. This is because our other **Components** are all types of **Scene Components**, which inherently require physical locations. **Movement Controllers**, however, are not **Scene Components** and do not represent physical objects, so the concept of existing at a physical location or being physically attached to another **Component** does not apply to them.[/REGION]
 
-1. **Pawns** have a function called "GetMovementComponent" that is used to enable other classes in the engine to access the **Pawn Movement Component** that the **Pawn** is currently using. We will need to override that function so that it returns our custom **Pawn Movement Component**. In the class definition, we need to add:
+1. **Pawns** have a function called "GetMovementComponent" that is used to enable other classes in the engine to access the **Pawn Movement Component** that the **Pawn** is currently using. We will need to override that function so that it returns our custom **Pawn Movement Component**. In the class definition in `CollidingPawn.h`, we need to add:
 
 		virtual UPawnMovementComponent* GetMovementComponent() const override;
 
-	And in "CollidingPawn.cpp", we need to add the definition of our overridden function, as follows:
+	And in `CollidingPawn.cpp`, we need to add the definition of our overridden function, as follows:
 
 		UPawnMovementComponent* ACollidingPawn::GetMovementComponent() const
 		{
 			return OurMovementComponent;
 		}
 
-1. With our new **Pawn Movement Component** set up, we can create code to handle the input that our **Pawn** will receive. We will start by declaring a few functions in our class definition in "CollidingPawn.h":
+1. With our new **Pawn Movement Component** set up, we can create code to handle the input that our **Pawn** will receive. We will start by declaring a few functions in our class definition in `CollidingPawn.h`:
 
 		void MoveForward(float AxisValue);
 		void MoveRight(float AxisValue);
 		void Turn(float AxisValue);
 		void ParticleToggle();
 
-	In "CollidingPawn.cpp", we will add the definitions of those functions as follows:
+	In `CollidingPawn.cpp`, we will add the definitions of those functions as follows:
 
 			void ACollidingPawn::MoveForward(float AxisValue)
 			{
diff --git a/Engine/Documentation/Source/Programming/Tutorials/PlayerCamera/1/PlayerCamera_1.INT.udn b/Engine/Documentation/Source/Programming/Tutorials/PlayerCamera/1/PlayerCamera_1.INT.udn
--- a/Engine/Documentation/Source/Programming/Tutorials/PlayerCamera/1/PlayerCamera_1.INT.udn
+++ b/Engine/Documentation/Source/Programming/Tutorials/PlayerCamera/1/PlayerCamera_1.INT.udn
@@ -26,7 +26,7 @@ Description:Creating a spring arm and a camera.
 
 	[![](NamePawnClass.png)](NamePawnClass.png)
 
-1. Next, in **Visual Studio**, we should open "PawnWithCamera.h" and add the following code to the bottom of our class definition:
+1. Next, in **Visual Studio**, we should open `PawnWithCamera.h` and add the following code to the bottom of our class definition:
 
 		protected:
 			USpringArmComponent* OurCameraSpringArm;
@@ -34,7 +34,7 @@ Description:Creating a spring arm and a camera.
 
 	We will use these variables to create a **SpringArmComponent** with a **CameraComponent** attached to the end. Spring arms are a good, simple way to attach cameras (or other **Components**) so that they're not overly rigid and feel more fluid as they move.
 
-1. After that, we actually need to create our **Components** in our constructor. Add the following code to "PawnWithCamera.cpp" inside of **APawnWithCamera::APawnWithCamera**:
+1. After that, we actually need to create our **Components** in our constructor. Add the following code to `PawnWithCamera.cpp` inside of **APawnWithCamera::APawnWithCamera**:
 
 		//Create our components
 		RootComponent = CreateDefaultSubobject<USceneComponent>(TEXT("RootComponent"));
diff --git a/Engine/Documentation/Source/Programming/Tutorials/PlayerCamera/3/PlayerCamera_3.INT.udn b/Engine/Documentation/Source/Programming/Tutorials/PlayerCamera/3/PlayerCamera_3.INT.udn
--- a/Engine/Documentation/Source/Programming/Tutorials/PlayerCamera/3/PlayerCamera_3.INT.udn
+++ b/Engine/Documentation/Source/Programming/Tutorials/PlayerCamera/3/PlayerCamera_3.INT.udn
@@ -21,7 +21,7 @@ Description:Implementing code to control a pawn and a camera with input.
 
 %Steps%
 
-1. Our game now has input mappings that we can use, so let's set up some member variables to store the data we receive. During our update, we'll need to know the values of our movement and mouse-looking axes, each of which are two-dimensional vectors. We'll also want to know whether we should be moving toward our zoomed-in or zoomed-out view, and how far between those two states we currently are. To accomplish this, we should add the following code to our class definition in "PawnWithCamera.h":
+1. Our game now has input mappings that we can use, so let's set up some member variables to store the data we receive. During our update, we'll need to know the values of our movement and mouse-looking axes, each of which are two-dimensional vectors. We'll also want to know whether we should be moving toward our zoomed-in or zoomed-out view, and how far between those two states we currently are. To accomplish this, we should add the following code to our class definition in `PawnWithCamera.h`:
 
 		//Input variables
 		FVector2D MovementInput;
@@ -29,7 +29,7 @@ Description:Implementing code to control a pawn and a camera with input.
 		float ZoomFactor;
 		bool bZoomingIn;
 
-1. We'll need to create functions to track our input, so let's add the following to our class definition in "PawnWithCamera.h" as well:
+1. We'll need to create functions to track our input, so let's add the following to our class definition in `PawnWithCamera.h` as well:
 
 		//Input functions
 		void MoveForward(float AxisValue);
@@ -39,7 +39,7 @@ Description:Implementing code to control a pawn and a camera with input.
 		void ZoomIn();
 		void ZoomOut();
 
-	And we can now fill these functions out in "PawnWithCamera.cpp" with the following code:
+	And we can now fill these functions out in `PawnWithCamera.cpp` with the following code:
 
 		//Input functions
 		void APawnWithCamera::MoveForward(float AxisValue)
@@ -86,7 +86,7 @@ Description:Implementing code to control a pawn and a camera with input.
 		InputComponent->BindAxis("CameraPitch", this, &APawnWithCamera::PitchCamera);
 		InputComponent->BindAxis("CameraYaw", this, &APawnWithCamera::YawCamera);
 
-1. Finally, we can use those values in our **Tick** function to update our **Pawn** and **Camera** each frame. The following code blocks should all be added to **APawnWithCamera::Tick** in "PawnWithCamera.cpp":
+1. Finally, we can use those values in our **Tick** function to update our **Pawn** and **Camera** each frame. The following code blocks should all be added to **APawnWithCamera::Tick** in `PawnWithCamera.cpp`:
 
 		//Zoom in if ZoomIn button is down, zoom back out if it's not
 		{
diff --git a/Engine/Documentation/Source/Programming/Tutorials/PlayerInput/1/PlayerInput_1.INT.udn b/Engine/Documentation/Source/Programming/Tutorials/PlayerInput/1/PlayerInput_1.INT.udn
--- a/Engine/Documentation/Source/Programming/Tutorials/PlayerInput/1/PlayerInput_1.INT.udn
+++ b/Engine/Documentation/Source/Programming/Tutorials/PlayerInput/1/PlayerInput_1.INT.udn
@@ -30,19 +30,19 @@ Description:Introduction to Pawns
 
 	[![](NameYourNewPawn.png)](NameYourNewPawn.png)
 
-1. The first thing we're going to do is set our MyPawn to respond to player input automatically upon the game starting. The **Pawn** class provides a variable we can set during initialization that handles this for us. In "MyPawn.cpp", add the following code to **AMyPawn::AMyPawn**:
+1. The first thing we're going to do is set our MyPawn to respond to player input automatically upon the game starting. The **Pawn** class provides a variable we can set during initialization that handles this for us. In `MyPawn.cpp`, add the following code to **AMyPawn::AMyPawn**:
 
 		// Set this pawn to be controlled by the lowest-numbered player
 		AutoPossessPlayer = EAutoReceiveInput::Player0;
 
-1. Next, we'll build a few basic **Components**. If you want to learn more about adding and managing **Components** in code, as well as some common types of **Components** you will encounter, try reading our [%Programming\Tutorials\Components:title% tutorial](Programming\Tutorials\Components). In order to keep track of the **Component** we will create, we should add the following code to "MyPawn.h", at the bottom of our class definition:
+1. Next, we'll build a few basic **Components**. If you want to learn more about adding and managing **Components** in code, as well as some common types of **Components** you will encounter, try reading our [%Programming\Tutorials\Components:title% tutorial](Programming\Tutorials\Components). In order to keep track of the **Component** we will create, we should add the following code to `MyPawn.h`, at the bottom of our class definition:
 
 		UPROPERTY()
 		USceneComponent* OurVisibleComponent;
 
 	[REGION:note]This variable is tagged as a **UPROPERTY** so that it will be visible to **Unreal Engine**. This is important because it prevents the variable from being reset when the game is launched, or when the project or level is closed and reloaded.[/REGION]
 
-	And back in "MyPawn.cpp", we should add the following code to **AMyPawn::AMyPawn**:
+	And back in `MyPawn.cpp`, we should add the following code to **AMyPawn::AMyPawn**:
 
 		// Create a dummy root component we can attach things to.
 		RootComponent = CreateDefaultSubobject<USceneComponent>(TEXT("RootComponent"));
diff --git a/Engine/Documentation/Source/Programming/Tutorials/PlayerInput/3/PlayerInput_3.INT.udn b/Engine/Documentation/Source/Programming/Tutorials/PlayerInput/3/PlayerInput_3.INT.udn
--- a/Engine/Documentation/Source/Programming/Tutorials/PlayerInput/3/PlayerInput_3.INT.udn
+++ b/Engine/Documentation/Source/Programming/Tutorials/PlayerInput/3/PlayerInput_3.INT.udn
@@ -21,7 +21,7 @@ Description:Binding input to C++ code
 
 %Steps%
 
-1. In **Visual Studio**, open "MyPawn.h" and add the following code to the bottom of MyPawn's class definition:
+1. In **Visual Studio**, open `MyPawn.h` and add the following code to the bottom of MyPawn's class definition:
 
 		//Input functions
 		void Move_XAxis(float AxisValue);
@@ -35,7 +35,7 @@ Description:Binding input to C++ code
 
 	[REGION:note]The four input functions are going to be bound to our input events. When they run, they will update the values stored in our new input variables, which MyPawn will use to determine what it should do during in the game.[/REGION]
 
-1. We'll switch over to "MyPawn.cpp" and program the four functions we just declared. Add the following code:
+1. We'll switch over to `MyPawn.cpp` and program the four functions we just declared. Add the following code:
 
 		void AMyPawn::Move_XAxis(float AxisValue)
 		{
diff --git a/Engine/Documentation/Source/Programming/Tutorials/UMG/2/UMG_2.INT.udn b/Engine/Documentation/Source/Programming/Tutorials/UMG/2/UMG_2.INT.udn
--- a/Engine/Documentation/Source/Programming/Tutorials/UMG/2/UMG_2.INT.udn
+++ b/Engine/Documentation/Source/Programming/Tutorials/UMG/2/UMG_2.INT.udn
@@ -22,7 +22,7 @@ Description:Extending the project's Game Mode to add and remove menus from the s
 %Steps%
 
 
-1. The menus we create will be made from **User Widgets**. We'll need to write a function that creates and displays a new **User Widget**, and then call that function when the game starts. We'll also need to keep track of what we have created so that we can remove it later. Since each project already comes with a custom **Game Mode** class, we can simply open ours, which is defined in "HowTo_UMGGameMode.h". The following functions and properties will need to be added to the bottom of the class:
+1. The menus we create will be made from **User Widgets**. We'll need to write a function that creates and displays a new **User Widget**, and then call that function when the game starts. We'll also need to keep track of what we have created so that we can remove it later. Since each project already comes with a custom **Game Mode** class, we can simply open ours, which is defined in `HowTo_UMGGameMode.h`. The following functions and properties will need to be added to the bottom of the class:
 
 		public:
 			/** Called when the game starts. */
@@ -45,7 +45,7 @@ Description:Extending the project's Game Mode to add and remove menus from the s
 
 		#include "Blueprint/UserWidget.h"
 
-1. Moving to "HowTo_UMGGameMode.cpp" now, we need to fill out the bodies of the two functions we declared. We'll start with overriding BeginPlay():
+1. Moving to `HowTo_UMGGameMode.cpp` now, we need to fill out the bodies of the two functions we declared. We'll start with overriding BeginPlay():
 
 		void AHowTo_UMGGameMode::BeginPlay()
 		{
@@ -55,7 +55,7 @@ Description:Extending the project's Game Mode to add and remove menus from the s
 
 	[REGION:note]When overriding functions from a parent class (referenced by the word Super), as we do here with _BeginPlay_, it is often important to call the parent class' version of that function. Since our version of the function is only meant to add one step to the end of the existing procedure, we call _Super::BeginPlay_ in the first line of the function.[/REGION]
 
-1. Next, still in "HowTo_UMGGameMode.cpp", we need to define how we change between menus. We will need to remove whatever **User Widget** we have active from the viewport, if any, first. After that, we can create and a new **User Widget** and add it to the viewport.
+1. Next, still in `HowTo_UMGGameMode.cpp`, we need to define how we change between menus. We will need to remove whatever **User Widget** we have active from the viewport, if any, first. After that, we can create and a new **User Widget** and add it to the viewport.
 
 		void AHowTo_UMGGameMode::ChangeMenuWidget(TSubclassOf<UUserWidget> NewWidgetClass)
 		{
diff --git a/Engine/Documentation/Source/Programming/Tutorials/VariablesTimersEvents/1/VariablesTimersEvents_1.INT.udn b/Engine/Documentation/Source/Programming/Tutorials/VariablesTimersEvents/1/VariablesTimersEvents_1.INT.udn
--- a/Engine/Documentation/Source/Programming/Tutorials/VariablesTimersEvents/1/VariablesTimersEvents_1.INT.udn
+++ b/Engine/Documentation/Source/Programming/Tutorials/VariablesTimersEvents/1/VariablesTimersEvents_1.INT.udn
@@ -28,7 +28,7 @@ Description:Starting and stopping a timer that runs code at a fixed rate.
 
 	[![](NameYourActor.png)](NameYourActor.png)
 
-1. We'll start by creating a simple countdown timer that we can see in-game. In "Countdown.h", add the following lines to the end of our class definition:
+1. We'll start by creating a simple countdown timer that we can see in-game. In `Countdown.h`, add the following lines to the end of our class definition:
 
 		int32 CountdownTime;
 
@@ -36,7 +36,7 @@ Description:Starting and stopping a timer that runs code at a fixed rate.
 
 		void UpdateTimerDisplay();
 
-1. In "Countdown.cpp", we can create our renderable text **Component** and initialize our countdown time to 3 seconds. We can also turn **Ticking** off for this type of **Actor**, since we will not need it. **ACountdown::ACountdown** should look like this:
+1. In `Countdown.cpp`, we can create our renderable text **Component** and initialize our countdown time to 3 seconds. We can also turn **Ticking** off for this type of **Actor**, since we will not need it. **ACountdown::ACountdown** should look like this:
 
 		PrimaryActorTick.bCanEverTick = false;
 
@@ -55,7 +55,7 @@ Description:Starting and stopping a timer that runs code at a fixed rate.
 			CountdownText->SetText(FString::FromInt(FMath::Max(CountdownTime, 0)));
 		}
 
-1. Whenever we assign a **Timer** to run a function, we are given a **Timer Handle**. We need to hold onto that handle so that we can shut the **Timer** down when the countdown finishes. Let's add the function to count time down, and the **Timer Handle** we'll need to control it, to the class definition in "Countdown.h". While we're there, let's also add a function to do something special when the countdown ends:
+1. Whenever we assign a **Timer** to run a function, we are given a **Timer Handle**. We need to hold onto that handle so that we can shut the **Timer** down when the countdown finishes. Let's add the function to count time down, and the **Timer Handle** we'll need to control it, to the class definition in `Countdown.h`. While we're there, let's also add a function to do something special when the countdown ends:
 
 		void AdvanceTimer();
 
@@ -63,7 +63,7 @@ Description:Starting and stopping a timer that runs code at a fixed rate.
 
 		FTimerHandle CountdownTimerHandle;
 
-	We can also write the body of **ACountdown::AdvanceTimer** and **ACountdown::CountdownHasFinished** in "Countdown.cpp" now:
+	We can also write the body of **ACountdown::AdvanceTimer** and **ACountdown::CountdownHasFinished** in `Countdown.cpp` now:
 
 		void ACountdown::AdvanceTimer()
 		{
diff --git a/Engine/Documentation/Source/Programming/Tutorials/VariablesTimersEvents/2/VariablesTimersEvents_2.INT.udn b/Engine/Documentation/Source/Programming/Tutorials/VariablesTimersEvents/2/VariablesTimersEvents_2.INT.udn
--- a/Engine/Documentation/Source/Programming/Tutorials/VariablesTimersEvents/2/VariablesTimersEvents_2.INT.udn
+++ b/Engine/Documentation/Source/Programming/Tutorials/VariablesTimersEvents/2/VariablesTimersEvents_2.INT.udn
@@ -21,7 +21,7 @@ Description:Using the UPROPERTY and UFUNCTION macros to expose variables and fun
 
 %Steps%
 
-1. Our countdown timer is currently hard-coded to use a value of 3 seconds. It would be more useful if we could set the countdown time to any value we want in the editor, and this is easy to do. In **Visual Studio**, we can open "Countdown.h" and find the line that says:
+1. Our countdown timer is currently hard-coded to use a value of 3 seconds. It would be more useful if we could set the countdown time to any value we want in the editor, and this is easy to do. In **Visual Studio**, we can open `Countdown.h` and find the line that says:
 
 		int32 CountdownTime;
 
@@ -49,7 +49,7 @@ Description:Using the UPROPERTY and UFUNCTION macros to expose variables and fun
 	
 	When we return to the **Unreal Editor** and press **Compile**, our variable will appear in the **Details Panel** for the **ACountdown** we placed earlier, and we can test out different timer values by changing this number and pressing **Play**.
 
-1. In addition to changing the value of the timer, let's also enable non-programming developers to change what happens when the timer is up. In **Visual Studio**, we'll open "Countdown.h" and find the following line:
+1. In addition to changing the value of the timer, let's also enable non-programming developers to change what happens when the timer is up. In **Visual Studio**, we'll open `Countdown.h` and find the following line:
 
 		void CountdownHasFinished();
 
@@ -63,18 +63,18 @@ Description:Using the UPROPERTY and UFUNCTION macros to expose variables and fun
 	
 	1. **BlueprintCallable** functions are written in C++ and can be called from the **Blueprint Graph**, but cannot be changed or overridden without editing C++ code. Functions marked this way are usually features that have been programmed for non-programmer use, but that are not supposed to be changed or wouldn't make sense to change. An easy example of this would be any kind of math function.
 
-	1. **BlueprintImplementableEvent** functions are set up in a C++ header (".h") file, but the body of the function is written entirely in the **Blueprint Graph**, not in C++. These are usually created to give a non-programmer the ability to create custom reactions to special situations that have no expected default action or standard behavior. An example of this might be an event that happens when a powerup touches the player's ship in a spaceship game.
+	1. **BlueprintImplementableEvent** functions are set up in a C++ header (`.h`) file, but the body of the function is written entirely in the **Blueprint Graph**, not in C++. These are usually created to give a non-programmer the ability to create custom reactions to special situations that have no expected default action or standard behavior. An example of this might be an event that happens when a powerup touches the player's ship in a spaceship game.
 
 	1. **BlueprintNativeEvent** functions are like a combination of **BlueprintCallable** and **BlueprintImplementableEvent** functions. They have default behaviors programmed in C++, but these can be supplemented or replaced by overriding in the **Blueprint Graph**. When programming these, the C++ code always goes in a virtual function with "_Implementation" added to the end of the name, as shown below. This is the most flexible option, so we will use it for this tutorial.
 	[/REGION]
 
-	To grant non-programmers the ability to call our C++ function and to override it with **Blueprints**, we need to make the following changes to "Countdown.h":
+	To grant non-programmers the ability to call our C++ function and to override it with **Blueprints**, we need to make the following changes to `Countdown.h`:
 
 		UFUNCTION(BlueprintNativeEvent)
 		void CountdownHasFinished();
 		virtual void CountdownHasFinished_Implementation();
 
-	Then, in "Countdown.cpp", we will need to change the line that says:
+	Then, in `Countdown.cpp`, we will need to change the line that says:
 
 		void ACountdown::CountdownHasFinished()
 
