sp: Sensitive

diff --git a/Engine/Documentation/Source/Gameplay/HowTo/PossessPawns/Blueprints/GHT_4B.INT.udn b/Engine/Documentation/Source/Gameplay/HowTo/PossessPawns/Blueprints/GHT_4B.INT.udn
--- a/Engine/Documentation/Source/Gameplay/HowTo/PossessPawns/Blueprints/GHT_4B.INT.udn
+++ b/Engine/Documentation/Source/Gameplay/HowTo/PossessPawns/Blueprints/GHT_4B.INT.udn
@@ -170,7 +170,7 @@ _This example is based off using the **Blueprint 3rd Person and Vehicle Template
 
 1.	**Right-click** in the graph and add the **E** Key Event and connect it to a **Set Actor Transform** node. 
 
-1.	Drag off the **SpawnLocation** variable, un-check **Context Senstive**, and add the **Get World Transform** node. 
+1.	Drag off the **SpawnLocation** variable, un-check **Context Sensitive**, and add the **Get World Transform** node. 
 
 1.	Connect the **Get World Transform** to the **Set Actor Transform** node. 
 
@@ -214,7 +214,7 @@ _This example is based off using the **Blueprint 3rd Person and Vehicle Template
 
 1.	Following the **Possess** node, add a **Set Actor Hidden in Game** node (set to _True_) then a **Set Actor Enable Collision** (set to _False_). 
 
-1.	Drag off the **Car** variable, re-check the **Context Senstive** option. 
+1.	Drag off the **Car** variable, re-check the **Context Sensitive** option. 
 
 1.	Search for and add a **Cast To Sedan** node following the **Set Actor Enable Collision** node. 
 
