#UE4doc #docs: removed unneeded comma "a" to "an"

diff --git a/Engine/Documentation/Source/Gameplay/HowTo/SettingUpAGameMode/Blueprints/GameMode_BP.INT.udn b/Engine/Documentation/Source/Gameplay/HowTo/SettingUpAGameMode/Blueprints/GameMode_BP.INT.udn
--- a/Engine/Documentation/Source/Gameplay/HowTo/SettingUpAGameMode/Blueprints/GameMode_BP.INT.udn
+++ b/Engine/Documentation/Source/Gameplay/HowTo/SettingUpAGameMode/Blueprints/GameMode_BP.INT.udn
@@ -30,7 +30,7 @@ In the **Pick Parent Class** window, select the **Game Mode** Class.
 	Here we are assigning the Character Blueprint called **MyCharacter** as the **Default Pawn Class** for players to use in the game. 
 
 	[REGION:note]
-	The **Game Mode** Blueprint points to existing Blueprints of the **Character**, **HUD**, **PlayerController**, **Spectator**, and **Game State** Classes. You will need to create these separately, then specify them for use in the **Game Mode** Blueprint in order to actually use them in your game. 
+	The **Game Mode** Blueprint points to existing Blueprints of the **Character**, **HUD**, **PlayerController**, **Spectator**, and **Game State** Classes. You will need to create these separately then specify them for use in the **Game Mode** Blueprint in order to actually use them in your game. 
 	[/REGION]
 
 
