#UE4doc #docs: sp: primarily

diff --git a/Engine/Documentation/Source/Gameplay/HowTo/SetUpInput/SetUpInputOverview.INT.udn b/Engine/Documentation/Source/Gameplay/HowTo/SetUpInput/SetUpInputOverview.INT.udn
--- a/Engine/Documentation/Source/Gameplay/HowTo/SetUpInput/SetUpInputOverview.INT.udn
+++ b/Engine/Documentation/Source/Gameplay/HowTo/SetUpInput/SetUpInputOverview.INT.udn
@@ -6,7 +6,7 @@ Description:A How To Guide for Setting Up Inputs in Unreal Engine 4.
 
 ##Overview
 
-After you have imported or created your Character, Vehicle, or Flying Spaceship Asset, you will need to define how it reacts to Player **Input**. Defining Input is primiarly done through user defined **Bindings** of **Action Mappings** (which are for key presses and releases) and **Axis Mappings** (which are for inputs that have a continuous range). 
+After you have imported or created your Character, Vehicle, or Flying Spaceship Asset, you will need to define how it reacts to Player **Input**. Defining Input is primarily done through user defined **Bindings** of **Action Mappings** (which are for key presses and releases) and **Axis Mappings** (which are for inputs that have a continuous range). 
 
 Once you have defined your **Mappings**, you can then bind them to behaviors in Blueprints or C++ (see the **Setup** section below). 
 
