#UE4doc #docs: sp: singular

diff --git a/Engine/Documentation/Source/Gameplay/HowTo/UseRaycasts/GHT_2_Overview.INT.udn b/Engine/Documentation/Source/Gameplay/HowTo/UseRaycasts/GHT_2_Overview.INT.udn
--- a/Engine/Documentation/Source/Gameplay/HowTo/UseRaycasts/GHT_2_Overview.INT.udn
+++ b/Engine/Documentation/Source/Gameplay/HowTo/UseRaycasts/GHT_2_Overview.INT.udn
@@ -10,7 +10,7 @@ There may be instances in your game where you want to determine if the player ch
 
 There are several different options available when running a Trace. You can run a Trace to check for collision with any **Objects** where any Objects hit are returned, or you can run a Trace by **Trace Channel** where any Objects hit will return hit information only if the Object is set to specifically respond to a specified **Trace Channel** (which can be set via **Collision Settings**). 
 
-In addition to running traces by **Objects** or **Trace Channel**, you can run your Trace to detect **Single** hits or **Multi** hits, where a **Single Trace** returns only one a singluar hit result and a **Multi Trace** returns multiple hits resulting from the Trace. With Traces, you can also specify the type of ray that is used: either a straight line, a box, a capsule, or sphere. 
+In addition to running traces by **Objects** or **Trace Channel**, you can run your Trace to detect **Single** hits or **Multi** hits, where a **Single Trace** returns only one a singular hit result and a **Multi Trace** returns multiple hits resulting from the Trace. With Traces, you can also specify the type of ray that is used: either a straight line, a box, a capsule, or sphere. 
 
 ## Setup
 
