#UE4doc #docs: sp: detected

diff --git a/Engine/Documentation/Source/Gameplay/HowTo/UseOnHit/UseOnHit.INT.udn b/Engine/Documentation/Source/Gameplay/HowTo/UseOnHit/UseOnHit.INT.udn
--- a/Engine/Documentation/Source/Gameplay/HowTo/UseOnHit/UseOnHit.INT.udn
+++ b/Engine/Documentation/Source/Gameplay/HowTo/UseOnHit/UseOnHit.INT.udn
@@ -6,7 +6,7 @@ Description:A How To Guide to for using the OnHit Event in Blueprints or Code.
 
 ##Overview
 
-If you are creating a game that involves any kind of damage to a player, enemy, or object, chances are you will run into the situation where you need to determine if they were hit by something and if so, who hit them, where the impact point was, or other information regarding the hit dectected. The **Event Hit** Event will provide this information and much more when the collision occurs in which you can use the data to drive changes in your game be it affecting health, destroying an object, or other gameplay related actions. 
+If you are creating a game that involves any kind of damage to a player, enemy, or object, chances are you will run into the situation where you need to determine if they were hit by something and if so, who hit them, where the impact point was, or other information regarding the hit detected. The **Event Hit** Event will provide this information and much more when the collision occurs in which you can use the data to drive changes in your game be it affecting health, destroying an object, or other gameplay related actions. 
 
 
 ## Setup
