#UE4 DOCS: updates to the Behavior Tree Node Reference pages. Split it into multiple pages for easier searching and flow.

diff --git a/Engine/Documentation/Source/Engine/AI/BehaviorTrees/NodeReference/BehaviorTreeNodeReference.INT.udn b/Engine/Documentation/Source/Engine/AI/BehaviorTrees/NodeReference/BehaviorTreeNodeReference.INT.udn
--- a/Engine/Documentation/Source/Engine/AI/BehaviorTrees/NodeReference/BehaviorTreeNodeReference.INT.udn
+++ b/Engine/Documentation/Source/Engine/AI/BehaviorTrees/NodeReference/BehaviorTreeNodeReference.INT.udn
@@ -5,351 +5,28 @@ Availability:Public
 
 [TOC(start:2 end:2)]
 
-## Node Types
-
-This document describes and covers the uses of the nodes available to use in the Behavior Tree editor. Broadly the nodes are broken up into 5 categories:
-
-* **Root** - The starting execution node for the Behavior Tree. Every Behavior Tree has one.
-* **Composite** - These are the nodes that define the root of a branch and define the base rules for how that branch is executed.
-* **Task** - These are the leaves of the tree, the nodes that "do" things.
-* **Decorator** - Also known as conditionals. These attach to another node and make decisions on whether or not a branch in the tree, or even a single node, can be executed.
-* **Service** - These attach to Composite nodes, and will execute at their defined frequency as long as their branch is being executed. These are often used to make checks and to update the Blackboard. These take the place of traditional Parallel nodes in other Behavior Tree systems
-
-
-## Default
-
-### Root
+## Overview
+
+[OBJECT:TopicButtonList]
+	[PARAM:icon]
+		![](%ROOT%/tech_icon.png)(convert:false)
+	[/PARAM]
+	[PARAM:title]
+		Behavior Tree Node Types
+	[/PARAM]
+	[PARAM:description]
+	[/PARAM]
+	[PARAM:links]
+	* [Composites](Engine/AI/BehaviorTrees/NodeReference/Composites "%Engine/AI/BehaviorTrees/NodeReference/Composites:description% These are the nodes that define the root of a branch and the base rules for how that branch is executed")
+	* [Tasks](Engine/AI/BehaviorTrees/NodeReference/Tasks "%Engine/AI/BehaviorTrees/NodeReference/Tasks:description% These are the leaves of the tree, the nodes that 'do' things.")
+	* [Decorators](Engine/AI/BehaviorTrees/NodeReference/Decorators "%Engine/AI/BehaviorTrees/NodeReference/Decorators:description% Also known as conditionals. These attach to another node and make decisions on whether or not a branch in the tree, or even a single node, can be executed.")
+	* [Services](Engine/AI/BehaviorTrees/NodeReference/Services "%Engine/AI/BehaviorTrees/NodeReference/Services:description% These attach to Composite nodes, and will execute at their defined frequency as long as their branch is being executed. These are often used to make checks and to update the Blackboard. These take the place of traditional Parallel nodes in other Behavior Tree systems.")
+	[/PARAM]
+[/OBJECT]
+
+## Root
+
+The Root node is unique in the Behavior Tree and is the starting point for the Behavior Tree. It can only have one connection, and you cannot attach Decorators or Services to it. The Root Node has no properties of its own, but selecting it will show the Behavior Tree properties in the Details Panel, where you can set the Blackboard Asset of the Behavior Tree.
 
 ![](Root.png)
 
-The Root node is the starting point for the Behavior Tree. It can only have one connection, and you cannot attach Decorators or Services to it.
-
-## Composites
-
-### Selector
-
-![](Selector.png)
-
-Selector Nodes execute their children from left to right, and will stop executing its children when one of their children **Succeeds**. If a Selector's child succeeds, the Selector succeeds. If all the Selector's children fail, the Selector fails.
-
-| Property | Description |
-| -------- | ----------- |
-| **Node Name** | The name the node should display in the Behavior Tree graph. |
-	
-### Sequence
-
-![](Sequence.png)
-
-Sequence Nodes execute their children from left to right, and will stop executing its children when one of their children **Fails**. If a child fails, then the Sequence fails. If all the Sequence's children succeed, then the Sequence succeeds.
-
-| Property | Description |
-| -------- | ----------- |
-| **Node Name** | The name the node should display in the Behavior Tree graph. |
-
-### Simple Parallel
-
-![](SimpleParallel.png)
-
-The Simple Parallel node allows a single main task node to be executed along side of a full tree. When the main task finishes, the setting in **Finish Mode** dictates if the node should finish **Immediately**, aborting the secondary tree, or if it should **Delay** for the secondary tree to finish.
-
-| Property | Description |
-| -------- | ----------- |
-| **Finish Mode** | [INCLUDE:#subprop1] |
-| **Node Name** | The name the node should display in the Behavior Tree graph. |
-
-<!--
-[EXCERPT:subprop1]
-| --- | --- |
-| **Immediate** | Once the main task finishes, the background tree will be aborted.  |
-| **Delayed** | The background tree will be permitted to finish once the main task has finished. |
-[/EXCERPT:subprop1]
--->
-
-
-## Tasks
-
-### Make Noise
-	
-![](MakeNoise.png)
-
-If the controlled pawn has the **PawnNoiseEmitter** Component, the Make Noise task will cause the Pawn to "produce a noise" (send a message) that other Pawns with the **PawnSensing** Component can hear (receive the message).
-
-| Property | Description |
-| -------- | ----------- |
-| **Loudness** | How loud the generated sound is. |
-| **Node Name** | The name the node should display in the Behavior Tree graph. |
-
-### Move To
-
-![](MoveTo.png)
-
-The Move To node will cause a Pawn with a Character Movement component to move using the NavMesh to the Vector Blackboard key.
-
-| Property | Description |
-| -------- | ----------- |
-| **Acceptable Radius** | How close does the pawn have to be to the target for the **Move To** Task to succeed. |
-| **Filter Class** | Which navigation data should be used? If set to None the default navigation data will be used. |
-| **Allow Strafe** | Should the pawn be allowed to strafe? |
-| **Blackboard Key** | The location to move to (Vector3). |
-| **Node Name** | The name the node should display in the Behavior Tree graph. |
-
-### Play Sound
-
-![](PlaySound.png)
-
-The Play Sound node will play the sound given in the **Sound to Play** property.
-
-| Property | Description |
-| -------- | ----------- |
-| **Sound to Play** | The Sound Cue asset to play. |
-| **Node Name** | The name the node should display in the Behavior Tree graph. |
-
-### Run Behavior
-
-![](RunBehavior.png)
-
-The Run Behavior node enables you to run another Behavior Tree. 
-
-| Property | Description |
-| -------- | ----------- |
-| **Behavior Asset** | The Behavior Tree asset to run. |
-| **Node Name** | The name the node should display in the Behavior Tree graph. |
-
-
-### Run EQS Query
-
-![](RunEQSQuery.png)
-
-The Run EQS Query node
-
-| Property | Description |
-| -------- | ----------- |
-| **Query Template** | The Environment Query to use for this EQS Query (created in the Content Browser) |
-| **Query Params** | An Environment Query can have properties in its Generator and Tests set as Parameters, much like how Materials can have Parameters. So, by entering the names of these Parameters and setting their Values here, you can have one Environment Query that handles a broad range of scenarios. |
-| **Blackboard Key**| Vector |
-| **Node Name** | The name the node should display in the Behavior Tree graph. |
-
-
-### Wait
-
-![](Wait.png)
-
-The Behavior Tree will wait on this node until the wait time is up.
-
-| Property | Description |
-| -------- | ----------- |
-| **Wait Time** | How long to wait, in seconds. |
-| **Node Name** | The name the node should display in the Behavior Tree graph. |
-
-### Wait Blackboard Time
-
-Works just like the Wait Task node, except it will pull a Blackboard value for how long it should wait.
-
-![](waitBBTime.png)
-
-| Property | Description |
-| -------- | ----------- |
-| **Blackboard Key** | The float Blackboard Key to reference for how long to wait. |
-| **Node Name** | The name the node should display in the Behavior Tree graph. |
-
-## Services
-
-### Default Focus
-
-![](SetDefaultFocus.png)
-
-The Default Focus creates a short cut to accessing an Actor in your Blueprints and Code by setting the focus of the AI Controller. By setting the focus of the AI Controller to an Actor, you can directly access it from the AI Controller instead of accessing a Blackboard key.
-
-| Property | Description |
-| -------- | ----------- |
-| **Blackboard Key**| This takes in an Actor Key (and Object Key with its base class set to Actor) to be quickly referenced by the AI Controller using `GetFocusActor`. |
-| **Node Name** | The name the node should display in the Behavior Tree graph. |
-
-[REGION:note]
-As of 4.4 the Blackboard Key will appear to accept Vector Blackboard Keys, these will do nothing as the Set/Clear/Get Focus Actor functions operate on an Actor.
-[/REGION]
-
-## Decorators
-
-### Blackboard
-
-![](Blackboard.png)
-
-The Blackboard node will check to see if a value is set on the given Blackboard Key.
-
-| Property | Description |
-| -------- | ----------- |
-| **Notify Observer** | [INCLUDE:#subprop2] |
-| **Observer Aborts** | [INCLUDE:#subprop3] |
-| **Key Query** | [INCLUDE:#subprop4] |
-| **Blackboard Key**| The key to check. This is most useful for data types that can return `None` such as Objects, as other types may return their initialized values (0, false, {0,0,0}, etc...). |
-| **Node Name** | The name the node should display in the Behavior Tree graph. |
-
-<!--
-[EXCERPT:subprop2]
-| --- | --- |
-| **On Result Change** | Reevaluate only when the condition has changed. |
-| **On Value Change** | Reevaluate only when the observed Blackboard Key changes. |
-[/EXCERPT:subprop2]
-
-[EXCERPT:subprop3]
-| --- | --- |
-| **None** | Do not abort anything. |
-| **Self** | Abort self, and any sub-trees running under this node. |
-| **Lower Priority** | Abort any nodes to the right of this node. |
-| **Both** | Abort self, any sub-trees running under me, and any nodes to the right of this node. |
-[/EXCERPT:subprop3]
-
-[EXCERPT:subprop4]
-| --- | --- |
-| **Is Set** | Is the value set? |
-| **Is Not Set** | Is the value not set? |
-[/EXCERPT:subprop4]
--->
-
-### Compare Blackboard Entries
-
-![](CompareBBEntries.png)
-
-The Compare Blackboard Entries nodes will compare the values of two Blackboard keys and will block or allow the execution of a node based on that result (either `equal to` or `not equal to`).
-
-| Property | Description |
-| -------- | ----------- |
-| **Operator** | [INCLUDE:#subprop5] |
-| **Blackboard Key A**| The first key to use in this comparison. |
-| **Blackboard Key B**| The second key to use in this comparison. |
-| **Observer Aborts** | [INCLUDE:#subprop3] |
-| **Node Name** | The name the node should display in the Behavior Tree graph. |
-
-<!--
-[EXCERPT:subprop5]
-| --- | --- |
-| **Is Equal To** | Are the two keys equal to each other? |
-| **Is Not Equal To** | Are the two keys different to each other? |
-[/EXCERPT:subprop5]
--->
-
-### Composite
-
-![](Composite.png)
-
-The Composite Decorator Node enables you to setup more advanced logic than the built in nodes but not as complex as a full Blueprint. Once you have added a Composite Decorator to a node, **double-click** the Composite Decorator to bring up the Composite's Graph. By **right-clicking** in the graph area you can add Decorator Nodes as stand alone nodes, then wire them together through `AND` nodes, `OR` nodes, and `NOT` nodes, to create more advanced logic:
-
-![](CompositeGraph.png)(w:940)
-
-| Property | Description |
-| -------- | ----------- |
-| **Composite Name** | The name the node should display in the Behavior Tree graph. |
-| **Show Operations** | This will cause the operations to be listed on the node in plain text. |
-
-
-### Cone Check
-
-![](ConeCheck.png)
-
-The Cone Check decorator takes in 3 vector keys: the first for the location to start the cone, the second to define the direction the cone points, and the third for the location to check if it is inside the cone. You define the angle of the cone by using the **Cone Half Angle** property.
-
-| Property | Description |
-| -------- | ----------- |
-| **Cone Half Angle** | The half angle of the cone. So for a 90 degree cone, this value should be set to 45 degrees. |
-| **Observer Aborts** | [INCLUDE:#subprop3] |
-| **Cone Origin** | The location where the cone should start (the tip of the cone). |
-| **Cone Direction** | The direction in which the mouth of the cone should point. |
-| **Observed** | The location or Actor that is being checked to see if it is in the cone. |
-| **Inverse Condition** | Should the result be reverse (_true_ to _false_, _false_ to _true_). |
-| **Node Name** | The name the node should display in the Behavior Tree graph. |
-
-### Cooldown
-
-![](Cooldown.png)
-
-The Cooldown node will lock the execution of a node or branch until the cooldown time expires.
-
-| Property | Description |
-| -------- | ----------- |
-| **Cool Down time** | The time, in seconds, the Cooldown decorator should lock execution of this node. |
-| **Observer Aborts** | [INCLUDE:#subprop3] |
-| **Node Name** | The name the node should display in the Behavior Tree graph. |
-
-### Does Path Exist
-
-![](DoesPathExist.png)
-
-The Does Path Exist node checks to see if a path can be made from the two vectors: Blackboard Key A and Blackboard Key B.
-
-| Property | Description |
-| -------- | ----------- |
-| **Blackboard Key A**| The first location in the path. |
-| **Blackboard Key B**| The second location in the path. |
-| **Path Query Type** | [INCLUDE:#subprop6] |
-| **Inverse Condition** | Should the result be reverse (_true_ to _false_, _false_ to _true_). |
-| **Observer Aborts** | [INCLUDE:#subprop3] |
-| **Filter Class** | Which navigation data should be used? If set to None the default navigation data will be used. |
-| **Node Name** | The name the node should display in the Behavior Tree graph. |
-
-<!--
-[EXCERPT:subprop6]
-| --- | --- |
-| **NavMesh Raycast 2D** | Really Fast. |
-| **Hierarchical Query** | Fast. |
-| **Regular Path Finding** | Slow. |
-[/EXCERPT:subprop6]
--->
-
-### Force Success
-
-![](ForceSuccess.png)
-
-The Force Success decorator changes the node result to a success.
-
-| Property | Description |
-| -------- | ----------- |
-| **Node Name** | The name the node should display in the Behavior Tree graph. |
-
-
-### Keep in Cone
-
-![](KeepInCone.png)
-
-| Property | Description |
-| -------- | ----------- |
-| **Cone Half Angle** | The half angle of the cone. So for a 90 degree cone, this value should be set to 45 degrees. |
-| **Observer Aborts** | [INCLUDE:#subprop3] |
-| **Cone Origin** | The location where the cone should start (the tip of the cone) |
-| **Observed** | The location or Actor to keep in the cone. |
-| **Node Name** | The name the node should display in the Behavior Tree graph. |
-
-### Loop 
-	
-![](Loop.png)
-
-Loops the node or branch a number of times, or infinitely.
-
-| Property | Description |
-| -------- | ----------- |
-| **Num Loops** | The number of loops to run. |
-| **Infinite Loop** | Should this loop run indefinitely? |
-| **Node Name** | The name the node should display in the Behavior Tree graph. |
-
-### Reached Move Goal
-
-![](ReachedMoveGoal.png)
-
-The Reached Move Goal will check the path of the Pawn (with a Character Movement component) and return a success if the Pawn has reached its movement goal.
-
-| Property | Description |
-| -------- | ----------- |
-| **Inverse Condition** | Should the result be reverse (_true_ to _false_, _false_ to _true_). |
-| **Node Name** | The name the node should display in the Behavior Tree graph. |
-
-### Time Limit
-
-![](TimeLimit.png)
-
-The Time Limit decorator will give a branch or node a set amount of time to finish before stopping it and failing out. The timer is reset every time the node gains focus.
-
-| Property | Description |
-| -------- | ----------- |
-| **Time Limit** | The time limit before the node fails out, in seconds. |
-| **Observer Aborts** | [INCLUDE:#subprop3] |
-| **Node Name** | The name the node should display in the Behavior Tree graph. |
\ No newline at end of file
diff --git a/Engine/Documentation/Source/Engine/AI/BehaviorTrees/NodeReference/Composites/Composites.INT.udn b/Engine/Documentation/Source/Engine/AI/BehaviorTrees/NodeReference/Composites/Composites.INT.udn
new file mode 100644
--- /dev/null
+++ b/Engine/Documentation/Source/Engine/AI/BehaviorTrees/NodeReference/Composites/Composites.INT.udn
@@ -0,0 +1,50 @@
+Title: Behavior Tree Node Reference: Composites
+Crumbs:%ROOT%, Engine, Gameplay, Engine/AI/BehaviorTrees
+Description:Reference for the Behavior Tree Composite nodes.
+Availability:Public
+
+[TOC(start:2 end:2)]
+
+## Overview
+
+**Composite Nodes** define the root of a branch and the base rules for how that branch is executed. They can have [**Decorators**](Engine/AI/BehaviorTrees/NodeReference/Decorators) applied to them to modify entry into their branch or even cancel out mid execution. Also, they can have [**Services**](Engine/AI/BehaviorTrees/NodeReference/Services) attached to them that will only be active if the children of the Composite are being executed.
+
+## Selector
+
+![](Selector.png)
+
+Selector Nodes execute their children from left to right, and will stop executing its children when one of their children **Succeeds**. If a Selector's child succeeds, the Selector succeeds. If all the Selector's children fail, the Selector fails.
+
+| Property | Description |
+| -------- | ----------- |
+| **Node Name** | The name the node should display in the Behavior Tree graph. |
+	
+## Sequence
+
+![](Sequence.png)
+
+Sequence Nodes execute their children from left to right, and will stop executing its children when one of their children **Fails**. If a child fails, then the Sequence fails. If all the Sequence's children succeed, then the Sequence succeeds.
+
+| Property | Description |
+| -------- | ----------- |
+| **Node Name** | The name the node should display in the Behavior Tree graph. |
+
+## Simple Parallel
+
+![](SimpleParallel.png)
+
+The Simple Parallel node allows a single main task node to be executed along side of a full tree. When the main task finishes, the setting in **Finish Mode** dictates if the node should finish **Immediately**, aborting the secondary tree, or if it should **Delay** for the secondary tree to finish.
+
+| Property | Description |
+| -------- | ----------- |
+| **Finish Mode** | [INCLUDE:#subprop1] |
+| **Node Name** | The name the node should display in the Behavior Tree graph. |
+
+<!--
+[EXCERPT:subprop1]
+| --- | --- |
+| **Immediate** | Once the main task finishes, the background tree will be aborted.  |
+| **Delayed** | The background tree will be permitted to finish once the main task has finished. |
+[/EXCERPT:subprop1]
+-->
+
diff --git a/Engine/Documentation/Source/Engine/AI/BehaviorTrees/NodeReference/Decorators/Decorators.INT.udn b/Engine/Documentation/Source/Engine/AI/BehaviorTrees/NodeReference/Decorators/Decorators.INT.udn
new file mode 100644
--- /dev/null
+++ b/Engine/Documentation/Source/Engine/AI/BehaviorTrees/NodeReference/Decorators/Decorators.INT.udn
@@ -0,0 +1,251 @@
+Title: Behavior Tree Node Reference: Decorators
+Crumbs:%ROOT%, Engine, Gameplay, Engine/AI/BehaviorTrees
+Description:Reference for the Behavior Tree Decorator nodes.
+Availability:Public
+
+[TOC(start:2 end:2)]
+
+## Overview
+
+**Decorator**, also known as conditionals in other Behavior Tree systems, are attached to either a [**Composite**](Engine/AI/BehaviorTrees/NodeReference/Composites) or a [**Task**](Engine/AI/BehaviorTrees/NodeReference/Tasks) node and define whether or not a branch in the tree, or even a single node, can be executed.
+
+## Blackboard
+
+![](Blackboard.png)
+
+The Blackboard node will check to see if a value is set on the given Blackboard Key.
+
+| Property | Description |
+| -------- | ----------- |
+| **Notify Observer** | [INCLUDE:#subprop2] |
+| **Observer Aborts** | [INCLUDE:#subprop3] |
+| **Key Query** | [INCLUDE:#subprop4] |
+| **Blackboard Key**| The key to check. This is most useful for data types that can return `None` such as Objects, as other types may return their initialized values (0, false, {0,0,0}, etc...). |
+| **Node Name** | The name the node should display in the Behavior Tree graph. |
+
+<!--
+[EXCERPT:subprop2]
+| --- | --- |
+| **On Result Change** | Reevaluate only when the condition has changed. |
+| **On Value Change** | Reevaluate only when the observed Blackboard Key changes. |
+[/EXCERPT:subprop2]
+
+[EXCERPT:subprop3]
+| --- | --- |
+| **None** | Do not abort anything. |
+| **Self** | Abort self, and any sub-trees running under this node. |
+| **Lower Priority** | Abort any nodes to the right of this node. |
+| **Both** | Abort self, any sub-trees running under me, and any nodes to the right of this node. |
+[/EXCERPT:subprop3]
+
+[EXCERPT:subprop4]
+| --- | --- |
+| **Is Set** | Is the value set? |
+| **Is Not Set** | Is the value not set? |
+[/EXCERPT:subprop4]
+-->
+
+## Compare Blackboard Entries
+
+![](CompareBBEntries.png)
+
+The Compare Blackboard Entries nodes will compare the values of two Blackboard keys and will block or allow the execution of a node based on that result (either `equal to` or `not equal to`).
+
+| Property | Description |
+| -------- | ----------- |
+| **Operator** | [INCLUDE:#subprop5] |
+| **Blackboard Key A**| The first key to use in this comparison. |
+| **Blackboard Key B**| The second key to use in this comparison. |
+| **Observer Aborts** | [INCLUDE:#subprop3] |
+| **Node Name** | The name the node should display in the Behavior Tree graph. |
+
+<!--
+[EXCERPT:subprop5]
+| --- | --- |
+| **Is Equal To** | Are the two keys equal to each other? |
+| **Is Not Equal To** | Are the two keys different to each other? |
+[/EXCERPT:subprop5]
+-->
+
+## Composite
+
+![](Composite.png)
+
+The Composite Decorator Node enables you to setup more advanced logic than the built in nodes but not as complex as a full Blueprint. Once you have added a Composite Decorator to a node, **double-click** the Composite Decorator to bring up the Composite's Graph. By **right-clicking** in the graph area you can add Decorator Nodes as stand alone nodes, then wire them together through `AND` nodes, `OR` nodes, and `NOT` nodes, to create more advanced logic:
+
+![](CompositeGraph.png)(w:940)
+
+| Property | Description |
+| -------- | ----------- |
+| **Composite Name** | The name the node should display in the Behavior Tree graph. |
+| **Show Operations** | This will cause the operations to be listed on the node in plain text. |
+
+## Conditional Loop
+
+As long as the **Key Query** condition is met, this Decorator will have the node it's attached to loop.
+
+![](ConditionalLoop.png)
+
+| Property | Description |
+| -------- | ----------- |
+| **Node Name** | The name the node should display in the Behavior Tree graph. |
+| **Key Query** | [INCLUDE:#subprop4] |
+| **Blackboard Key** | The key to check. This is most useful for data types that can return `None` such as Objects, as other types may return their initialized values (0, false, {0,0,0}, etc...). |
+
+
+## Cone Check
+
+![](ConeCheck.png)
+
+The Cone Check decorator takes in 3 vector keys: the first for the location to start the cone, the second to define the direction the cone points, and the third for the location to check if it is inside the cone. You define the angle of the cone by using the **Cone Half Angle** property.
+
+| Property | Description |
+| -------- | ----------- |
+| **Cone Half Angle** | The half angle of the cone. So for a 90 degree cone, this value should be set to 45 degrees. |
+| **Observer Aborts** | [INCLUDE:#subprop3] |
+| **Cone Origin** | The location where the cone should start (the tip of the cone). |
+| **Cone Direction** | The direction in which the mouth of the cone should point. |
+| **Observed** | The location or Actor that is being checked to see if it is in the cone. |
+| **Inverse Condition** | Should the result be reverse (_true_ to _false_, _false_ to _true_). |
+| **Node Name** | The name the node should display in the Behavior Tree graph. |
+
+## Cooldown
+
+![](Cooldown.png)
+
+The Cooldown node will lock the execution of a node or branch until the cooldown time expires.
+
+| Property | Description |
+| -------- | ----------- |
+| **Cool Down time** | The time, in seconds, the Cooldown decorator should lock execution of this node. |
+| **Observer Aborts** | [INCLUDE:#subprop3] |
+| **Node Name** | The name the node should display in the Behavior Tree graph. |
+
+## Does Path Exist
+
+![](DoesPathExist.png)
+
+The Does Path Exist node checks to see if a path can be made from the two vectors: Blackboard Key A and Blackboard Key B.
+
+| Property | Description |
+| -------- | ----------- |
+| **Blackboard Key A**| The first location in the path. |
+| **Blackboard Key B**| The second location in the path. |
+| **Path Query Type** | [INCLUDE:#subprop6] |
+| **Inverse Condition** | Should the result be reverse (_true_ to _false_, _false_ to _true_). |
+| **Observer Aborts** | [INCLUDE:#subprop3] |
+| **Filter Class** | Which navigation data should be used? If set to None the default navigation data will be used. |
+| **Node Name** | The name the node should display in the Behavior Tree graph. |
+
+<!--
+[EXCERPT:subprop6]
+| --- | --- |
+| **NavMesh Raycast 2D** | Really Fast. |
+| **Hierarchical Query** | Fast. |
+| **Regular Path Finding** | Slow. |
+[/EXCERPT:subprop6]
+-->
+
+## Force Success
+
+![](ForceSuccess.png)
+
+The Force Success decorator changes the node result to a success.
+
+| Property | Description |
+| -------- | ----------- |
+| **Node Name** | The name the node should display in the Behavior Tree graph. |
+
+## Gameplay Tag Condition
+
+![](GameplayTagCondition.png)
+
+| Property | Description |
+| -------- | ----------- |
+| **Actor to Check** | The Blackboard Key that contains the reference to the Actor to check. |
+| **Tags to Match** | Should the Decorator match **Any** or  **All** of the Tags listed in the property **Gameplay Tags**. |
+| **Gameplay Tags** | Loads the Gameplay Tag editor to select the tags that should be used in this Decorator. |
+| **Inverse Condition** | Should the result of this Decorator be inversed? (false becomes true, true becomes false) |
+| **Node Name** | The name the node should display in the Behavior Tree graph. |
+
+## Is Blackboard Value of Given Class
+
+![](IsBlackboardValueOfGivenClass.png)
+
+| Property | Description |
+| -------- | ----------- |
+| **Test Class** | The Class type that is tested against Class of the Object of the Blackboard Key property. |
+| **Blackboard Key** | The Blackboard Key that is tested.  |
+| **Observer Aborts** | [INCLUDE:#subprop3] |
+| **Node Name** | The name the node should display in the Behavior Tree graph. |
+
+## Keep in Cone
+
+![](KeepInCone.png)
+
+| Property | Description |
+| -------- | ----------- |
+| **Cone Half Angle** | The half angle of the cone. So for a 90 degree cone, this value should be set to 45 degrees. |
+| **Observer Aborts** | [INCLUDE:#subprop3] |
+| **Cone Origin** | The location where the cone should start (the tip of the cone) |
+| **Observed** | The location or Actor to keep in the cone. |
+| **Node Name** | The name the node should display in the Behavior Tree graph. |
+
+## Loop 
+	
+![](Loop.png)
+
+Loops the node or branch a number of times, or infinitely.
+
+| Property | Description |
+| -------- | ----------- |
+| **Num Loops** | The number of loops to run. |
+| **Infinite Loop** | Should this loop run indefinitely? |
+| **Node Name** | The name the node should display in the Behavior Tree graph. |
+
+## Reached Move Goal
+
+![](ReachedMoveGoal.png)
+
+The Reached Move Goal will check the path of the Pawn (with a Character Movement component) and return a success if the Pawn has reached its movement goal.
+
+| Property | Description |
+| -------- | ----------- |
+| **Inverse Condition** | Should the result be reverse (_true_ to _false_, _false_ to _true_). |
+| **Node Name** | The name the node should display in the Behavior Tree graph. |
+
+## Set Tag Cooldown
+
+![](SetTagCooldown.png)
+
+| Property | Description |
+| -------- | ----------- |
+| **Cooldown Tag** | The GameplayTag that will be used for the Cooldown. |
+| **Cooldown Duration** | How long the Cooldown is in seconds. |
+| **Add to Existing Duration** | If there is an existing Cooldown on the given GameplayTag, should we add more? |
+| **Node Name** | The name the node should display in the Behavior Tree graph. |
+
+## Tag Cooldown
+
+![](TagCooldown.png)
+
+| Property | Description |
+| -------- | ----------- |
+| **Cooldown Tag** | The GameplayTag that will be used for the Cooldown. |
+| **Cooldown Duration** | How long the Cooldown is in seconds. |
+| **Add to Existing Duration** | If there is an existing Cooldown on the given GameplayTag, should we add more? |
+| **Adds/Sets Cooldown on Deactivation** | Whether or not we are adding/setting to the CooldownTag's value when the decorator deactivates. |
+| **Observer Aborts** | [INCLUDE:#subprop3] |
+| **Node Name** | The name the node should display in the Behavior Tree graph. |
+
+## Time Limit
+
+![](TimeLimit.png)
+
+The Time Limit decorator will give a branch or node a set amount of time to finish before stopping it and failing out. The timer is reset every time the node gains focus.
+
+| Property | Description |
+| -------- | ----------- |
+| **Time Limit** | The time limit before the node fails out, in seconds. |
+| **Observer Aborts** | [INCLUDE:#subprop3] |
+| **Node Name** | The name the node should display in the Behavior Tree graph. |
\ No newline at end of file
diff --git a/Engine/Documentation/Source/Engine/AI/BehaviorTrees/NodeReference/Services/Services.INT.udn b/Engine/Documentation/Source/Engine/AI/BehaviorTrees/NodeReference/Services/Services.INT.udn
new file mode 100644
--- /dev/null
+++ b/Engine/Documentation/Source/Engine/AI/BehaviorTrees/NodeReference/Services/Services.INT.udn
@@ -0,0 +1,27 @@
+Title: Behavior Tree Node Reference: Services
+Crumbs:%ROOT%, Engine, Gameplay, Engine/AI/BehaviorTrees
+Description:Reference for the Behavior Tree Service nodes.
+Availability:Public
+
+[TOC(start:2 end:2)]
+
+## Overview
+
+**Services** attach to [**Composite**](Engine/AI/BehaviorTrees/NodeReference/Composites) nodes, and will execute at their defined frequency as long as their branch is being executed. These are often used to make checks and to update the Blackboard. These take the place of traditional Parallel nodes in other Behavior Tree systems
+
+
+## Default Focus
+
+![](SetDefaultFocus.png)
+
+The Default Focus creates a short cut to accessing an Actor in your Blueprints and Code by setting the focus of the AI Controller. By setting the focus of the AI Controller to an Actor, you can directly access it from the AI Controller instead of accessing a Blackboard key.
+
+| Property | Description |
+| -------- | ----------- |
+| **Blackboard Key**| This takes in an Actor Key (and Object Key with its base class set to Actor) to be quickly referenced by the AI Controller using `GetFocusActor`. |
+| **Node Name** | The name the node should display in the Behavior Tree graph. |
+
+[REGION:note]
+As of 4.4 the Blackboard Key will appear to accept Vector Blackboard Keys, these will do nothing as the Set/Clear/Get Focus Actor functions operate on an Actor.
+[/REGION]
+
diff --git a/Engine/Documentation/Source/Engine/AI/BehaviorTrees/NodeReference/Tasks/Tasks.INT.udn b/Engine/Documentation/Source/Engine/AI/BehaviorTrees/NodeReference/Tasks/Tasks.INT.udn
new file mode 100644
--- /dev/null
+++ b/Engine/Documentation/Source/Engine/AI/BehaviorTrees/NodeReference/Tasks/Tasks.INT.udn
@@ -0,0 +1,94 @@
+Title: Behavior Tree Node Reference: Tasks
+Crumbs:%ROOT%, Engine, Gameplay, Engine/AI/BehaviorTrees
+Description:Reference for the Behavior Tree Task nodes.
+Availability:Public
+
+[TOC(start:2 end:2)]
+
+## Overview
+
+This document is a reference page for the **Task Nodes** available in the Behavior Tree Editor. Tasks are nodes that "do" things, like move an AI, or adjust Blackboard values. They can have [**Decorators**](Engine/AI/BehaviorTrees/NodeReference/Decorators) attached to them.
+
+## Make Noise
+	
+![](MakeNoise.png)
+
+If the controlled pawn has the **PawnNoiseEmitter** Component, the Make Noise task will cause the Pawn to "produce a noise" (send a message) that other Pawns with the **PawnSensing** Component can hear (receive the message).
+
+| Property | Description |
+| -------- | ----------- |
+| **Loudness** | How loud the generated sound is. |
+| **Node Name** | The name the node should display in the Behavior Tree graph. |
+
+## Move To
+
+![](MoveTo.png)
+
+The Move To node will cause a Pawn with a Character Movement component to move using the NavMesh to the Vector Blackboard key.
+
+| Property | Description |
+| -------- | ----------- |
+| **Acceptable Radius** | How close does the pawn have to be to the target for the **Move To** Task to succeed. |
+| **Filter Class** | Which navigation data should be used? If set to None the default navigation data will be used. |
+| **Allow Strafe** | Should the pawn be allowed to strafe? |
+| **Blackboard Key** | The location to move to (Vector3). |
+| **Node Name** | The name the node should display in the Behavior Tree graph. |
+
+## Play Sound
+
+![](PlaySound.png)
+
+The Play Sound node will play the sound given in the **Sound to Play** property.
+
+| Property | Description |
+| -------- | ----------- |
+| **Sound to Play** | The Sound Cue asset to play. |
+| **Node Name** | The name the node should display in the Behavior Tree graph. |
+
+## Run Behavior
+
+![](RunBehavior.png)
+
+The Run Behavior node enables you to run another Behavior Tree. 
+
+| Property | Description |
+| -------- | ----------- |
+| **Behavior Asset** | The Behavior Tree asset to run. |
+| **Node Name** | The name the node should display in the Behavior Tree graph. |
+
+
+## Run EQS Query
+
+![](RunEQSQuery.png)
+
+The Run EQS Query node
+
+| Property | Description |
+| -------- | ----------- |
+| **Query Template** | The Environment Query to use for this EQS Query (created in the Content Browser) |
+| **Query Params** | An Environment Query can have properties in its Generator and Tests set as Parameters, much like how Materials can have Parameters. So, by entering the names of these Parameters and setting their Values here, you can have one Environment Query that handles a broad range of scenarios. |
+| **Blackboard Key**| Vector |
+| **Node Name** | The name the node should display in the Behavior Tree graph. |
+
+
+## Wait
+
+![](Wait.png)
+
+The Behavior Tree will wait on this node until the wait time is up.
+
+| Property | Description |
+| -------- | ----------- |
+| **Wait Time** | How long to wait, in seconds. |
+| **Node Name** | The name the node should display in the Behavior Tree graph. |
+
+## Wait Blackboard Time
+
+Works just like the Wait Task node, except it will pull a Blackboard value for how long it should wait.
+
+![](waitBBTime.png)
+
+| Property | Description |
+| -------- | ----------- |
+| **Blackboard Key** | The float Blackboard Key to reference for how long to wait. |
+| **Node Name** | The name the node should display in the Behavior Tree graph. |
