#UE4doc #docs: Changed "Setup" to "Set Up" in title, Changed instances of "MyCharacter" to "FirstPersonCharacter"

diff --git a/Engine/Documentation/Source/Gameplay/HowTo/ActorInput/Blueprints/ActorInput_BP.INT.udn b/Engine/Documentation/Source/Gameplay/HowTo/ActorInput/Blueprints/ActorInput_BP.INT.udn
--- a/Engine/Documentation/Source/Gameplay/HowTo/ActorInput/Blueprints/ActorInput_BP.INT.udn
+++ b/Engine/Documentation/Source/Gameplay/HowTo/ActorInput/Blueprints/ActorInput_BP.INT.udn
@@ -1,5 +1,5 @@
 Availability: Public
-Title:How to Setup Input on an Actor in Blueprints
+Title:How to Set Up Input on an Actor in Blueprints
 Crumbs: %ROOT%, Gameplay, Gameplay\HowTo\ActorInput
 Description:A How To Guide for setting up input on an Actor in Blueprints. 
 
@@ -18,7 +18,7 @@ The **Enable/Disable Input Usage** section also covers an approach to Enabling/D
 
 The steps below will show you how to Enable Input for an Actor:
 
-1.	Inside the **Game/Blueprints** folder, open the **Blueprint_Effect_Fire** asset. 
+1.	Inside the **StarterContent/Blueprints** folder, open the **Blueprint_Effect_Fire** asset. 
 
 1.	On the **Event Graph** tab, **Right-click** in the graph and search for then add the **Event Begin Play** node. 
 
@@ -56,9 +56,9 @@ While the methods noted above work for Enabling/Disabling Input, typically you w
 
 We would want to Enable/Disable Input based on if the player is near the light, similar to the below:  
 
-1.	Inside the **Game/Blueprints** folder, open the **Blueprint_CeilingLight** asset. 
+1.	Inside the **StarterContent/Blueprints** folder, open the **Blueprint_CeilingLight** asset. 
 
-1.	In the **Components** panel, click the **Add Component**, then search for and add a **Sphere**. 
+1.	In the **Components** panel, click the **Add Component**, then search for and add **Sphere Collision**. 
 
 	![](input_3.png)
 
@@ -68,7 +68,7 @@ We would want to Enable/Disable Input based on if the player is near the light,
 
 	We are going to Enable Input (when the player is inside) or Disable Input (when the player is outside) of the sphere. 
 
-1.	Click the **Graph** tab in the upper right of the window to return to the **Event Graph**. 
+1.	Return to the **Event Graph**. 
 
 1.	In the **MyBlueprint** window, **Right-click** on the **Trigger** and choose **Add Event -> Add OnComponentBeginOverlap**. 
 
@@ -80,13 +80,13 @@ We would want to Enable/Disable Input based on if the player is near the light,
 
 	These two Event nodes should be added to the graph. 
 
-1.	Drag off the **Other Actor** pin of the Begin Overlap node and search for and add the **Cast To MyCharacter** node. 
+1.	Drag off the **Other Actor** pin of the Begin Overlap node and search for and add the **Cast To FirstPersonCharacter** node. 
 
 	![](input_6.png)
 
-	Here we are saying that when an Actor called **MyCharacter** (which is the default player character for the Project Template) overlaps the **Trigger**, do something. 
+	Here we are saying that when an Actor called **FirstPersonCharacter** (which is the default player character for the Project Template) overlaps the **Trigger**, do something. 
 
-1.	Drag off the **Other Actor** pin of the End Overlap node and add the **Cast To MyCharacter** node. 
+1.	Drag off the **Other Actor** pin of the End Overlap node and add the **Cast To FirstPersonCharacter** node. 
 
 1.	**Right-click** and add the **Get Player Controller** node, **Enable Input**, and **Disable Input** nodes. 
 
@@ -100,7 +100,7 @@ We would want to Enable/Disable Input based on if the player is near the light,
 
 	![](input_8.png)
 
-1.	**Compile** and **Save**, then close the Blueprint. 
+1.	**Compile** and **Save**, then close the Blueprint. (If you receive a Warning, ignore it and proceed to the next step.)
 
 1.	Drag the **Blueprint_CeilingLight** into the level, then click the **Play** Button to play in the Editor. 
 
