#UE4doc #docs: removed unnecessary capitalization

diff --git a/Engine/Documentation/Source/Gameplay/HowTo/UseRaycasts/Blueprints/GHT_2B.INT.udn b/Engine/Documentation/Source/Gameplay/HowTo/UseRaycasts/Blueprints/GHT_2B.INT.udn
--- a/Engine/Documentation/Source/Gameplay/HowTo/UseRaycasts/Blueprints/GHT_2B.INT.udn
+++ b/Engine/Documentation/Source/Gameplay/HowTo/UseRaycasts/Blueprints/GHT_2B.INT.udn
@@ -144,7 +144,7 @@ Below you will find steps for setting up a **Multi Line Trace by Channel**.
 
 	![](GHT2B_18.png)
 
-	Since multiple Actors are hit, For Each of them we want to do something (in this example, print the Actors to the screen). 
+	Since multiple Actors are hit, for each of them we want to do something (in this example, print the Actors to the screen). 
 
 1.  Drag off the **Array Element** and add a **Break Hit Result**, then off the **Hit Actor**, add a **To String (Object)** and connect to a **Print String**. 
 	
