#UE4 DOCS: added overview

diff --git a/Engine/Documentation/Source/Engine/AI/BehaviorTrees/NodeReference/BehaviorTreeNodeReference.INT.udn b/Engine/Documentation/Source/Engine/AI/BehaviorTrees/NodeReference/BehaviorTreeNodeReference.INT.udn
--- a/Engine/Documentation/Source/Engine/AI/BehaviorTrees/NodeReference/BehaviorTreeNodeReference.INT.udn
+++ b/Engine/Documentation/Source/Engine/AI/BehaviorTrees/NodeReference/BehaviorTreeNodeReference.INT.udn
@@ -7,6 +7,18 @@ Availability:Public
 
 ## Overview
 
+This is a reference for the default nodes available in the Behavior Tree Editor. You may see more depending on the project you are working with (such as Shooter Game). There are five types of Behavior Tree Nodes:
+
+| Node Type | Description |
+| --- | --- |
+| **Composite** | These are the nodes that define the root of a branch and the base rules for how that branch is executed |
+| **Task** | These are the leaves of the Behavior Tree, the nodes that "do" things and don't have an output connection. |
+| **Decorator** | Also known as conditionals. These attach to another node and make decisions on whether or not a branch in the tree, or even a single node, can be executed. |
+| **Service** | These attach to Composite nodes, and will execute at their defined frequency as long as their branch is being executed. These are often used to make checks and to update the Blackboard. These take the place of traditional Parallel nodes in other Behavior Tree systems. |
+| **Root** | [INCLUDE:#rootNode] |
+
+See the links below to see full listings of the default Behavior Tree Nodes.
+
 [OBJECT:TopicButtonList]
 	[PARAM:icon]
 		![](%ROOT%/tech_icon.png)(convert:false)
@@ -24,9 +36,11 @@ Availability:Public
 	[/PARAM]
 [/OBJECT]
 
-## Root
 
+<!--
+[EXCERPT:rootNode]
 The Root node is unique in the Behavior Tree and is the starting point for the Behavior Tree. It can only have one connection, and you cannot attach Decorators or Services to it. The Root Node has no properties of its own, but selecting it will show the Behavior Tree properties in the Details Panel, where you can set the Blackboard Asset of the Behavior Tree.
 
 ![](Root.png)
-
+[/EXCERPT]
+-->
\ No newline at end of file
