updated docs

diff --git a/Engine/Documentation/Source/Programming/Rendering/ShaderDevelopment/AsyncCompute/AsyncCompute.INT.udn b/Engine/Documentation/Source/Programming/Rendering/ShaderDevelopment/AsyncCompute/AsyncCompute.INT.udn
--- a/Engine/Documentation/Source/Programming/Rendering/ShaderDevelopment/AsyncCompute/AsyncCompute.INT.udn
+++ b/Engine/Documentation/Source/Programming/Rendering/ShaderDevelopment/AsyncCompute/AsyncCompute.INT.udn
@@ -46,4 +46,15 @@ Async compute context jobs are not captured in GPU captures, so these captures c
 async compute is enabled. For graphics debugging purposes, async compute should be disabled using the above cvar.
 Async compute is supported by PIX timing captures. These show up in the timeline like this:
 
-![](PixTimingCapture.png)
\ No newline at end of file
+![](PixTimingCapture.png)
+
+## Thanks and Future
+
+This feature was implemented by Lionhead Studios. We integrated it and indend to make use of it as a tool to optimize the XboxOne rendering.
+
+As more more APIs expose the hardware feature we would like make the system more cross platform. Features that make use use AsyncCompute you always
+be able to run without (console variable / define) to run on other platforms and easier debugging and profiling. AsyncCompute should be used with caution
+as it can cause more unpredicatble performance and requires more coding effort for synchromization.
+
+
+
