#UE4doc #docs: Updated 3 images

diff --git a/Engine/Documentation/Source/Gameplay/HowTo/AddingComponents/Blueprints/Components_BP.INT.udn b/Engine/Documentation/Source/Gameplay/HowTo/AddingComponents/Blueprints/Components_BP.INT.udn
--- a/Engine/Documentation/Source/Gameplay/HowTo/AddingComponents/Blueprints/Components_BP.INT.udn
+++ b/Engine/Documentation/Source/Gameplay/HowTo/AddingComponents/Blueprints/Components_BP.INT.udn
@@ -15,7 +15,7 @@ On this page, you will learn how to add **Components** to a **Blueprint** and ho
 
 The steps below will show you how to Add Components to a Blueprint Actor. 
 
-1.	In the **Content Browser**, under the **Game/Props** folder, **Right-click** on the **SM_Bush** asset and select **Create Blueprint Using**. 
+1.	In the **Content Browser**, under the **StarterContent/Props** folder, **Right-click** on the **SM_Bush** asset, scroll up to **Asset Actions**, and select **Create Blueprint Using This...**. 
 
 	![](Components1.jpg)
 
@@ -23,9 +23,9 @@ The steps below will show you how to Add Components to a Blueprint Actor.
 
 1.	Click **Ok** when the **Create Blueprint** window opens. 
 
-1.	In the **Content Browser**, under the **Game/Blueprint** folder, open the **SM_Bush** Blueprint. 
+1.	In the **Content Browser**, under the **StarterContent/Blueprint** folder, open the **SM_Bush** Blueprint. 
 
-1.	In the **Blueprint Editor** window, click the **Components** tab in the upper right of the window. 
+1.	In the **Blueprint Editor** window, click the **Viewport** tab in the upper left of the window. 
 
 	![](Components2.png)
 
@@ -80,11 +80,7 @@ We have created a **Blueprint Actor** that is comprised of three **Components**,
 
 Each of the Components inside a Blueprint can be accessed and modified from within the Blueprint itself or from an external Blueprint. For example, if we had a Car Blueprint and had an Engine Blueprint as a Component, we could modify the way the car performs by tweaking the values inside the Engine Blueprint. 
 
-Using the burning bush example from this page, inside the Blueprint on the **Graph** tab: 
-
-![](Components11.png)
-
-If you look in the **MyBlueprint** window, you will notice that each of the Components are provided in the Components section:
+Using the burning bush example from this page, if you look in the **MyBlueprint** window, you will notice that each of the Components are provided in the Components section:
 
 ![](Components12.png)
 
