#UE4doc #docs: Updated description of directory in step one.

diff --git a/Engine/Documentation/Source/Gameplay/HowTo/CharacterMovement/Blueprints/Setup_1/Setup_1.INT.udn b/Engine/Documentation/Source/Gameplay/HowTo/CharacterMovement/Blueprints/Setup_1/Setup_1.INT.udn
--- a/Engine/Documentation/Source/Gameplay/HowTo/CharacterMovement/Blueprints/Setup_1/Setup_1.INT.udn
+++ b/Engine/Documentation/Source/Gameplay/HowTo/CharacterMovement/Blueprints/Setup_1/Setup_1.INT.udn
@@ -31,7 +31,7 @@ For this project, we are using the **Blank** Project Template in order to start
 [/REGION]
 
 
-1.	With your project open, **Right-click** in the **Content Browser** and select **Animation** then **Animation Blueprint**. 
+1.	With your project open, **Right-click** in the **Content Browser** and select **New Asset** then **Animation** then **Animation Blueprint**. 
 
 	![](Characters_1b.png)
 
