TypoCorrection

diff --git a/Engine/Documentation/Source/Engine/AI/BehaviorTrees/NodeReference/BehaviorTreeNodeReference.INT.udn b/Engine/Documentation/Source/Engine/AI/BehaviorTrees/NodeReference/BehaviorTreeNodeReference.INT.udn
--- a/Engine/Documentation/Source/Engine/AI/BehaviorTrees/NodeReference/BehaviorTreeNodeReference.INT.udn
+++ b/Engine/Documentation/Source/Engine/AI/BehaviorTrees/NodeReference/BehaviorTreeNodeReference.INT.udn
@@ -1,7 +1,7 @@
+﻿Availability:Public
 Title: Behavior Trees Nodes Reference
 Crumbs:%ROOT%, Engine, Gameplay, Engine/AI/BehaviorTrees
 Description:Reference for the Behavior Tree nodes.
-Availability:Public
 
 [TOC(start:2 end:2)]
 
