TypoCorrection

diff --git a/Engine/Documentation/Source/Engine/AI/BehaviorTrees/NodeReference/Tasks/Tasks.INT.udn b/Engine/Documentation/Source/Engine/AI/BehaviorTrees/NodeReference/Tasks/Tasks.INT.udn
--- a/Engine/Documentation/Source/Engine/AI/BehaviorTrees/NodeReference/Tasks/Tasks.INT.udn
+++ b/Engine/Documentation/Source/Engine/AI/BehaviorTrees/NodeReference/Tasks/Tasks.INT.udn
@@ -1,7 +1,7 @@
+﻿Availability:Public
 Title: Behavior Tree Node Reference: Tasks
 Crumbs:%ROOT%, Engine, Gameplay, Engine/AI/BehaviorTrees
 Description:Reference for the Behavior Tree Task nodes.
-Availability:Public
 
 [TOC(start:2 end:2)]
 
