TypoCorrection

diff --git a/Engine/Documentation/Source/Engine/AI/BehaviorTrees/NodeReference/Decorators/Decorators.INT.udn b/Engine/Documentation/Source/Engine/AI/BehaviorTrees/NodeReference/Decorators/Decorators.INT.udn
--- a/Engine/Documentation/Source/Engine/AI/BehaviorTrees/NodeReference/Decorators/Decorators.INT.udn
+++ b/Engine/Documentation/Source/Engine/AI/BehaviorTrees/NodeReference/Decorators/Decorators.INT.udn
@@ -1,7 +1,7 @@
+﻿Availability:Public
 Title: Behavior Tree Node Reference: Decorators
 Crumbs:%ROOT%, Engine, Gameplay, Engine/AI/BehaviorTrees
 Description:Reference for the Behavior Tree Decorator nodes.
-Availability:Public
 
 [TOC(start:2 end:2)]
 
