#UE4doc #docs: sp: Global_Desaturation

diff --git a/Engine/Documentation/Source/Engine/Rendering/Materials/ParameterCollections/ParemeterCollections.INT.udn b/Engine/Documentation/Source/Engine/Rendering/Materials/ParameterCollections/ParemeterCollections.INT.udn
--- a/Engine/Documentation/Source/Engine/Rendering/Materials/ParameterCollections/ParemeterCollections.INT.udn
+++ b/Engine/Documentation/Source/Engine/Rendering/Materials/ParameterCollections/ParemeterCollections.INT.udn
@@ -70,8 +70,8 @@ In the following section we will look at how we can add a Material Collection No
 	![](MPC_Create_New_MPC.png)
 
 1. Once the Material Collection Parameter has been created open it up by **Double-Clicking** on it with the **Left Mouse Button**. 
-	Once the Material Collection Parameter is open add one Scalar Parameters to and name it **Global_Desatration**. 
-	Leave the default value of Global_Desatration at 0.0. When completed your Material Collection Parameter should look like this. 
+	Once the Material Collection Parameter is open add one Scalar Parameters to and name it **Global_Desaturation**. 
+	Leave the default value of Global_Desaturation at 0.0. When completed your Material Collection Parameter should look like this. 
 	![](MPC_Setup_MPC.png)
 
 1. Now that the Material Collection Parameter has been created we now need to create a simple Material to use it with. 
@@ -100,7 +100,7 @@ In the following section we will look at how we can add a Material Collection No
 
 1. To see the Material Collection Parameter in action, we must first apply the Material that is using the Material Collection Parameter to some objects in the world. 
 	Once the Material has been applied, open up the Material Collection Parameter by **Double-Clicking** on it with the **Left Mouse Button** inside of  the **Content Browser**. 
-	Once opened, we can then change the **Default Value** of Global_Desatration from 0 to 1 and see the objects with the Materials applied go from, in this case, blue to white like in the following example.  
+	Once opened, we can then change the **Default Value** of Global_Desaturation from 0 to 1 and see the objects with the Materials applied go from, in this case, blue to white like in the following example.  
 
 	[OBJECT:EmbeddedVideo]
 		[PARAMLITERAL:width]
