#UE4doc #docs: minor edits

diff --git a/Engine/Documentation/Source/Gameplay/HowTo/CharacterMovement/Blueprints/Input/Input.INT.udn b/Engine/Documentation/Source/Gameplay/HowTo/CharacterMovement/Blueprints/Input/Input.INT.udn
--- a/Engine/Documentation/Source/Gameplay/HowTo/CharacterMovement/Blueprints/Input/Input.INT.udn
+++ b/Engine/Documentation/Source/Gameplay/HowTo/CharacterMovement/Blueprints/Input/Input.INT.udn
@@ -57,7 +57,7 @@ Here we will set up the **Input** (Game Controls) to control the character and c
 	For more details on this process, see [](Gameplay\HowTo\SetUpInput\Blueprints) for more information. 
 	[/REGION]
 
-1.	Click the **Maps & Modes** option under the **Game** section of Project Settings. 
+1.	Click the **Maps & Modes** option under the **Project** section of Project Settings. 
 
 1.	In Maps & Modes under **Default Modes**, set the **Default GameMode** to use the **MyGame** Game Mode. 
 
diff --git a/Engine/Documentation/Source/Gameplay/HowTo/CharacterMovement/Blueprints/Setup_2/Setup_2.INT.udn b/Engine/Documentation/Source/Gameplay/HowTo/CharacterMovement/Blueprints/Setup_2/Setup_2.INT.udn
--- a/Engine/Documentation/Source/Gameplay/HowTo/CharacterMovement/Blueprints/Setup_2/Setup_2.INT.udn
+++ b/Engine/Documentation/Source/Gameplay/HowTo/CharacterMovement/Blueprints/Setup_2/Setup_2.INT.udn
@@ -47,7 +47,7 @@ In this section, we will finish setting up the Character by implementing the scr
 
 1.	**Right-click** in the graph and add the **InputAction Jog** Event. 
 
-1.	In the **MyBlueprint** window, click the **Show inherited variables** checkbox and drag the **CharacterMovement** into the graph. 
+1.	In the **MyBlueprint** window, click the **Show inherited variables** checkbox and **Control** + drag the **CharacterMovement** into the graph. 
 
 	![](Character_JogA.png)
 
