TypoCorrection

diff --git a/Engine/Documentation/Source/Gameplay/HowTo/CharacterMovement/Blueprints/Setup_1/Setup_1.INT.udn b/Engine/Documentation/Source/Gameplay/HowTo/CharacterMovement/Blueprints/Setup_1/Setup_1.INT.udn
--- a/Engine/Documentation/Source/Gameplay/HowTo/CharacterMovement/Blueprints/Setup_1/Setup_1.INT.udn
+++ b/Engine/Documentation/Source/Gameplay/HowTo/CharacterMovement/Blueprints/Setup_1/Setup_1.INT.udn
@@ -1,5 +1,4 @@
-﻿INTSourceChangelist:2474065
-Availability: Public
+﻿Availability: Public
 Title:1. Character Setup
 Crumbs: %ROOT%, Gameplay, Gameplay\HowTo\CharacterMovement, Gameplay\HowTo\CharacterMovement\Blueprints
 Description:Here we will set up our character and create some assets that we will need later in this guide. 
