TypoCorrection

diff --git a/Engine/Documentation/Source/Platforms/HTML5/BuildingEmscripten/HTML5BuildingEmscripten.INT.udn b/Engine/Documentation/Source/Platforms/HTML5/BuildingEmscripten/HTML5BuildingEmscripten.INT.udn
--- a/Engine/Documentation/Source/Platforms/HTML5/BuildingEmscripten/HTML5BuildingEmscripten.INT.udn
+++ b/Engine/Documentation/Source/Platforms/HTML5/BuildingEmscripten/HTML5BuildingEmscripten.INT.udn
@@ -76,4 +76,4 @@ All projects should build successfully and there should be 0 errors in the final
 
 As with the normal installation process for Emscripten, you will need to verify that all environment variables are set to correct values. Because of the manual nature of building Emscripten from source, you will likely need to modify most of these variables.
 
-See the [Environment Variables](Platforms/HTML5/GettingStarted#environmentvariables) section of the [](Platforms/HTML5/GettingStarted) page for more details.
+See the [](Platforms/HTML5/GettingStarted) page for more details.
