#loc UE4Doc  fix for minor unique errors

diff --git a/Engine/Documentation/Source/Resources/ContentExamples/Destructables/1_5/1_5.INT.udn b/Engine/Documentation/Source/Resources/ContentExamples/Destructables/1_5/1_5.INT.udn
--- a/Engine/Documentation/Source/Resources/ContentExamples/Destructables/1_5/1_5.INT.udn
+++ b/Engine/Documentation/Source/Resources/ContentExamples/Destructables/1_5/1_5.INT.udn
@@ -1,4 +1,4 @@
-Availability: Docs
+﻿Availability: Docs
 Title:1.5 - WorldSupport
 Crumbs:%ROOT%, Resources, Resources/ContentExamples/Destructables
 Description:An overview of the Destructibility example level, example 1.5: WorldSupport
@@ -31,6 +31,6 @@ Description:An overview of the Destructibility example level, example 1.5: World
 
 This ties the destructible mesh to a piece of Static world geometry. Every part that of the destructible mesh that touches a **Static Mesh** will act like it is fixed to the environment.
 
-[REGION:lightbox]
+<!--[REGION:lightbox]
 [![](1_3_solid.png)(w:300)](1_3_solid.png) [![](1_3_fractured.png)(w:300)](1_3_fractured.png) [![](1_3_blown.png)(w:300)](1_3_blown.png)
-[/REGION]
+[/REGION]-->
diff --git a/Engine/Documentation/Source/Resources/ContentExamples/Destructables/1_6/1_6.INT.udn b/Engine/Documentation/Source/Resources/ContentExamples/Destructables/1_6/1_6.INT.udn
--- a/Engine/Documentation/Source/Resources/ContentExamples/Destructables/1_6/1_6.INT.udn
+++ b/Engine/Documentation/Source/Resources/ContentExamples/Destructables/1_6/1_6.INT.udn
@@ -1,4 +1,4 @@
-Availability: Docs
+﻿Availability: Docs
 Title:1.6 - Destruction Particle Effects
 Crumbs:%ROOT%, Resources, Resources/ContentExamples/Destructables
 Description:An overview of the Destructibility example level, example 1.6: Destruction Particle Effects
@@ -35,6 +35,6 @@ Spawn a particle effect when mesh gets destroyed.
 Some brief rules about where the emitter is spawned would be nice. Content example is lacking a particle effect for this section.
 [/REGION]
 
-[REGION:lightbox]
+<!--[REGION:lightbox]
 [![](1_4_solid.png)(w:300)](1_4_solid.png) [![](1_4_fractured.png)(w:300)](1_4_fractured.png) [![](1_4_blown.png)(w:300)](1_4_blown.png)
-[/REGION]
+[/REGION]-->
diff --git a/Engine/Documentation/Source/Resources/ContentExamples/Destructables/1_7/1_7.INT.udn b/Engine/Documentation/Source/Resources/ContentExamples/Destructables/1_7/1_7.INT.udn
--- a/Engine/Documentation/Source/Resources/ContentExamples/Destructables/1_7/1_7.INT.udn
+++ b/Engine/Documentation/Source/Resources/ContentExamples/Destructables/1_7/1_7.INT.udn
@@ -1,4 +1,4 @@
-Availability: Docs
+﻿Availability: Docs
 Title:1.7 - Destruction Sounds
 Crumbs:%ROOT%, Resources, Resources/ContentExamples/Destructables
 Description:An overview of the Destructibility example level, example 1.7: Destruction Sounds
@@ -35,6 +35,6 @@ Plays sound on Destruction event.
 Needs to point out either blueprint use or how this actually works, would be nice 
 [/REGION]
 
-[REGION:lightbox]
+<!--[REGION:lightbox]
 [![](1_5_solid.png)(w:300)](1_5_solid.png)
-[/REGION]
+[/REGION]-->
diff --git a/Engine/Documentation/Source/Resources/ContentExamples/Destructables/1_8/1_8.INT.udn b/Engine/Documentation/Source/Resources/ContentExamples/Destructables/1_8/1_8.INT.udn
--- a/Engine/Documentation/Source/Resources/ContentExamples/Destructables/1_8/1_8.INT.udn
+++ b/Engine/Documentation/Source/Resources/ContentExamples/Destructables/1_8/1_8.INT.udn
@@ -1,4 +1,4 @@
-Availability: Docs
+﻿Availability: Docs
 Title:1.8 - Damage on Impact
 Crumbs:%ROOT%, Resources, Resources/ContentExamples/Destructables
 Description:An overview of the Destructibility example level, example 1.8: Damage on Impact
@@ -35,6 +35,6 @@ Plays sound on Destruction event.
 Needs to point out either blueprint use or how this actually works, would be nice 
 [/REGION]
 
-[REGION:lightbox]
+<!--[REGION:lightbox]
 [![](1_5_solid.png)(w:300)](1_5_solid.png)
-[/REGION]
+[/REGION]-->
