#UE4doc #docs: included hints to check ôpressedö button in æSETÆ nodes where appropriate.

diff --git a/Engine/Documentation/Source/Gameplay/HowTo/CharacterMovement/Blueprints/Setup_2/Setup_2.INT.udn b/Engine/Documentation/Source/Gameplay/HowTo/CharacterMovement/Blueprints/Setup_2/Setup_2.INT.udn
--- a/Engine/Documentation/Source/Gameplay/HowTo/CharacterMovement/Blueprints/Setup_2/Setup_2.INT.udn
+++ b/Engine/Documentation/Source/Gameplay/HowTo/CharacterMovement/Blueprints/Setup_2/Setup_2.INT.udn
@@ -53,7 +53,7 @@ In this section, we will finish setting up the Character by implementing the scr
 
 1.	Hold **Alt** and drag in the **Jog Pressed** variable, then hold **Control** and drag in the **Crouch Pressed** variable. 
 
-1.	Set the node network up as shown below. 
+1.	Set up the node network making sure to check "Jog Pressed" in the appropriate **Set** node as shown below:
 
 	![](Character_JogB.png)(w:720)
 
@@ -65,7 +65,7 @@ In this section, we will finish setting up the Character by implementing the scr
 
 1.	Hold **Alt** and drag in the **Crouch Pressed** variable, then hold **Control** and drag in the **Jog Pressed** and **Prone Pressed** variables. 
 
-1.	Set the node network up as shown below. 
+1.	Set up the node network making sure to check "Crouch Pressed" in the appropriate **Set** node as shown below: 
 
 	[REGION:lightbox]
 	[![](Character_Crouch.png)(w:720)](Character_Crouch.png)
@@ -83,7 +83,7 @@ In this section, we will finish setting up the Character by implementing the scr
 
 1.	Hold **Control** and drag in the **Crouch Pressed** variable. 
 
-1.	Set the node network up as shown below. 
+1.	Set up the node network making sure to check "Prone Pressed" in the appropriate **Set** node as shown below: 
 
 	[REGION:lightbox]
 	[![](Character_Prone.png)(w:720)](Character_Prone.png)
@@ -99,7 +99,7 @@ In this section, we will finish setting up the Character by implementing the scr
 
 1.	Hold **Alt** and drag in the **Jump Pressed** variable, then **Control** drag in the **DisableMovement** and **CrouchPressed** variables. 
 
-1.	Set the node network up as shown below. 
+1.	Set up the node network making sure to check "Jump Pressed" in the appropriate **Set** nodes as shown below: 
 
 	[REGION:lightbox]
 	[![](Character_Jump.png)(w:720)](Character_Jump.png)
