#UE4doc #docs: Adjusted text to match remaining image updates

diff --git a/Engine/Documentation/Source/Gameplay/HowTo/FindingActors/Blueprints/FindingActors_BP.INT.udn b/Engine/Documentation/Source/Gameplay/HowTo/FindingActors/Blueprints/FindingActors_BP.INT.udn
--- a/Engine/Documentation/Source/Gameplay/HowTo/FindingActors/Blueprints/FindingActors_BP.INT.udn
+++ b/Engine/Documentation/Source/Gameplay/HowTo/FindingActors/Blueprints/FindingActors_BP.INT.udn
@@ -46,17 +46,7 @@ The steps below will show you how you can use the **Get All Actors of Class** no
 
 	The **Out Actors** is an **Array** of all the Actors (of the Class specified) and "For Each" of them we can do something.
 
-1.	On the **ForEachLoop**, drag off the **Array Element** and add a **Cast To Blueprint_Effect_Fire** node. 
-
-	![](Actors6.png)
-
-	In order to access/modify the **Blueprint_Effect_Fire** asset, we need to associate the Actor Reference with the Blueprint. By taking the Actor Reference and **Casting** to the Object we want to associate it with, we are saying "Yes Actor, you are indeed of this type so let me access and modify you."
-
-	[REGION:note]
-	For more on **Casting** and **Blueprint Communication**, refer to the [](Engine/Blueprints/UserGuide/BlueprintCommsUsage) documentation. 
-	[/REGION]
-
-1.	Off the **As Blueprint Effect Fire C** pin of the Cast node, get **P_Fire** and **Fire Audio**. 
+1.	Drag off the **Array Element** of the **ForEachLoop** node and get **P_Fire** and **Fire Audio**. 
 
 1.	**Right-click** in the graph and add a **Deactivate** node, then connect the fire and audio pins to the **Target**. 
 
@@ -91,7 +81,7 @@ One method you can take is by using **Tags** on the Actor, for example:
 
 	![](Actors9.png)
 
-1.	Inside the **MyCharacter** Blueprint, off the **As Blueprint Effect Fire C** pin, add the **Get Tags** node. 
+1.	Inside the **MyCharacter** Blueprint, drag off the **Array Element** of the **ForEachLoop** and add the **Get Tags** node. 
 
 	![](Actors10.png)
 
@@ -107,11 +97,11 @@ One method you can take is by using **Tags** on the Actor, for example:
 
 1.	Connect the out pin from the **Get** node to the **Tag** pin on the **Actor Has Tag** node. 
 
-1.	Connect the **As Blueprint Effect Fire C** pin to the **Target** pin on the **Actor Has Tag** node. 
+1.	Connect the **Array Element** of the **ForEachLoop** pin to the **Target** pin on the **Actor Has Tag** node. 
 
 	![](Actors12.png)(w:800)
 
-1.	Drag off the **As Blueprint Effect Fire C** pin and get **P_Fire** and **Fire Audio** and connect them to a **Deactivate** node. 
+1.	Drag off the **Array Element** of the **ForEachLoop** pin to get **P_Fire** and **Fire Audio** and connect them to a **Deactivate** node. 
 
 	[REGION:lightbox]
 	[![](Actors13.png)(w:800)](Actors13.png)
@@ -146,4 +136,3 @@ Here we are getting all Actors in our level of the Class **TriggerBox** and disa
 
 
 
-
