#UE4doc #docs: "Text are" to "Text is"

diff --git a/Engine/Documentation/Source/Gameplay/Localization/Overview/Overview.INT.udn b/Engine/Documentation/Source/Gameplay/Localization/Overview/Overview.INT.udn
--- a/Engine/Documentation/Source/Gameplay/Localization/Overview/Overview.INT.udn
+++ b/Engine/Documentation/Source/Gameplay/Localization/Overview/Overview.INT.udn
@@ -21,7 +21,7 @@ Localization of a project involves a few key concepts:
 
 
 ## Text
-Text is the basic unit of localization. Text are defined by a namespace, a key, a source string, and a display string. Together, the namespace and key form a unique identity by which a text can be referred by. 
+Text is the basic unit of localization. Text is defined by a namespace, a key, a source string, and a display string. Together, the namespace and key form a unique identity by which a text can be referred by. 
 The namespace allows text to be translated differently in the case of homographs (identical strings with different meanings) by offering a different identity. The key provides specific context regarding the text. 
 The source string is the string in its native form, without having been translated. The display string is the string which will be shown, typically a translated form of the source string.
 
