Split off Blueprint Communication How To page into individual HowTo's in the Blueprint HowTo's section.
Updated related pages and content.

diff --git a/Engine/Documentation/Source/Engine/Blueprints/BP_HowTo/BlueprintHowTo.INT.udn b/Engine/Documentation/Source/Engine/Blueprints/BP_HowTo/BlueprintHowTo.INT.udn
--- a/Engine/Documentation/Source/Engine/Blueprints/BP_HowTo/BlueprintHowTo.INT.udn
+++ b/Engine/Documentation/Source/Engine/Blueprints/BP_HowTo/BlueprintHowTo.INT.udn
@@ -27,6 +27,7 @@ Refer to the table below for more information and click each link to view its co
 | **[](Engine/Blueprints/UserGuide/Variables)** | The entire user guide for Variables. |
 | **[](Engine\Blueprints\BP_HowTo\DirectBlueprintComs)** | %Engine\Blueprints\BP_HowTo\DirectBlueprintComs:description% |
 | **[](Engine\Blueprints\BP_HowTo\Casting)** | %Engine\Blueprints\BP_HowTo\Casting:description% |
+| **[](Engine\Blueprints\BP_HowTo\Interfaces)** | %Engine\Blueprints\BP_HowTo\Interfaces:description% |
 | **[](Engine\Blueprints\BP_HowTo\EventDispatcher)** | %Engine\Blueprints\BP_HowTo\EventDispatcher:description% |
 | **[](Engine/Blueprints/BP_HowTo/MakingMacros)** | %Engine/Blueprints/BP_HowTo/MakingMacros:description% |
 | **[](Engine/Blueprints/BP_HowTo/Debugging)** | %Engine/Blueprints/BP_HowTo/Debugging:description% |
diff --git a/Engine/Documentation/Source/Engine/Blueprints/BP_HowTo/DirectBlueprintComs/DirectBlueprintComs.INT.udn b/Engine/Documentation/Source/Engine/Blueprints/BP_HowTo/DirectBlueprintComs/DirectBlueprintComs.INT.udn
--- a/Engine/Documentation/Source/Engine/Blueprints/BP_HowTo/DirectBlueprintComs/DirectBlueprintComs.INT.udn
+++ b/Engine/Documentation/Source/Engine/Blueprints/BP_HowTo/DirectBlueprintComs/DirectBlueprintComs.INT.udn
@@ -6,7 +6,7 @@ Related: Engine/Blueprints/UserGuide/BlueprintComms
 Related: Engine/Blueprints/UserGuide/BlueprintCommsUsage
 Related: Engine/Blueprints/UserGuide
 
-This page covers how to directly communicate between two Blueprints and pass information between them. 
+This page covers using [Direct Blueprint Communication](Engine/Blueprints/UserGuide/BlueprintComms) to communicate between two Blueprints and pass information between them. 
 
 For this example we will allow the player to pass through a box and communicate with and activate a particle effect. 
 
diff --git a/Engine/Documentation/Source/Engine/Blueprints/BP_HowTo/EventDispatcher/EventDispatcher.INT.udn b/Engine/Documentation/Source/Engine/Blueprints/BP_HowTo/EventDispatcher/EventDispatcher.INT.udn
--- a/Engine/Documentation/Source/Engine/Blueprints/BP_HowTo/EventDispatcher/EventDispatcher.INT.udn
+++ b/Engine/Documentation/Source/Engine/Blueprints/BP_HowTo/EventDispatcher/EventDispatcher.INT.udn
@@ -6,7 +6,7 @@ Related: Engine/Blueprints/UserGuide/EventDispatcher
 Related: Engine/Blueprints/UserGuide/BlueprintCommsUsage
 Related: Engine/Blueprints/UserGuide
 
-When an **Event Dispatcher** is called, any Blueprints that implement and have Events bound to the Event Dispatcher will also execute when the Event Dispatcher is called. In other words, whenever the Event Dispatcher is called, any Events that are bound to it are also called allowing you to execute multiple Events at once from a single source. 
+When an [Event Dispatcher](Engine/Blueprints/UserGuide/EventDispatcher) is called, any Blueprints that implement and have Events bound to the Event Dispatcher will also execute when the Event Dispatcher is called. In other words, whenever the Event Dispatcher is called, any Events that are bound to it are also called allowing you to execute multiple Events at once from a single source. 
 
 This page will show you how to set up an Event Dispatcher and Bind an Event to the Event Dispatcher. 
 
diff --git a/Engine/Documentation/Source/Engine/Blueprints/BP_HowTo/Interfaces/Interfaces.INT.udn b/Engine/Documentation/Source/Engine/Blueprints/BP_HowTo/Interfaces/Interfaces.INT.udn
new file mode 100644
--- /dev/null
+++ b/Engine/Documentation/Source/Engine/Blueprints/BP_HowTo/Interfaces/Interfaces.INT.udn
@@ -0,0 +1,132 @@
+Availability: Public
+Title:Using Interfaces
+Crumbs: %ROOT%, Engine, Engine\Blueprints, Engine\Blueprints\BP_HowTo
+Description:Sample of using Blueprint Interfaces to communicate with multiple Blueprints at once. 
+Related: Engine/Blueprints/UserGuide/Types/Interface/UsingInterfaces
+Related: Engine/Blueprints/UserGuide/BlueprintCommsUsage
+Related: Engine/Blueprints/UserGuide
+
+In this How To, we will set up [Blueprint Interface](Engine/Blueprints/UserGuide/Types/Interface/UsingInterfaces) which will execute different scripts in multiple Blueprints when called. 
+
+[REGION:note]
+For this guide we are using the **Blueprint First Person** template with **Starter Content** enabled. 
+[/REGION]
+
+1.  In the **Content/StarterContent/Props** folder, **Right-click** on the **SM_GlassWindow** and choose **Create Destructible Mesh**. 
+
+	![](Interface1.png)
+
+	We are going to turn this Static Mesh into a Destructible Mesh which we will break when shot. 
+
+1. Open the Destructible Mesh, then click **Fracture Mesh** and **Save** then close the Destructible Mesh Editor window. 
+
+	![](Interface2.png)
+
+1. **Right-click** on the **SM_GlassWindow_DM** asset, and under **Asset Actions**, select **Create Blueprint Using This...**. 
+
+	![](Interface3.png)
+
+	This will prompt you to enter a name and save location for the Blueprint (we used **SM_GlassWindow_DM**). 
+
+1. **Right-click** inside the **Content Browser**, then under **Blueprints** select **Blueprint Interface**. 
+
+	![](Interface4.png)
+
+	This will again ask you to enter a name and save location for the Blueprint Interface (we used **DM_Interface**). 
+
+1. Inside the Blueprint Interface, name the existing Function **IsDestructible?** and add the **Inputs** shown below by clicking **New**. 
+
+	![](Interface5.png)
+
+	Above we have created three **Inputs** which we can use to pass variables between Blueprints that implement this Blueprint Interface. We have create two **Float** variables, one called **Damage** which will be used to pass an amount of damage to cause to the object we hit and the other **Strength** the amount of Impulse Strength to apply. Lastly, we added a **Vector** called **HitLoc** to hold the location in which the object was hit so that we may pass that information through as well. 
+
+1. **Compile** and **Save** then close the Blueprint Interface.
+
+1.	In the **Content Browser** open the **SM_GlassWindow_DM** Blueprint, then click the **Class Settings** option from the tool bar. 
+
+	![](Interface6.png)
+
+1.	In the **Details** panel, click the **Add** button under **Interfaces** and add the **DM_Interface**. 
+
+	![](Interface7.png)
+
+	This will implement the Blueprint Interface in this Blueprint, allowing us to access the functions within it. 
+
+1.	**Right-click** in the graph and search for a **?** and select the **Event Is Destructible?** Event. 
+
+	![](Interface8.png)
+
+1.	Hold **Control** and drag in the **Destructible** component from the **MyBlueprint** window and create the script below. 
+
+	![](Interface9.png)
+
+	Here, off the **Destructible** component we used **Apply Damage** and connected the damage, strength, location and direction requirements needed to apply damage to our glass window. The actual damage, strength and hit location information will be passed through the Blueprint Interface to this Blueprint from another Blueprint which we will do in a moment.
+
+1.	**Compile** and **Save** then close the Blueprint. 
+
+1.	In the **Content Browser** under **Content/StarterContent/Blueprints**, open the **Blueprint_CeilingLight** Blueprint. 
+
+	![](Interface10.png)
+
+	Using the Blueprint Interface, we will do something different when we shoot this light. 
+
+1.	Add the **DM_Interface** through **Class Settings** then drag in the **PointLight1** component and create the script below. 
+
+	![](Interface11.png)
+
+	Here, when **Event Is Destructible?** is called, we will first set the light to a green color. The second time it is called, we will set the intensity to 0 (turning off the light). Notice that we did not use the Damage, Strength or HitLoc information contained in the interface as it was not relevant to this Blueprint. By using a Blueprint Interface we can make a general function that is called in all Blueprints that implement that interface and function, and then determine what script is executed on each individual Blueprint. 
+
+1.	**Compile** and **Save** then close the Blueprint. 
+
+1.	Inside the **Content Browser** under **Content/FirstPersonBP/Blueprints**, open the **FirstPersonProjectile** Blueprint. 
+
+	![](Interface12.png)
+
+	This Blueprint includes the script for the projectile that is spawned in-game when the player fires their weapon. 
+
+1.	Add the **DM_Interface** through **Class Settings** then **Right-click** in the graph and add the **IsDestructible?(Message)** node. 
+
+	![](Interface13.png)
+
+	This will execute the function IsDestructible? and all Blueprints that implement the interface and function will also execute. 
+
+	[REGION:note]
+	If you do not see the node, un-check the **Context Sensitive** option in the upper-right corner of the context menu. 
+	[/REGION]
+
+1.	Off the **False** pin of the **Branch** connect the **IsDestructible?** node and connect the **Other** and **HitLocation** pins. 
+
+	[REGION:lightbox]
+	[![](Interface14.png)(w:820)](Interface14.png)
+	[/REGION]
+
+	_Click for full image._
+
+	Above, we have added to the existing script that checks if the object that was hit is simulating physics. If it is not simulating physics, we call our **IsDestructible?** function in our Blueprint Interface and apply 1000 damage and 100 strength to the target (which is gathered from the **Event Hit** node's **Other** pin where we also get the **Hit Location** and assign it to our **HitLoc** variable to pass through). 
+
+1.	**Compile** and **Save**, then close the Blueprint. 
+
+1.	Drag the **SM_GlassWindow_DM** and **Blueprint_CeilingLight** Blueprints into the level then click **Play** to play in the editor. 
+
+	Shooting the glass will cause it to fracture and break apart. 
+
+	![](Example1.png)(w:720)
+
+	Shooting the light once, will cause it to change to the color you specified in step 13. 
+
+	![](Example2.png)(w:720)
+
+	Shooting the light a second time will cause it to turn off (also specified in step 13). 
+	
+	![](Example3.png)(w:720)
+
+This is one example of using a _Blueprint Interface_ to pass variables between Blueprints and call functions which execute different scripts in the Blueprints that implement the interface. For more information on working with Blueprint Interfaces and other forms of Blueprint communication, refer to the **Related Topics** section below. 
+
+
+
+
+
+
+
+
+	
diff --git a/Engine/Documentation/Source/Engine/Blueprints/UserGuide/BlueprintComms/BlueprintComms.INT.udn b/Engine/Documentation/Source/Engine/Blueprints/UserGuide/BlueprintComms/BlueprintComms.INT.udn
--- a/Engine/Documentation/Source/Engine/Blueprints/UserGuide/BlueprintComms/BlueprintComms.INT.udn
+++ b/Engine/Documentation/Source/Engine/Blueprints/UserGuide/BlueprintComms/BlueprintComms.INT.udn
@@ -3,7 +3,6 @@ Title:Direct Blueprint Communications
 Crumbs: %ROOT%, Engine, Engine/Blueprints, Engine/Blueprints/UserGuide
 Description:How to use one Blueprint's functions, events, and variables from another Blueprint.
 related:Engine/Blueprints/UserGuide/BlueprintCommsUsage
-related:Engine/Blueprints/UserGuide/BlueprintCommsUsage/BPComHowTo
 related:Engine/Blueprints/UserGuide/Types/Interface/UsingInterfaces
 related:Engine/Blueprints/UserGuide/EventDispatcher
 
diff --git a/Engine/Documentation/Source/Engine/Blueprints/UserGuide/BlueprintCommsUsage/BPComHowTo/BPComHowTo.INT.udn b/Engine/Documentation/Source/Engine/Blueprints/UserGuide/BlueprintCommsUsage/BPComHowTo/BPComHowTo.INT.udn
--- a/Engine/Documentation/Source/Engine/Blueprints/UserGuide/BlueprintCommsUsage/BPComHowTo/BPComHowTo.INT.udn
+++ b/Engine/Documentation/Source/Engine/Blueprints/UserGuide/BlueprintCommsUsage/BPComHowTo/BPComHowTo.INT.udn
@@ -1,4 +1,4 @@
-Availability:Public
+Availability:Docs
 Title:How To Use Blueprint Communications
 Crumbs: %ROOT%, Engine, Engine/Blueprints, Engine/Blueprints/UserGuide, Engine/Blueprints/UserGuide/BlueprintCommsUsage
 Description:Examples of how to setup and use different methods of Blueprint Communication. 
diff --git a/Engine/Documentation/Source/Engine/Blueprints/UserGuide/BlueprintCommsUsage/BlueprintCommsUsage.INT.udn b/Engine/Documentation/Source/Engine/Blueprints/UserGuide/BlueprintCommsUsage/BlueprintCommsUsage.INT.udn
--- a/Engine/Documentation/Source/Engine/Blueprints/UserGuide/BlueprintCommsUsage/BlueprintCommsUsage.INT.udn
+++ b/Engine/Documentation/Source/Engine/Blueprints/UserGuide/BlueprintCommsUsage/BlueprintCommsUsage.INT.udn
@@ -2,6 +2,10 @@ Availability:Public
 Title:Blueprint Communication Usage
 Crumbs: %ROOT%, Engine, Engine/Blueprints, Engine/Blueprints/UserGuide
 Description:Overview of when to use different methods of Blueprint Communications. 
+Related: Engine\Blueprints\UserGuide\BlueprintComms
+Related: Engine/Blueprints/UserGuide/EventDispatcher
+Related: Engine/Blueprints/UserGuide/CastNodes
+Related: Engine/Blueprints/UserGuide/Types/Interface/UsingInterfaces
 
 [VAR:Topic]
 [OBJECT:Topic]
@@ -72,7 +76,7 @@ Below are some examples when you might use **Blueprint Comms**:
 *	**You have an Actor in your level that you want to communicate with another Actor in your level.**
 	*	Interacting with a switch in your level opens a specific door or turns on a specific light (each of which are separate Blueprints). 
 
-Refer to [Setting Up Direct Blueprint Communications](Engine\Blueprints\UserGuide\BlueprintCommsUsage\BPComHowTo#directblueprintcommunications) for a walk through of how to set one up. 
+Refer to [Setting Up Direct Blueprint Communications](Engine\Blueprints\BP_HowTo\DirectBlueprintComs) for a walk through of how to set one up. 
 
 
 	
@@ -105,7 +109,7 @@ The examples below would be a case where an **Event Dispatcher** could be used:
 	*	You spawn a boss and an Event is fired when the boss is killed which spawns a reward in the world.
 	*	You spawn an item in your level (weapon, health, boost, etc.) and through an Event, notify the item and character when it is picked up.
 
-Refer to [Setting Up Event Dispatchers](Engine\Blueprints\UserGuide\BlueprintCommsUsage\BPComHowTo#eventdispatchers) for a walk through of how to set one up. 
+Refer to [Setting Up Event Dispatchers](Engine\Blueprints\BP_HowTo\EventDispatcher) for a walk through of how to set one up. 
 
 
 ## Blueprint Interfaces
@@ -134,7 +138,7 @@ Each of the examples below are scenarios where a **Blueprint Interface** might b
 	*	Player has a "use" button, pressing the use button opens doors, turns off lights, picks up items, etc. 
 	*	Enemy has a special ability that changes based on player health. 
 
-Refer to [Setting Up Blueprint Interfaces](Engine\Blueprints\UserGuide\BlueprintCommsUsage\BPComHowTo#blueprintinterfaces) for a walkthrough of how to set one up. 
+Refer to [Setting Up Blueprint Interfaces](Engine\Blueprints\BP_HowTo\Interfaces) for a walkthrough of how to set one up. 
 
 
 ## Casting For Communications
@@ -167,7 +171,7 @@ Here are a few examples when you could use **Blueprint Casting**:
 	*	You have several lights in your level and you want them all to turn on or off when an event occurs. 
 [/EXCERPT:Castusage]
 
-Refer to [Setting Up Blueprint Casts](Engine\Blueprints\UserGuide\BlueprintCommsUsage\BPComHowTo#blueprintcasting) for a walk through of how to set one up. 
+Refer to [Setting Up Blueprint Casts](Engine\Blueprints\BP_HowTo\Casting) for a walk through of how to set one up. 
 
 
 
diff --git a/Engine/Documentation/Source/Engine/Blueprints/UserGuide/Types/Interface/UsingInterfaces/UsingInterfaces.INT.udn b/Engine/Documentation/Source/Engine/Blueprints/UserGuide/Types/Interface/UsingInterfaces/UsingInterfaces.INT.udn
--- a/Engine/Documentation/Source/Engine/Blueprints/UserGuide/Types/Interface/UsingInterfaces/UsingInterfaces.INT.udn
+++ b/Engine/Documentation/Source/Engine/Blueprints/UserGuide/Types/Interface/UsingInterfaces/UsingInterfaces.INT.udn
@@ -45,11 +45,13 @@ Description:Blueprint Interfaces allow Blueprints to call functions in different
 
 ## Using Interfaces in Blueprints
 
+[EXCERPT:intro]
 The use of _Blueprint Interfaces_ allows for a common method of interacting with multiple different types of objects that all share 
 some specific functionality. This means you can have completely different types of objects, such as a car and a tree, that 
 share one specific thing, like they can both be shot by weapon fire and take damage. By creating a _Blueprint Interface_ that 
 contains an `OnTakeWeaponFire` function, and having both the car and the tree implement that _Blueprint Interface_, you can treat the 
 car and the tree as the same type and simply call the `OnTakeWeaponFire` function when either of them is shot.
+[/EXCERPT:intro]
 
 
 ### Adding Interfaces
