minor typo fix.

diff --git a/Engine/Documentation/Source/Engine/Blueprints/BP_HowTo/Interfaces/Interfaces.INT.udn b/Engine/Documentation/Source/Engine/Blueprints/BP_HowTo/Interfaces/Interfaces.INT.udn
--- a/Engine/Documentation/Source/Engine/Blueprints/BP_HowTo/Interfaces/Interfaces.INT.udn
+++ b/Engine/Documentation/Source/Engine/Blueprints/BP_HowTo/Interfaces/Interfaces.INT.udn
@@ -6,7 +6,7 @@ Related: Engine/Blueprints/UserGuide/Types/Interface/UsingInterfaces
 Related: Engine/Blueprints/UserGuide/BlueprintCommsUsage
 Related: Engine/Blueprints/UserGuide
 
-In this How To, we will set up [Blueprint Interface](Engine/Blueprints/UserGuide/Types/Interface/UsingInterfaces) which will execute different scripts in multiple Blueprints when called. 
+In this How To, we will set up a [Blueprint Interface](Engine/Blueprints/UserGuide/Types/Interface/UsingInterfaces) which will execute different scripts in multiple Blueprints when called. 
 
 [REGION:note]
 For this guide we are using the **Blueprint First Person** template with **Starter Content** enabled. 
