#UE4doc #docs: "Text are" to "Text is"

diff --git a/Engine/Documentation/Source/Gameplay/Localization/Overview/Overview.INT.udn b/Engine/Documentation/Source/Gameplay/Localization/Overview/Overview.INT.udn
--- a/Engine/Documentation/Source/Gameplay/Localization/Overview/Overview.INT.udn
+++ b/Engine/Documentation/Source/Gameplay/Localization/Overview/Overview.INT.udn
@@ -56,7 +56,7 @@ Archives store source strings and their translations, mapped by namespace, in hu
 ## Localization Resources (LocRes)
 LocRes store translated text in a binary format to be loaded by the system. LocRes are generated by a commandlet that compiles a specified manifest and some number of specified archives.
 
-The system loads LocRes files based on the project settings and the current culture. Localized text from the LocRes of the current culture are used in addition to localized text from the LocRes of
+The system loads LocRes files based on the project settings and the current culture. Localized text from the LocRes of the current culture is used in addition to localized text from the LocRes of
  all of the current culture's parent cultures. This allows general translations to be made for a language in tandem with having more specific translations made for a region. As a basic example, 
  for a single target containing the text "Color" and supporting the cultures "en" (English) and "en-UK" (English for the United Kingdom), the "en" LocRes may have "Color" localized as "Color" while 
  the "en-UK" LocRes may have it localized as "Colour". If the user switches to "en-CA" (English for Canada), but the LocRes for "en-CA" lacks a localization for "Color", the LocRes for "en" will provide "Color".
\ No newline at end of file
diff --git a/Engine/Documentation/Source/Gameplay/Networking/Actors/OwningConnections/OwningConnections.INT.udn b/Engine/Documentation/Source/Gameplay/Networking/Actors/OwningConnections/OwningConnections.INT.udn
--- a/Engine/Documentation/Source/Gameplay/Networking/Actors/OwningConnections/OwningConnections.INT.udn
+++ b/Engine/Documentation/Source/Gameplay/Networking/Actors/OwningConnections/OwningConnections.INT.udn
@@ -10,7 +10,7 @@ You will hear a lot about the owning connection of an actor. Ultimately, each co
 
 One example of this is when Pawn actors are possessed by a PlayerController. Their owner will be the PlayerController that they are possessed by. During this time, they are owned by the connection of the PlayerController. The Pawn is only owned by this connection during the time it is also owned/possessed by the PlayerController. So as soon as the PlayerController no longer possesses the Pawn, the Pawn is no longer owned by the connection.
 
-Another example are inventory items that are owned by a Pawn. These inventory items are owned by the same connection (if any) that might own that Pawn.
+Another example is inventory items that are owned by a Pawn. These inventory items are owned by the same connection (if any) that might own that Pawn.
 
 Components are a little special in how they determine their owning connection. In this case, we first determine the components owner by walking the components outer chain until we find the owning actor, and then we continue as above by determining owning connection of this actor.
 
