#UE4doc #docs: "are" to "is" for singular subject,  re-worded sentence for clarity

diff --git a/Engine/Documentation/Source/Gameplay/Networking/Actors/OwningConnections/OwningConnections.INT.udn b/Engine/Documentation/Source/Gameplay/Networking/Actors/OwningConnections/OwningConnections.INT.udn
--- a/Engine/Documentation/Source/Gameplay/Networking/Actors/OwningConnections/OwningConnections.INT.udn
+++ b/Engine/Documentation/Source/Gameplay/Networking/Actors/OwningConnections/OwningConnections.INT.udn
@@ -21,7 +21,7 @@ Connection ownership is important for a few things:
 
 Connection ownership is important for things like RPCs, because when you call an RPC function on an actor, unless the RPC is marked as multicast, it needs to know which client to execute that RPC on. It determines the connection to send the RPC to by finding the owning connection.
 
-Connection ownership is used during actor replication, and determining which connections get updates for each actor. For actors that have bOnlyRelevantToOwner set to true, only the connection that owns that actor will receive property updates for that actor. By default, all PlayerControllers have this flag set, and this is why each client only receives updates for PlayerControllers that he owns. This done for various reasons, the main ones being player cheating, and efficiency.
+Connection ownership is used during actor replication, and determining which connections get updates for each actor. For actors that have bOnlyRelevantToOwner set to true, only the connection that owns that actor will receive property updates for that actor. By default, all PlayerControllers have this flag set, and this is why each client only receives updates for PlayerControllers that he owns. This is done for various reasons, the main ones being to prevent players from cheating and efficiency.
 
 Connection ownership is important during [property replication involving conditions](Gameplay/Networking/Actors/Properties/Conditions) that use the owner. For example. When `COND_OnlyOwner` is used, only the owners of that actor will receive these property updates.
 
