#UE4doc #docs: Updated images and file names to reflect UE4.8

diff --git a/Engine/Documentation/Source/Gameplay/HowTo/Networking/ReplicateFunction/Blueprints/ReplicateFunction_BP.INT.udn b/Engine/Documentation/Source/Gameplay/HowTo/Networking/ReplicateFunction/Blueprints/ReplicateFunction_BP.INT.udn
--- a/Engine/Documentation/Source/Gameplay/HowTo/Networking/ReplicateFunction/Blueprints/ReplicateFunction_BP.INT.udn
+++ b/Engine/Documentation/Source/Gameplay/HowTo/Networking/ReplicateFunction/Blueprints/ReplicateFunction_BP.INT.udn
@@ -34,7 +34,7 @@ There are 3 primary types of replicated functions: **Multicast**, **Run on Serve
 	You can change the size of the windows that are generated for multiplayer inside the **Advanced Settings** option. 
 	[/REGION]
 
-1.	Next find the **Game/Blueprints** folder and open the **MyCharacter** Blueprint. 
+1.	Next find the **Game/Blueprints** folder and open the **ThirdPersonCharacter** Blueprint. 
 
 1.	**Right-click** in the graph, and add a **Space Bar** Key Event and connect it to a **Spawn Actor from Class** node. 
 
@@ -52,7 +52,7 @@ There are 3 primary types of replicated functions: **Multicast**, **Run on Serve
 
 	You should see that only the server window spawns the fire effect and the other players do not see the effect spawned. When attempting to jump with the other players, they too only spawn the fire effect locally and each of the other players are not aware that an effect has been spawned. 
 
-1.	Return to the **MyCharacter** Blueprint and disconnect the **Space Bar** Key Event by **Right-clicking** on the pin. 
+1.	Return to the **ThirdPersonCharacter** Blueprint and disconnect the **Space Bar** Key Event by **Right-clicking** on the pin. 
 
 1.	**Right-click** and search for then add a **Custom Event** and call the event **MulticastSpawn** or whatever name you wish. 
 
@@ -87,7 +87,7 @@ There are 3 primary types of replicated functions: **Multicast**, **Run on Serve
 
  The steps below show how to set up a **Run on Server** replicated function. 
  
-1.	Inside the **Game/Blueprints** folder, open the **MyCharacter** Blueprint. 
+1.	Inside the **Game/Blueprints** folder, open the **ThirdPersonCharacter** Blueprint. 
 
 1.	Select the custom event tied to the **Spawn Actor from Class** node and change the **Replicates** drop down option to **Run on Server**. 
 
@@ -126,7 +126,7 @@ What we will do in this example is create a variable that only gets updated on a
 
  The steps below show how to set up a **Run on owning Client** replicated function. 
  
-1.	Inside the **Game/Blueprints** folder, open the **MyCharacter** Blueprint. 
+1.	Inside the **Game/Blueprints** folder, open the **ThirdPersonCharacter** Blueprint. 
 
 1.	On the **Graph** tab in the **MyBlueprint** window, create a new variable of the **String** type and give it a name (**Inventory** for example). 
 
@@ -142,7 +142,7 @@ What we will do in this example is create a variable that only gets updated on a
 
 	![](HowTo13.png)
 
-1.	**Compile** and **Save** then close the **MyCharacter** Blueprint. 
+1.	**Compile** and **Save** then close the **ThirdPersonCharacter** Blueprint. 
 
 1.	Next, from the **Modes** window on the **Basic** tab, drag in a **Box Trigger** and place it in the middle of the player starts. 
 
@@ -184,7 +184,7 @@ What we will do in this example is create a variable that only gets updated on a
 
 	Here we are stating that when overlapping the trigger, if the overlap occurs on the server, run the **Add Item** event on the server and replicate it to the owning client (which is the character that overlaps the trigger). The "it" is the script that is called when **Add Item** is fired off and is only executed on the server but replicated to the client. When the character leaves the trigger box, again determined by the server, run the **Remove Item** event on the server and replicate it to the owning client. 
 
-1.	Off the **Add Item** event, add a **Print String** (text set to **Item Added**) then drag off **Character** and **Cast To MyCharacter**. 
+1.	Off the **Add Item** event, add a **Print String** (text set to **Item Added**) then drag off **Character** and **Cast To ThirdPersonCharacter**. 
 
 	![](HowTo20.png)
 
@@ -215,6 +215,7 @@ What we will do in this example is create a variable that only gets updated on a
 
 
 
+
 	
 
 
diff --git a/Engine/Documentation/Source/Gameplay/HowTo/PossessPawns/Blueprints/GHT_4B.INT.udn b/Engine/Documentation/Source/Gameplay/HowTo/PossessPawns/Blueprints/GHT_4B.INT.udn
--- a/Engine/Documentation/Source/Gameplay/HowTo/PossessPawns/Blueprints/GHT_4B.INT.udn
+++ b/Engine/Documentation/Source/Gameplay/HowTo/PossessPawns/Blueprints/GHT_4B.INT.udn
@@ -19,7 +19,7 @@ _The steps below are based off using the **Blueprint Side Scroller Template** pr
 
 	![](GHT4B_1.png)
 
-	Here we added two **MyCharacter** Blueprints to the level and a third character is spawned via the **PlayerStart**. 
+	Here we added two **SideScrollerCharacter** Blueprints to the level and a third character is spawned via the **PlayerStart**. 
 
 1.	Select each of your characters in the level (the two we added to the level). 
 
@@ -37,13 +37,13 @@ _The steps below are based off using the **Blueprint Side Scroller Template** pr
 
 1.	**Right-click** in the graph and add an **Event Begin Play** node, **Right-click** again and add a **Get Player Character** node.  
 
-1.	Drag off the **Return Value** of the **Get Player Character** and add the **Cast To MyCharacter** node. 
+1.	Drag off the **Return Value** of the **Get Player Character** and add the **Cast To SideScrollerCharacter** node. 
 
 	![](GHT4B_5.png)
 
 	We are going to store a reference to the default character so that we may switch back to them when we press a button. 
 
-1.	Drag off the **As My Character C** pin and select **Promote to variable** and call the variable **DefaultCharacter**. 
+1.	Drag off the **As SideScrollerCharacter** pin and select **Promote to variable** and call the variable **DefaultCharacter**. 
 
 	![](GHT4B_6.png)
 
