#UE4doc #docs: Updated images and file names to reflect UE4.8, Removed phrase that no longer applies

diff --git a/Engine/Documentation/Source/Gameplay/HowTo/PossessPawns/Blueprints/GHT_4B.INT.udn b/Engine/Documentation/Source/Gameplay/HowTo/PossessPawns/Blueprints/GHT_4B.INT.udn
--- a/Engine/Documentation/Source/Gameplay/HowTo/PossessPawns/Blueprints/GHT_4B.INT.udn
+++ b/Engine/Documentation/Source/Gameplay/HowTo/PossessPawns/Blueprints/GHT_4B.INT.udn
@@ -19,7 +19,7 @@ _The steps below are based off using the **Blueprint Side Scroller Template** pr
 
 	![](GHT4B_1.png)
 
-	Here we added two **SideScrollerCharacter** Blueprints to the level and a third character is spawned via the **PlayerStart**. 
+	Here we added two **SideScrollerCharacter** Blueprints to the level. 
 
 1.	Select each of your characters in the level (the two we added to the level). 
 
