Updated Networking HowTo pages for content/images.

diff --git a/Engine/Documentation/Source/Gameplay/HowTo/Networking/ReplicateActor/Blueprints/ReplicateActor_BP.INT.udn b/Engine/Documentation/Source/Gameplay/HowTo/Networking/ReplicateActor/Blueprints/ReplicateActor_BP.INT.udn
--- a/Engine/Documentation/Source/Gameplay/HowTo/Networking/ReplicateActor/Blueprints/ReplicateActor_BP.INT.udn
+++ b/Engine/Documentation/Source/Gameplay/HowTo/Networking/ReplicateActor/Blueprints/ReplicateActor_BP.INT.udn
@@ -27,7 +27,7 @@ This page shows how to use the **Replicates** option on Actors in order to have
 
 	![](HowTo2.png)
 
-1.	Open the **Fire_NoReplicate** Blueprint and click the **Class Defaults** ![](Engine\Blueprints\UserGuide\Variables\classDefaults_button.png) button to open the Blueprint Defaults in the **Details** Panel, then un-check **Net Load on Client**. 
+1.	Open the **Fire_NoReplicate** Blueprint and click the ![](Engine\Blueprints\UserGuide\Variables\classDefaults_button.png) button, then in the **Details** Panel, un-check **Net Load on Client**. 
 
 	![](HowTo3.png)
 
@@ -35,11 +35,9 @@ This page shows how to use the **Replicates** option on Actors in order to have
 
 1.	**Compile** and **Save**, then close the Blueprint and open the **Fire_Replicate** Blueprint. 
 
-1.	Click the **Class Defaults** ![](Engine\Blueprints\UserGuide\Variables\classDefaults_button.png) button to open the Blueprint Defaults in the **Details** Panel, then check the **Replicates** checkbox.
+1.	Click the ![](Engine\Blueprints\UserGuide\Variables\classDefaults_button.png) button to open the Blueprint Defaults in the **Details** Panel, then check the **Replicates** checkbox.
 
-	[REGION:raw]
-	![](HowTo4.png)
-	[/REGION]
+	![](HowTo4.png)	
 
 	By checking **Replicates**, we are saying that this Actor can be replicated on client machines. With **Net Load on Client** also checked, the Actor is automatically loaded when the map loads and is displayed on client's machines. Where **Replicates** comes into play is, say you wanted to spawn in this Actor during gameplay at some point instead of on map load, you would need **Replicates** enabled to say that when the Actor is spawned, go ahead and spawn a replicated copy on client machines. 
 
@@ -53,12 +51,6 @@ This page shows how to use the **Replicates** option on Actors in order to have
 
 	![](HowTo6.png)(w:700)
 
-1.	Select the **PlayerStart** in the level, press **Control+W** to copy it, then press **W** and move it over a bit. 
-
-	![](HowTo7.png)
-
-	This will create an additional spawn point for multiplayer when the game is launched. 
-
 1.	Click the down-arrow next to the **Play** Button and change the **Number of Players** option to **2**. 
 
 	[REGION:raw]
@@ -80,9 +72,9 @@ The steps below showcase how to dynamically spawn Actors during gameplay using *
 
 1.	In the level viewport, delete the **Fire_Replicate** and **Fire_NoReplicate** Actors from the level.
 
-1.	From the **Modes** menu, drag in two **Target Points** and place them in your level. 
+1.	From the **Modes** menu under **All Classess**, drag two **Target Points** into your level. 
 
-	![](Switch1.png)
+	![](Switch1.png)(w:820)
 
 1.	Rename one Target Point to **No_Replicate** and the other to **Replicate** and select both in the **World Outliner**. 
 
diff --git a/Engine/Documentation/Source/Gameplay/HowTo/Networking/ReplicateFunction/Blueprints/ReplicateFunction_BP.INT.udn b/Engine/Documentation/Source/Gameplay/HowTo/Networking/ReplicateFunction/Blueprints/ReplicateFunction_BP.INT.udn
--- a/Engine/Documentation/Source/Gameplay/HowTo/Networking/ReplicateFunction/Blueprints/ReplicateFunction_BP.INT.udn
+++ b/Engine/Documentation/Source/Gameplay/HowTo/Networking/ReplicateFunction/Blueprints/ReplicateFunction_BP.INT.udn
@@ -20,51 +20,51 @@ There are 3 primary types of replicated functions: **Multicast**, **Run on Serve
 
  The steps below illustrate how to set a **Multicast** Replicated Function. 
  
-1.	First, select the **PlayerStart** in your level and duplicate 3 times so that there are four of them in your level. 
-
-1.	Using the **W** (Translation) and **E** (Rotation) key, arrange the 4 spawn points so that they face one another similar to below. 
-
-	![](HowTo1.png)
-
-1.	Click the down arrow next to the **Play** button and change the number of players to anything other than 1 (we are using 4). 
+1.	Increase the number of players by clicking the down arrow next to **Play** and changing **Number of Players** to **4**. 
 
+	[REGION:raw]
 	![](HowTo2.png)
-
+	[/REGION]
+	
 	[REGION:tip]
 	You can change the size of the windows that are generated for multiplayer inside the **Advanced Settings** option. 
 	[/REGION]
 
-1.	Next find the **Game/Blueprints** folder and open the **ThirdPersonCharacter** Blueprint. 
+1.	In the **Content Browser** under **Content/ThirdPersonBP/Blueprints**, open the **ThirdPersonCharacter** Blueprint. 
 
-1.	**Right-click** in the graph, and add a **Space Bar** Key Event and connect it to a **Spawn Actor from Class** node. 
+1.	**Right-click** in the graph and add a **Space Bar** Key Event and a **Spawn Actor from Class** node. 
 
-1.	For the **Class** set it to **Blueprint_Effect_Fire**, then off the **Spawn Transform** add a **Get Actor Transform** node. 
+1.	**Right-click** and add a **Get Actor Transform** node and connect as shown below with **Class** set to **Blueprint_Effect_Fire**. 
 
 	![](HowTo3.png)
 
-	The setup should look similar to above which, when pressed will spawn a fire effect at the player character's position. 
+	When pressed, this will spawn a fire effect at the player character's position. 
 
 1.	**Compile** and **Save**, then close the Blueprint and then click the **Play** button to play in the editor. 
 
-1.	In game, locate the window that is the "Server" and press **Space Bar** to jump. 
-
+1.	In game, locate the window that is the Server and press **Space Bar** to jump. 
+		
 	![](HowTo4.png)
+	
+	You will see that only the server window spawns the fire effect and the other players do not see the effect spawned. 
+	When attempting to jump with the other players, they too only spawn the fire effect locally and each of the other players are not aware that an effect has been spawned. 
+	Additionally, we have not attached the fire effect to the character so it only spawns at the characters location at the time **Space Bar** is pressed. 
 
-	You should see that only the server window spawns the fire effect and the other players do not see the effect spawned. When attempting to jump with the other players, they too only spawn the fire effect locally and each of the other players are not aware that an effect has been spawned. 
-
-1.	Return to the **ThirdPersonCharacter** Blueprint and disconnect the **Space Bar** Key Event by **Right-clicking** on the pin. 
+	[REGION:tip]
+	You can press **Shift + F1** to gain Mouse Control and jump between game windows during Multiplayer testing. 
+	[/REGION]
 
-1.	**Right-click** and search for then add a **Custom Event** and call the event **MulticastSpawn** or whatever name you wish. 
+1.	In the **ThirdPersonCharacter** Blueprint, **Right-click** on **Pressed** for the **Space Bar** node and break the connection. 
 
-1.	Click on the custom event, then in the **Details** panel, set the **Replicates** drop-down menu to **Multicast**. 
+1.	**Right-click** and search for then add a **Custom Event** and call the event **MulticastSpawn**. 
 
-1.	Connect the custom event to the **Spawn Actor from Class** node. 
+1.	Click on the custom event, then in the **Details** panel set the **Replicates** drop-down menu to **Multicast** and connect as shown. 
 
 	![](HowTo5.png)
 
 	The last few steps should result in the above where we have set our custom event to use **Multicast** Replication. 
 
-1.	Off the **Space Bar** Key Event, search for and add the call function **MulticastSpawn** or whatever you called it. 
+1.	Off the **Space Bar** Key Event, search for and add the call function **MulticastSpawn**. 
 
 	![](HowTo6.png)
 
@@ -74,7 +74,7 @@ There are 3 primary types of replicated functions: **Multicast**, **Run on Serve
 
 1.	**Compile** and **Save**, then close the Blueprint and then click the **Play** button to play in the editor. 
 
-1.	In game, locate the window that is the "Server" and press **Space Bar** to jump. 
+1.	In game, locate the window that is the Server and press **Space Bar** to jump. 
 
 	![](HowTo7.png)
 
@@ -87,31 +87,31 @@ There are 3 primary types of replicated functions: **Multicast**, **Run on Serve
 
  The steps below show how to set up a **Run on Server** replicated function. 
  
-1.	Inside the **Game/Blueprints** folder, open the **ThirdPersonCharacter** Blueprint. 
+1.	In the **Content Browser** under **Content/ThirdPersonBP/Blueprints**, open the **ThirdPersonCharacter** Blueprint. 
 
-1.	Select the custom event tied to the **Spawn Actor from Class** node and change the **Replicates** drop down option to **Run on Server**. 
+1.	Select the **MulticastSpawn** Custom Event and change the **Replicates** drop down option to **Run on Server**. 
 
 	![](HowTo8.png)
-
+	
 	We also updated the name of the **Graph Node** to **Run On Server** for consistency. 
 
 1.	**Compile** and **Save**, then close the Blueprint and then click the **Play** button to play in the editor. 
 
-1.	In game, locate any game window other than "Server" and press **Space Bar** to jump. 
+1.	In game, locate any game window other than Server and press **Space Bar** to jump. 
 
 	![](HowTo9.png)
 
 	Notice that the fire effect is only spawned only on the server, regardless of which player jumped, and no other players can see it but the server. While we have told this script that it should execute on the server, which it did, we need to make sure that the effect itself is set to replicate so that it gets passed to all clients as well as the server. 
 
-1.	Inside the **Game/Blueprints** folder, open the **Blueprint_Effect_Fire** Blueprint and go to the **Defaults** tab. 
+1.	Inside the **Content/StarterContent/Blueprints** folder, open the **Blueprint_Effect_Fire** Blueprint. 
 
-1.	Check the **Replicates** checkbox so that it will replicate. 
+1.	With **Class Defaults** selected from the Main Toolbar, in the **Details** panel, check the **Replicates** checkbox. 
 
 	![](HowTo10.png)
 
 1.	**Compile** and **Save**, then close the Blueprint and then click the **Play** button to play in the editor. 
 
-1.	In game, locate any game window other than "Server" and press **Space Bar** to jump. 
+1.	In game, locate any game window other than Server and press **Space Bar** to jump. 
 
 	![](HowTo11.png)
 
@@ -126,30 +126,36 @@ What we will do in this example is create a variable that only gets updated on a
 
  The steps below show how to set up a **Run on owning Client** replicated function. 
  
-1.	Inside the **Game/Blueprints** folder, open the **ThirdPersonCharacter** Blueprint. 
+1.	Inside the **Content/ThirdPersonBP/Blueprints** folder, open the **ThirdPersonCharacter** Blueprint. 
 
-1.	On the **Graph** tab in the **MyBlueprint** window, create a new variable of the **String** type and give it a name (**Inventory** for example). 
+1.	In the **MyBlueprint** window, create a new variable and name it **Inventory** then click **Compile**. 
 
-1.	**Compile**, then in the **Details** panel for the variable, set it to **Editable**, **Replicated**, and enter **Empty** for the **Default Value**. 
+	![](Inventory.png)
 
-	![](HowTo12.png)
+1.	In the **Details** panel for the variable, set it to a **String**, **Editable**, **Replicated**, and enter **Empty** for the **Default Value**. 
 
-	[REGION:note]
-	Setting this variable to **Replicated** ensures that it will be replicated over the network to connected machines. 
-	[/REGION]
+	![](HowTo12.png)
 
-1.	In the **Graph**, add a **P** Key Event connected to a **Print String** node, then hold **Control** and drag in the variable and connect as shown. 
+	Setting this variable to **Replicated** ensures that it will be replicated over the network to connected machines. We are going to use this to simulate a character picking up an item in a multiplayer game when they enter a trigger volume and removing the item when they leave the trigger volume. 
+	
+1.	Add a **P** Key Event connected to a **Print String**, then hold **Control** and drag in the **Inventory** variable and connect as shown. 
 
 	![](HowTo13.png)
 
 1.	**Compile** and **Save** then close the **ThirdPersonCharacter** Blueprint. 
 
-1.	Next, from the **Modes** window on the **Basic** tab, drag in a **Box Trigger** and place it in the middle of the player starts. 
+1.	From the **Modes** window under **Basic**, drag a **Box Trigger** into your level. 
 
 	![](HowTo14.png)
 
 	When a player character enters the trigger, we will update the variable we created, but only in the client that enters the trigger box. 
 
+1.	In the **Details** panel for the **Box Trigger** under **Rendering**, un-check  **Actor Hidden In Game**. 
+
+	![](UnHideBox.png)
+
+	This will allow us to see the box in the level when we play in the editor making it easier to test. 
+
 1.	Click the **Box Trigger** so that it is selected, then open the **Level Blueprint** from the Main Toolbar. 
 
 	![](HowTo15.png)
@@ -176,6 +182,10 @@ What we will do in this example is create a variable that only gets updated on a
 
 	![](HowTo18.png)
 
+	[REGION:note]
+	If you see an error warning on the node after creating it, click the **Compile** button to remove the error. 
+	[/REGION]
+
 1.	Create another **Custom Event** called **Remove Item** with the same settings as the **Add Item** event. 
 
 1.	Connect the **Add Item** and **Remove Item** nodes as shown below off the two overlap events. 
@@ -188,7 +198,7 @@ What we will do in this example is create a variable that only gets updated on a
 
 	![](HowTo20.png)
 
-1.	Off the **As My Character** pin, search for and add the **Set Inventory** node with the text set to **Has the Item**. 
+1.	Off the **As Third Person Character** pin, search for and add the **Set Inventory** node with the text set to **Has the Item**. 
 
 	![](HowTo21.png)
 
diff --git a/Engine/Documentation/Source/Gameplay/HowTo/Networking/ReplicateVariable/Blueprints/ReplicateVariable_BP.INT.udn b/Engine/Documentation/Source/Gameplay/HowTo/Networking/ReplicateVariable/Blueprints/ReplicateVariable_BP.INT.udn
--- a/Engine/Documentation/Source/Gameplay/HowTo/Networking/ReplicateVariable/Blueprints/ReplicateVariable_BP.INT.udn
+++ b/Engine/Documentation/Source/Gameplay/HowTo/Networking/ReplicateVariable/Blueprints/ReplicateVariable_BP.INT.udn
@@ -20,21 +20,23 @@ This page covers some sample use cases for both **Replicated** variables and a *
 
 In this example, we will create a **Replicated** variable that is used for a countdown timer that both the Server and Clients can see. 
 
-1.	Inside the **Game/StarterContent/Blueprints** folder, open the **Blueprint_Effect_Fire** Blueprint and click the **Defaults** tab. 
+1.	Inside the **Content/StarterContent/Blueprints** folder, open the **Blueprint_Effect_Fire** Blueprint. 
 
-1.	Under **Replication**, check the **Replicates** option to enable this Blueprint for replication. 
+1.	With **Class Defaults** selected, in the **Details** panel Under **Replication**, check the **Replicates** option. 
 
 	![](HT_0.png)
 
-1.	Click the **Components** tab, then click the **Add Component** button and search for and add a **Text Render** component. 
+	This will ensure that this Blueprint it replicated across the network to all connect machines. 
+
+1.	In the **Components** window, click the **Add Component** button and search for and add a **Text Render** component. 
 
 	![](HT_1.png)
 
-1.	Inside the viewport, move the text to a desired location so that it is more visible. 
+1.	On the **Viewport** tab, move the text to a desired location above the fire so that it is more visible. 
 
 	![](HT_2.png)
 
-1.	Click the **Graph** tab, **Right-click** in the graph and add an **Event Begin Play** node. 
+1.	Click the **Event Graph** tab, **Right-click** in the graph and add an **Event Begin Play** node. 
 
 1.	Drag off the **Event Begin Play**, search for and add the **Switch Has Authority** node. 
 
@@ -44,15 +46,15 @@ In this example, we will create a **Replicated** variable that is used for a cou
 
 	Here we used the **Switch Has Authority** node to call the **Set Timer** function only on the Authority (Server) when this Blueprint is spawned. We are using this node to ensure that the Timer function is executed by the server only and not by a client. If this script were to be handling game critical data, we would only want the server to execute it then replicate any data to client machines rather than the client run it as to prevent any cheating. 
 
-1.	In the **MyBlueprint** window, create a new **Integer** variable called **CountDownTimer**. 
+1.	In the **MyBlueprint** window, create a new variable called **CountDownTimer**. 
 
-1.	Click **Compile**, then set the **Default Value** for the timer to **10**. 
+1.	In the **Details** panel for the variable, change it to an **Integer** type, **Compile** and set the **Default Value** to **10**. 
 
 	![](HT_4.png)
 
-1.	**Right-click** in the graph and create a new **Custom Event** called **Update Timer**. 
+1.	**Right-click** in the graph and create a new **Custom Event** and call it **Update Timer**. 
 
-1.	Connect the event to a **Branch** node, then off the **True** pin, add the **Set CountDownTimer** node. 
+1.	Connect **Update Timer** to a **Branch** node, then off the **True** pin, add the **Set CountDownTimer** node. 
 
 1.	**Right-click** in the graph and add the **Get CountDownTimer** node, then drag off it and use the **Int - Int** node and enter 1 in the box. 
 
@@ -70,13 +72,13 @@ In this example, we will create a **Replicated** variable that is used for a cou
 
 1.	**Right-click** in the graph and add the **Event Tick** node. 
 
-1.	Hold **Control** and drag in the **TextRender1** component from the **MyBlueprint** window, then drag off it and use the **Set Text** node. 
+1.	Hold **Control** and drag in the **TextRender** component from the **MyBlueprint** window, then drag off it and use the **Set Text** node. 
 
 1.	Hold **Control** and drag in the **CountDownTimer** variable and plug it into the **Value** of the **Set Text** node. 
 
 	![](HT_7.png)
 
-	Here we are using a simple method to update the text in our text render to display to the value of the **CountDownTimer**. 
+	A conversion node will automatically be created and here we are setting the text to the value of our countdown timer. 
 
 1.	**Compile** and **Save** then close the Blueprint. 
 
@@ -84,13 +86,11 @@ In this example, we will create a **Replicated** variable that is used for a cou
 
 	![](HT_9.png)
 
-1.	Select the **PlayerStart** icon in the level, press **Control+W** to duplicate it. 
-
-1.	Arrange the **PlayerStart** and the fire effect in the level so that it is positioned in front of the player starts. 
-
-	![](HT_10.png)
+1.	On the Main Toolbar, click the down arrow next to the **Play** button and set the **Number of Players** to 2. 
 
-1.	On the Main Toolbar, click the down arrow next to the **Play** button and set the **Number of Players** to 2 (or greater). 
+	[REGION:raw]
+	![](NoOfPlayers.png)
+	[/REGION]
 
 1.	Click the **Play** button to play in the editor. 
 
@@ -100,21 +100,21 @@ In this example, we will create a **Replicated** variable that is used for a cou
 
 1.	Return to the **Blueprint_Effect_Fire** Blueprint. 
 
-1.	In the **MyBlueprint** window, click on the **CountDownTimer** variable and change the **Replication** option to use **Replicated**. 
+1.	In the **Details** panel for the **CountDownTimer** variable, under **Replication**, set it to **Replicated**. 
 
 	![](HT_12.png)
 
-1.	The graph should update, the **CountDownTimer** variable will also show an error and an icon will appear above the node. 
+	This will enable this property to be replicated to connected machines. 
 
-	![](HT_13.png)
+1.	The graph will update and each of the instances of the **CountDownTimer** variable will contain a replicated icon. 
 
-	When you click **Compile**, the error message will go away. The icon however (in the yellow box above) indicates that the variable is set to **Replicated** and makes it easier to determine which variables are replicated and which are not when looking at them in the graph. 
+	![](HT_13.png)
 
 1.	**Compile** and **Save** then close the Blueprint and play in the editor. 
 
 	![](HT_14.png)
 
-	You should now see both the server and the client counting down from 10 and displaying the text appropriately. While we are using the integer to update a text render, this variable could just as easily be used to represent a player characters health, ammo, or other gameplay related variable. 
+	You will now see both the server and the client counting down from 10 and displaying the text appropriately. While we are using the integer to update a text render, this variable could just as easily be used to represent a player characters health, ammo, or other gameplay related variable. 
 
 
 
@@ -124,11 +124,11 @@ In this example, we will create a **Replicated** variable that is used for a cou
 
 In the **Replicated** example provided above, we created a countdown timer that is displayed on both the Server and Client. In this example, we will use a **RepNotify** variable to signal when the countdown timer reaches 0 at which point we will deactivate the fire effect on both the server and the client machines. 
 
-1.	Inside the **Game/StarterContent/Blueprints** folder, open the **Blueprint_Effect_Fire** Blueprint. 
+1.	Open the **Blueprint_Effect_Fire** Blueprint. 
 
 1.	In the **MyBlueprint** window, hold **Control** and drag in the **P_Fire** effect, then drag off it and use the **Deactivate** node. 
 
-1. Off the **False** pin of the **Branch** connected to **Update** timer, connect it to the **Deactivate** node. 
+1. Off the **False** pin of the **Branch** connected to **UpdateTimer**, connect the **Deactivate** node as shown below. 
 
 	![](HT_15.png)
 
@@ -138,7 +138,7 @@ In the **Replicated** example provided above, we created a countdown timer that
 
 	![](HT_16.png)
 
-	When you play in the editor, you should see that the fire effect is deactivated when reaching 0 only on the server. 
+	When you play in the editor, you will see that the fire effect is deactivated when reaching 0 only on the server. 
 
 1.	Return to the **Blueprint_Effect_Fire** Blueprint and create a new variable of the **Boolean** type called **Deactivate**. 
 
@@ -146,11 +146,11 @@ In the **Replicated** example provided above, we created a countdown timer that
 
 	![](HT_17.png)
 
-1.	After setting the variable to **RepNotify**, you should see that a new Function has been automatically created. 
+1.	After setting the variable to **RepNotify**, you will see that a new Function has been automatically created. 
 
 	![](HT_18.png)
 
-	This function will be called whenever the variable we assigned as **RepNotify** changes and will be executed on both the server and client machines. 
+	This is called whenever the variable we assigned as **RepNotify** changes and will be executed on both the server and client machines. 
 
 1.	**Double-click** on the **On Rep Deactivate** Function to open it up. 
 
@@ -170,6 +170,6 @@ In the **Replicated** example provided above, we created a countdown timer that
 
 	![](HT_21.png)
 
-	When you play in the editor this time, you should see that both client and server deactivate the fire effect when the timer reaches 0. 
+	When you play in the editor this time, you will see that both client and server deactivate the fire effect when the timer reaches 0. 
 
 
diff --git a/Engine/Documentation/Source/Gameplay/HowTo/Networking/TestMultiplayer/TestMultiplayer.INT.udn b/Engine/Documentation/Source/Gameplay/HowTo/Networking/TestMultiplayer/TestMultiplayer.INT.udn
--- a/Engine/Documentation/Source/Gameplay/HowTo/Networking/TestMultiplayer/TestMultiplayer.INT.udn
+++ b/Engine/Documentation/Source/Gameplay/HowTo/Networking/TestMultiplayer/TestMultiplayer.INT.udn
@@ -12,37 +12,50 @@ This page shows you how you can change some of the settings for the editor for t
 
 ## Setting Number of Players
 
-You can increase the number of players that spawn in when you launch the game by clicking the down arrow next to the **Play** button. 
+**To Increase the Number of Players:**
 
-![](NumberOfPlayers.png)
+* Click the down arrow next to the **Play** button and enter a value for **Number of Players**. 
 
-From there, you can set the **Number of Players** field to the amount of players you want to spawn when the game is launched. 
+	[REGION:raw]
+	![](NumberOfPlayers.png)
+	[/REGION]
+
+By default, the Server will use the **Selected Viewport** as the play window and new windows will be created for each player that is added. 
 
 ![](PIE.png)
 
-By default, the Server will use the **Selected Viewport** as the play window and new windows will be created for each player that is added. 
 
 ## Adjusting Play Windows
 
-While it is perfectly okay to use the Editor's Viewport as a play window, for clarity you may want to use separate windows for each player. 
+While it is perfectly okay to use the Editor's Viewport as a play window, for clarity you may want to use windows for each player. 
+
+**To adjust the play window for the server:**
 
-![](NewEditorWindow.png)
+*	Click the down arrow next to the **Play** button then select **New Editor Window**. 
 
-As pictured above, click the down arrow next to the **Play** button then select the **New Editor Window** option. This will create new windows for each of the number of players you have defined to spawn in your game. 
+	[REGION:raw]
+	![](NewEditorWindow.png)
+	[/REGION]
+
+While simulated clients have their own windows, this setting will also create a seperate window for the simulated server. 
 
 ## Advanced Settings
 
 If you have set up the **Play** method to use new editor windows, you may want to adjust the size of the editor windows. 
 
-You can do this from the **Advanced Settings** option. 
+**To Adjust Window Sizes:**
+
+1.	Click the down arrow next to the **Play** button then select **Advanced Settings**. 
 
-![](AdvancedSettings.png)
+	[REGION:raw]
+	![](AdvancedSettings.png)
+	[/REGION]
 
-Click the down arrow next to **Play** and select **Advanced Settings** to open the **Editor Preferences** window to the **Play** options. 
+1.	Set the desired window sizes under the **Play in New Window** section. 
 
-![](WindowSize.png)
+	![](WindowSize.png)
 
-The **Play in New Window** size option allows you to set the size and default position of windows that are created. You can choose from a number of preset window sizes or manually enter a window size (we specified 640x480 in our example). You set the **Left** and **Top** position windows that are created or click the checkbox for **Always center window to screen**. 
+	The **Play in New Window** size option allows you to set the size and default position of windows that are created. You can choose from a number of preset window sizes or manually enter a window size (we specified 640x480 in our example). You set the **Left** and **Top** position windows that are created or click the checkbox for **Always center window to screen**. 
 
 Once you have entered a window size, when you play in the editor, each new window will be of the same size. 
 
@@ -63,6 +76,7 @@ Also inside **Advanced Settings**, there is a section for setting up additional
 | **Run Dedicated Server** | If checked, a separate dedicated server will be launched. Otherwise the first player will act as a listen server that all other players can connect to. |
 | **Route 1st Gamepad to 2nd Client** | When running multiple player windows in a single process, this option determines how the game pad input get routed. If unchecked (default) the 1st game pad is attached to the 1st window, 2nd to the 2nd window, and so on. If it is checked, the 1st game pad goes to the 2nd window. The 1st window can then be controlled by keyboard/mouse, which is convenient if two people are testing on the same computer. |
 | **Use Single Process** | This spawns multiple player windows in a single instance of Unreal Engine 4. This will load much faster, but has potential to have more issues. When this is unchecked, additional options become available. |
+| **Create Audio Device for Every Player** | Enabling this will allow rendering accurate audio from every player's perspective but will use more CPU. |
 | **Play In Editor Description** | This is a description of what will occur when playing based on the currently applied multiplayer settings. |
 
 When **Use Single Process** is checked, multiple windows are spawned in a single instance of Unreal Engine 4. When this option is unchecked, multiple instances of UE4 will be launched for each player that is assigned and additional options become available: 
@@ -81,9 +95,15 @@ When launching a multiplayer game, there are two methods in which the game can b
 
 The second method is by using a **Dedicated Server**, which as the name implies, is dedicated to only hosting the game and no local players play on the machine as everyone connected is a client. Usually running as a Dedicated Server is more optimized than running as a Listen Server as there are no visuals or input being registered. 
 
-By default, the server type is set to a Listen Server when playing in the editor or standalone game. You can change this to a dedicated server by clicking the down arrow next to the **Play** button, then clicking the checkbox next to **Run Dedicated Server**. 
+By default, the server type is set to a Listen Server when playing in the editor or standalone game. 
+
+**To run a Dedicated Server:**
+
+*	Click the down arrow next to the **Play** button then check the **Run Dedicated Server** checkbox. 
 
-![](DedicatedServer.png)
+	[REGION:raw]
+	![](DedicatedServer.png)
+	[/REGION]
 
 
 
