#UE4doc #docs: Updated type and images to UE4.8 versions

diff --git a/Engine/Documentation/Source/Gameplay/HowTo/PossessPawns/Blueprints/GHT_4B.INT.udn b/Engine/Documentation/Source/Gameplay/HowTo/PossessPawns/Blueprints/GHT_4B.INT.udn
--- a/Engine/Documentation/Source/Gameplay/HowTo/PossessPawns/Blueprints/GHT_4B.INT.udn
+++ b/Engine/Documentation/Source/Gameplay/HowTo/PossessPawns/Blueprints/GHT_4B.INT.udn
@@ -88,7 +88,7 @@ _This example is based off using the **Blueprint 3rd Person and Vehicle Template
 
 1.	Create a **Blueprint Vehicle Template** and **Blueprint 3rd Person Template** project. 
 
-1.	Inside the **Blueprint 3rd Person** project, under **Game/Blueprints**, **Right-click** on the **MyCharacter** Blueprint and choose **Migrate**(located within **Asset Actions**). 
+1.	Inside the **Blueprint 3rd Person** project, under **Game/Blueprints**, **Right-click** on the **ThirdPersonCharacter** Blueprint and choose **Migrate**(located within **Asset Actions**). 
 
 	![](GHT4B_12.png)
 
@@ -96,7 +96,7 @@ _This example is based off using the **Blueprint 3rd Person and Vehicle Template
 
 1.	Inside the **Vehicle Project**, open the **World Settings** option from the Main Toolbar. 
 
-1.	Under **Game Mode**, set the **Default Pawn Class** to use the **MyCharacter** Blueprint.
+1.	Under **Game Mode**, set the **Default Pawn Class** to use the **ThirdPersonCharacter** Blueprint.
 
 	![](GHT4B_13.png)
 
@@ -124,7 +124,7 @@ _This example is based off using the **Blueprint 3rd Person and Vehicle Template
 
 1.	For the **SpawnLocation** (Arrow), in the **Details** panel under **Location**, set the **Z** to **131.0**. 
 
-1.	In the **Content Browser**, under the **Game/Blueprints** folder, open the **MyCharacter** Blueprint. 
+1.	In the **Content Browser**, under the **Game/Blueprints** folder, open the **ThirdPersonCharacter** Blueprint. 
 
 1.	Add a variable called **Car** of the **WheeledVehicle** type and a **Bool** variable called **NearCar**. 
 
@@ -140,13 +140,13 @@ _This example is based off using the **Blueprint 3rd Person and Vehicle Template
 
 	![](GHT4B_19.png)
 
-1.	In the graph, drag off the **Other Actor** pin of the **OnComponentBeginOverlap** node and add the **Cast To MyCharacter** node. 
+1.	In the graph, drag off the **Other Actor** pin of the **OnComponentBeginOverlap** node and add the **Cast To ThirdPersonCharacter** node. 
 
 1.	Off the **As My Character C** pin, add the **Set Near Car** variable and click the check box to set it to **True**. 
 
 	![](GHT4B_20.png)(w:780)
 
-	By using the **Cast To** node, we are able to access and update the variables in our **MyCharacter** Blueprint. 
+	By using the **Cast To** node, we are able to access and update the variables in our **ThirdPersonCharacter** Blueprint. 
 
 1.	Drag off **As My Character C** again and add the **Set Car** node, then drag off the **Car** pin and add a **Get Reference To Self** node. 
 
@@ -156,9 +156,9 @@ _This example is based off using the **Blueprint 3rd Person and Vehicle Template
 
 	_Click image for full view._
 
-	Here we created the reference to the **Sedan** Blueprint and assigned it to the variable in our **MyCharacter** Blueprint. You can also add a **Print String** node to the end with some text to indicate that you are in the **Trigger** box (we set this to **Enter Car**). 
+	Here we created the reference to the **Sedan** Blueprint and assigned it to the variable in our **ThirdPersonCharacter** Blueprint. You can also add a **Print String** node to the end with some text to indicate that you are in the **Trigger** box (we set this to **Enter Car**). 
 
-1.	Drag off the **Other Actor** pin of the **OnComponentEndOverlap** node and **Cast To MyCharacter**. 
+1.	Drag off the **Other Actor** pin of the **OnComponentEndOverlap** node and **Cast To ThirdPersonCharacter**. 
 
 1.	Off the **As My Character C** pin, add the **Set Near Car** node and leave the checkbox unchecked which sets it to **False**. 
 
@@ -200,7 +200,7 @@ _This example is based off using the **Blueprint 3rd Person and Vehicle Template
 
 	_Click image for full view._
 		
-1.	Open the **MyCharacter** Blueprint, in the graph add a **E** Key Event and connect the **Pressed** pin to a **Branch** node. 
+1.	Open the **ThirdPersonCharacter** Blueprint, in the graph add a **E** Key Event and connect the **Pressed** pin to a **Branch** node. 
 
 1.	Hold **Control** and drag in the **NearCar** variable and connect it to the **Branch** condition. 
 
@@ -220,7 +220,7 @@ _This example is based off using the **Blueprint 3rd Person and Vehicle Template
 
 1.	Drag off the **As Sedan C** pin and add the **Set Player** node, then drag off **Player** and add a **Get Reference To Self** node.
 
-1.	Your full sequence in **MyCharacter** should look like below. 
+1.	Your full sequence in **ThirdPersonCharacter** should look like below. 
 
 	[REGION:lightbox]
 	[![](GHT4B_26.png)(w:780)](GHT4B_26.png)
@@ -228,9 +228,9 @@ _This example is based off using the **Blueprint 3rd Person and Vehicle Template
 
 	_Click image for full view._
 
-	Here we are saying when **E** is pressed (if **NearCar** is _True_ which is set by entering the trigger box in the **Sedan** Blueprint), take possession of the car and set the **MyCharacter** Actor to hidden in game and remove its collision.  We are then telling the **Sedan** Blueprint that its **Player** variable in the Blueprint should reference the **MyCharacter** Blueprint so that we can return to **MyCharacter** when exiting the vehicle. 
+	Here we are saying when **E** is pressed (if **NearCar** is _True_ which is set by entering the trigger box in the **Sedan** Blueprint), take possession of the car and set the **ThirdPersonCharacter** Actor to hidden in game and remove its collision.  We are then telling the **Sedan** Blueprint that its **Player** variable in the Blueprint should reference the **ThirdPersonCharacter** Blueprint so that we can return to **ThirdPersonCharacter** when exiting the vehicle. 
 
-1. From the **Content Browser**, drag the **Sedan** Blueprint into the level. 
+1.  If there is a Sedan already placed in the level, delete it. Then, from the **Content Browser**, drag the **Sedan** Blueprint into the level. 
 
 1.	If you click **Compile** and then **Play** from the Main Toolbar to play in the editor, you should end up with something similar to below. 
 
