#UE4 Docs: First round of the OWT QS.

diff --git a/Engine/Documentation/Source/Engine/OpenWorldTools/Grass/QuickStart/1/GTQS_1.INT.udn b/Engine/Documentation/Source/Engine/OpenWorldTools/Grass/QuickStart/1/GTQS_1.INT.udn
new file mode 100644
--- /dev/null
+++ b/Engine/Documentation/Source/Engine/OpenWorldTools/Grass/QuickStart/1/GTQS_1.INT.udn
@@ -0,0 +1,119 @@
+Availability:Public
+Title:1. Grass Tool Initial Setup
+Crumbs:
+Description: How To set up the Initial Actors the Grass Tool needs to work.
+Navigation:topic
+SkillLevel:Intermediate
+Version: 4.8
+
+[VAR:Steps]
+[OBJECT:Navigation]
+	[PARAM:previous]
+	[/PARAM]
+	[PARAM:current]
+	[/PARAM]
+	[PARAM:home]
+		[How-To Home](Engine\OpenWorldTools\Grass\QuickStart "Engine\OpenWorldTools\Grass\QuickStart:Title%")
+	[/PARAM]
+	[PARAM:next]
+	[Next Step](Engine\OpenWorldTools\Grass\QuickStart\2 "%Engine\OpenWorldTools\Grass\QuickStart\2:title%") 
+	[/PARAM]
+[/OBJECT]
+[/VAR]
+
+%Steps%
+
+In the following section we will create a new level and all of the required Actors the Grass Tool requires to function correctly inside of a UE4 level.
+
+1. First create a new level inside UE4 by going to the **Main Toolbar** then going to **File** and clicking on the **New Level** option. When prompted pick the **Default Template** as a base level.
+
+	![](Engine\OpenWorldTools\ProceduralFoliage\QuickStart\1\T_New_Level.png)
+
+1. After the new level has loaded, add a new Landscape actor to the level by first clicking on the **Landscape Tool** icon and then pressing the **Create** button to add the Landscape Actor to the level.
+
+	![](Engine\OpenWorldTools\ProceduralFoliage\QuickStart\1\T_Create_Landscape.gif)(convert:false)
+
+1. Next in the **Landscape Tools** under the **Landscape Editor** section click on the **Sculpt Tool** icon to expose the available tools then pick the **Noise** tool from the tool list and adjust the following properties with the following values.
+	
+	![](T_Sculpt_Tool_Settings.png)
+
+	* **Brush Size:** 65536.0
+	* **Tool Strength:** 0.01
+	* ** Noise Scale:** 256
+
+1. Now position the Landscape Brush in the viewport so that it covers the entire Landscape Terrain then click the **Left Mouse Button**	three to four times to add some very subtle noise to the Landscape surface. 
+	When completed you should have something that looks similar to the following image.
+
+	![](T_Noise_On_Landscape.png)
+
+1. Exit the Landscape Tools by clicking on the **Place** icon in the **Modes** panel.
+
+	![](T_Exit_LS_Tools.png)
+
+1. Finally save your work and level by pressing the **Save All** button to save all content and the ** Save** button to save the level. 
+	When prompted for a level name use the name **GT_00**.
+
+	![](T_Save_All.png)
+
+
+
+
+
+
+
+
+
+
+
+
+
+
+
+
+
+
+
+
+
+
+
+
+
+
+
+
+
+
+
+
+
+
+
+
+
+
+
+
+
+
+
+
+
+
+
+
+
+
+
+
+
+
+
+
+
+
+
+
+
+
+%Steps%
\ No newline at end of file
diff --git a/Engine/Documentation/Source/Engine/OpenWorldTools/Grass/QuickStart/2/GTQS_2.INT.udn b/Engine/Documentation/Source/Engine/OpenWorldTools/Grass/QuickStart/2/GTQS_2.INT.udn
new file mode 100644
--- /dev/null
+++ b/Engine/Documentation/Source/Engine/OpenWorldTools/Grass/QuickStart/2/GTQS_2.INT.udn
@@ -0,0 +1,119 @@
+Availability:Public
+Title:2. Grass Tool Actor Creation and Setup
+Crumbs:
+Description: How to create and setup the required Actors for the Grass Tool
+Navigation:topic
+SkillLevel:Intermediate
+Version: 4.8
+
+[VAR:Steps]
+[OBJECT:Navigation]
+	[PARAM:previous]
+	[Previous Step](Engine\OpenWorldTools\Grass\QuickStart\1 "%Engine\OpenWorldTools\ProceduralFoliage\QuickStart\1:title%") 
+	[/PARAM]
+	[PARAM:current]
+	[/PARAM]
+	[PARAM:home]
+		[How-To Home](Engine\OpenWorldTools\Grass\QuickStart "%Engine\OpenWorldTools\Grass\QuickStart:Title%")
+	[/PARAM]
+	[PARAM:next]
+	[Next Step](Engine\OpenWorldTools\Grass\QuickStart\3 "%Engine\OpenWorldTools\Grass\QuickStart\3:title%") 
+	[/PARAM]
+[/OBJECT]
+[/VAR]
+
+%Steps%
+
+In the following section we will be creating the required Actors and Materials the Grass Tool needs to function correctly.
+We will continue to work with the **GT_00** level that was created in the last step so if you have closed down this level please re-open it before going any further.
+
+1. The first thing we need to do is to create a **Landscape Grass Type** which will tell our Landscape which Static Mesh to spawn and the amount it should spawn. 
+	To do this **Right - Click** with the Mouse in the **Content Browser** and from the menu that is displayed go to **Create Advanced Asset** > **Miscellaneous** > **Landscape Grass Type** and then give it a name of **Grass_00**.
+
+	![](T_Create_LS_Grass.png)
+
+1. Next open up the Landscape Grass Type by **Double - Clicking** on it with the **Left Mouse button** and 
+	once opened add a new item to the **Grass Varieties** array by pressing the **Plus** icon that is to the right of the **Grass Varieties** name.
+	
+	![](T_Add_New_GV.png)
+
+1. Now look for the **Grass Mesh** section and then click on the input box that says **None** then search for **SM_FieldGrass_01** Static Mesh by entering **SM_FieldGrass_01** as the search term. 
+
+	![](T_Set_Grass_Type.gif)(convert:false)
+
+1. Once the Static Mesh has been added make sure the following properties are enabled.
+
+	* **Grass Density:** 400.0
+	* **Use Grid:** Enabled
+	* **Random Rotation:** Enabled
+	* **Align to Surface:** Enabled
+
+	![](T_Grass_Props.png)
+
+1. Finally make sure to save all of you work by pressing both the **Save All** and **Save** buttons.
+
+	![](Engine\OpenWorldTools\ProceduralFoliage\QuickStart\T_Uni_Save_Icons.png)
+
+
+
+
+
+
+
+
+
+
+
+
+
+
+
+
+
+
+
+
+
+
+
+
+
+
+
+
+
+
+
+
+
+
+
+
+
+
+
+
+
+
+
+
+
+
+
+
+
+
+
+
+
+
+
+
+
+
+
+
+
+
+
+%Steps%
\ No newline at end of file
diff --git a/Engine/Documentation/Source/Engine/OpenWorldTools/Grass/QuickStart/3/GTQS_3.INT.udn b/Engine/Documentation/Source/Engine/OpenWorldTools/Grass/QuickStart/3/GTQS_3.INT.udn
new file mode 100644
--- /dev/null
+++ b/Engine/Documentation/Source/Engine/OpenWorldTools/Grass/QuickStart/3/GTQS_3.INT.udn
@@ -0,0 +1,85 @@
+Availability:Public
+Title:3. Grass Tool Material Setup
+Crumbs:
+Description: How to setup a Material that will work with the Grass system.
+Navigation:topic
+SkillLevel:Intermediate
+Version: 4.8
+
+[VAR:Steps]
+[OBJECT:Navigation]
+	[PARAM:previous]
+	[Previous Step](Engine\OpenWorldTools\Grass\QuickStart\1 "%Engine\OpenWorldTools\ProceduralFoliage\QuickStart\1:title%") 
+	[/PARAM]
+	[PARAM:current]
+	[/PARAM]
+	[PARAM:home]
+		[How-To Home](Engine\OpenWorldTools\Grass\QuickStart "%Engine\OpenWorldTools\Grass\QuickStart:Title%")
+	[/PARAM]
+	[PARAM:next]
+	[Next Step](Engine\OpenWorldTools\Grass\QuickStart\4 "%Engine\OpenWorldTools\Grass\QuickStart\4:title%") 
+	[/PARAM]
+[/OBJECT]
+[/VAR]
+
+%Steps%
+
+Before we can begin to use the Grass Tools we first have to create Material that can work with both the Landscape terrain and the **Landscape Grass Type**. 
+In the following section we will cover how to create and setup this Material as well as how to link it so that it works with the Landscape Grass Type. 
+
+[region:tip]
+If you would like to have a more in-depth look at how the Landscape terrain works in UE4 check out the [Landscape](https://docs.unrealengine.com/latest/INT/Engine/Landscape/index.html) pages for more information.
+[/region]
+
+1. First create a new Material for our Landscape Terrain by **Right - Click** in the **Content Browser** with the mouse and then select the **Material** option form the **Create Basic Asset** section. 
+	Name the new Material **MAT_GT_Grass**
+
+	![](T_Create_New_Material.png)
+
+1. Open up the **MAT_GT_Grass** Material by **Double - Clicking** with the **Left Mouse Button** inside of the **Content Browser** and then add the following two Textures from the **Open World Demo Collection** to the Material Graph.
+
+	* **T_AlpinePatch001_D_alt_R**
+	* **T_GDC_Grass01_D_NoisyAlpha**
+
+	![](T_Added_Textures.png)
+
+1. Next search for the following Material Expression nodes using the **Palette** search function.
+	Once you locate a required Material Expression node select it in the Palette by clicking on it with the **Left Mouse Button** and then dragging it into the Material Graph, releasing the **Left Mouse Button** to place it in the Material graph. 
+
+	* **Landscape Layer Blend** x 1
+	* **Landscape Layer Sample** x 1
+	* **Landscape Grass Output** x 1
+
+	![](T_Add_Material_Nodes.gif)(convert:false)
+
+1. Now select the **Landscape Layer Blend** node in the Material Graph and add two new **Layers** to it by pressing the **Plus**  icon twice. 
+	Once both layers have been added, set the **Layer Name** of one to **Grass** and set the **Layer Name** of the other to **Rock** and set the **Preview Weight** to 1.0 for both the Grass and Rock. 
+
+	![](T_LS_LayerBlend_Setup.png)
+
+1. Next connect the **T_AlpinePatch001_D_alt_R** Texture to the **Layer Rock** input on the **Landscape Layer Blend** node, then connect the **T_GDC_Grass01_D_NoisyAlpha** to the **Later Grass** input. 
+	Once that is completed connect the **Output** of the **Landscape Layer Blend** node into the **Base Color** input on the **Main Material Node**.
+
+	![](T_Hook_Up_Textures.png)
+	
+1. Then select the **Landscape Grass Output** node in the Material Graph and then in the **Details** panel in the **Grass Types** section press the **Plus** sign icon to add a new element to the Grass Types array. 
+	
+	![](T_Grass_Out_SetUp.gif)(convert:false)
+
+1. Now in **Content Browser** select the **Grass_00** Landscape Grass Type that was created in the last step and then in the **Material** under the **Grass Type** option, 
+	press the **Arrow** icon to load the Actor that is currently selected in the Content Browser. 
+
+	![](T_Input_Grass_Type.png)
+
+1. Next select the **Landscape Layer Sample** node under the **Parameter Name** input **Grass** for the name and connect the Output from the **Landscape Layer Sample** to the input of the **Landscape Grass Output** node. 
+
+	![](T_LLW_Setup.png)
+
+1. Finally press the **Apply** and **Save** buttons to compile and save the Material.
+
+	![](T_Final_Material.png)
+
+
+
+
+%Steps%
\ No newline at end of file
diff --git a/Engine/Documentation/Source/Engine/OpenWorldTools/Grass/QuickStart/4/GTQS_4.INT.udn b/Engine/Documentation/Source/Engine/OpenWorldTools/Grass/QuickStart/4/GTQS_4.INT.udn
new file mode 100644
--- /dev/null
+++ b/Engine/Documentation/Source/Engine/OpenWorldTools/Grass/QuickStart/4/GTQS_4.INT.udn
@@ -0,0 +1,111 @@
+Availability:Public
+Title:4. Using the Grass Tools
+Crumbs:
+Description: How to use the Grass Tools.
+Navigation:topic
+SkillLevel:Intermediate
+Version: 4.8
+
+[VAR:Steps]
+[OBJECT:Navigation]
+	[PARAM:previous]
+	[Previous Step](Engine\OpenWorldTools\Grass\QuickStart\3 "%Engine\OpenWorldTools\ProceduralFoliage\QuickStart\3:title%") 
+	[/PARAM]
+	[PARAM:current]
+	[/PARAM]
+	[PARAM:home]
+		[How-To Home](Engine\OpenWorldTools\Grass\QuickStart "%Engine\OpenWorldTools\Grass\QuickStart:Title%")
+	[/PARAM]
+	[PARAM:next]
+	[/PARAM]
+[/OBJECT]
+[/VAR]
+
+%Steps%
+
+In order to see the Grass system in action we need to apply the Material that was created in the last step to the Landscape and then use the Landscape painting tools to define where we want the grass to be spawned.
+In the following section we will go over how to apply our Material to the Landscape Terrain and then how to use the Landscape Painting Tools to define areas where grass should spawn.
+We will continue to work with the **GT_00** level that was created in the last step so if you have closed down this level please re-open it before going any further.
+
+1. First select the Landscape Actor in that was placed in the level by clicking on it with the **Left Mouse Button** in the viewport.
+
+	![](T_Selected_Landscape.png)
+
+1. Next locate the **MAT_GT_Grass** Material in the **Content Browser** and click on it with the **Left Mouse Button** to select it,
+	then on the Landscape Terrain under the **Details** panel in the **Landscape Material** section, press the **Arrow** icon to apply the **MAT_GT_Grass** Material to the Landscape Terrain.
+
+	![](T_Apply_Material.png)
+
+1. Now in the **Modes** panel go to the **Landscape** section by clicking on the **Landscape** icon and then go into **Paint** mode by clicking on the **Paint** icon.
+
+	![](T_Landscape_Paint_Mode.png)
+
+1. Then under the **Target Layers** section, add a new **Layer Info** by pressing the **Plus** sign icon that is on the far right of the layer name and when prompted, 
+	select the **Weight- Blended Layer(normal)** option and then select a location for the new Layer Blend to be saved to inside of the Content Browser. 
+	Make sure that you create **Layer Info** for both the Rock and the Grass.
+
+	![](T_Add_LWB.gif)(convert:false)
+
+1. After the **Landscape Layers** are created, they need to be applied in the **Target Layers** to work. 
+	To do this select the **Rock_Layerinfo** in the Content Browser and then under **Target Layers** input the  **Rock_Layerinfo** 
+	into the rock section by pressing the **Arrow** icon to load the currently selected object.
+
+	![](T_Load_Layer_Info.png)
+
+1. Next select the **Grass** Target Layer and then inside of the viewport hold down the **Left Mouse Button** to begin painting the Grass Material to the Landscape Terrain.
+	For this step try and completely cover the Landscape giving you results that look like this.
+	
+	![](T_Painting_Grass.png)
+
+	[region:note]
+	Depending on how powerful your development PC is when you start to paint on the Landscape you might experience some lag in the editor as the grass is spawned. 
+	This is normal behavior as the grass is dynamically spawned after you have finished painting. 
+	If you find that the lag makes it to hard to work, try turning down the **Grass Density** while painting and then turn it back up when done.
+	[/region]
+
+1. To remove grass from the Landscape Terrain select the **Rock Target Layer** then inside of the viewport hold down the **Left Mouse Button** 
+	to begin painting the Rock Texture with will remove the grass.
+	Adjusting the **Brush Size** and **Tool Strength** will help to better define how the grass is placed or removed while painting on the Landscape.
+
+	![](T_Painting_Removing_Grass.gif)(convert:false)
+
+
+
+
+
+
+
+
+
+
+
+
+
+
+
+
+
+
+
+
+
+
+
+
+
+
+
+
+
+
+
+
+
+
+
+
+
+
+
+
+%Steps%
\ No newline at end of file
diff --git a/Engine/Documentation/Source/Engine/OpenWorldTools/Grass/QuickStart/GrassToolsQuickStart.INT.udn b/Engine/Documentation/Source/Engine/OpenWorldTools/Grass/QuickStart/GrassToolsQuickStart.INT.udn
new file mode 100644
--- /dev/null
+++ b/Engine/Documentation/Source/Engine/OpenWorldTools/Grass/QuickStart/GrassToolsQuickStart.INT.udn
@@ -0,0 +1,125 @@
+Availability:Public
+Title: Grass Tool Quick Start
+Crumbs:
+Description:A Quick Start guide for setting up and using the Grass Tools.
+Navigation:topic
+SkillLevel:Intermediate
+Version: 4.8
+
+## Overview
+
+While the Procedural Foliage tools are great a simulating a huge forest that has many different species of trees and bushes in it, the one thing that it can not do very well is spawn the thousands and thousands of Static Meshes that would be needed to make the Landscape look like it was densely covered in grass or rocks. 
+To accomplish this a new tool was introduced to UE4, the **Grass Tool**. 
+The Grass Tool allows you to define a Static Mesh and amount to spawn on the Landscape Terrain whenever a certain Texture is painted on the Landscape. 
+In the following How - To we will cover all you need to know about using the Grass Tool in your levels.
+
+![](T_Grass_Intro.png)
+
+[region:note]
+Currently the Grass system will only work with the Landscape Terrain actor. Trying to use the Grass system to spawn grass on any other actor type will not work.
+[/region]
+
+[OBJECT:TopicButtonList]
+	[PARAM:icon]
+		![](%ROOT%/build_icon.png)(convert:false)
+	[/PARAM]
+	[PARAM:title]
+		Steps
+	[/PARAM]
+	[PARAM:description]
+	[/PARAM]
+	[PARAM:links]
+	* [](Engine\OpenWorldTools\Grass\QuickStart\1)
+	* [](Engine\OpenWorldTools\Grass\QuickStart\2)
+	* [](Engine\OpenWorldTools\Grass\QuickStart\3)
+	* [](Engine\OpenWorldTools\Grass\QuickStart\4)
+	[/PARAM]
+[/OBJECT]
+
+
+
+
+
+
+
+
+
+
+
+
+
+
+
+
+
+
+
+
+
+
+
+
+
+
+
+
+
+
+
+
+
+
+
+
+
+
+
+
+
+
+
+
+
+
+
+
+
+
+
+
+
+
+
+
+
+
+
+
+
+
+
+
+
+
+
+
+
+
+
+
+
+
+
+
+
+
+
+
+
+
+
+
+
+
+
+
diff --git a/Engine/Documentation/Source/Engine/OpenWorldTools/Overview/OpenWorldToolsOverview.INT.udn b/Engine/Documentation/Source/Engine/OpenWorldTools/Overview/OpenWorldToolsOverview.INT.udn
new file mode 100644
--- /dev/null
+++ b/Engine/Documentation/Source/Engine/OpenWorldTools/Overview/OpenWorldToolsOverview.INT.udn
@@ -0,0 +1,178 @@
+Availability:Public
+Title: Open World Tools Overview
+Crumbs:
+Description: A system for procedurally placing Static Meshes in your levels.
+Navigation:topic
+SkillLevel:Intermediate
+Version: 4.8
+
+[TOC (start:2 end:3)]
+
+![](T_Kite_Demo_00.png)
+
+During the 2015 Games Developer Conference(GDC) Epic showed off to the world UE4's ability to render massive open worlds full of both flora and fauna.
+With the release of Unreal Engine 4.8 (UE4) the tools used to create this demo, referred to as the Open World Tools, are now available for use in your UE4 projects. 
+In the following set of tutorials we will be taking a look at these tools and how you can use them in your UE4 projects.
+
+## Prerequisites
+
+Before you can use the Procedural Foliage tools in your project, you must first enable them by going to the **Main Toolbar** and then going to **Edit** > **Editor Preferences** > **Experimental**. 
+Once you are in the Experimental section enable Procedural Foliage option by clicking on the checkmark box next to the word **Procedural Foliage**. 
+
+![](T_Enable_PFT.gif)(convert:false)
+
+## Types of Procedural Tools 
+The Procedural Tools that come with the 4.8 release of Unreal Engine 4 allow you to place Static Meshes inside your levels using a computer algorithm to determine placement instead having to do this work manually. 
+These tools are great if you want to quickly generate a natural looking forest comprised of many different species of trees, or have a certain species of bush spawn whenever a certain Texture is painted on the Landscape Terrain.
+
+When talking about the Procedural Foliage tools inside UE4, we are actually referring to the two different tool sets. 
+The tool set that you use depends on what you want to accomplish. 
+For example if you want to make an expansive pine and birch tree forest then you're going to want to use the **Procedural Foliage Tools**. 
+If you want to densely populate the Landscape Terrain used in that pine and birch tree forest with some grass, flowers and small rocks then you would use the **Grass Tools**.
+
+### Procedural Foliage Tool 
+The Procedural Foliage Tool is used for simulating huge forests that are filled with lots of different kinds of trees and bushes. 
+This tool works by simulating how a forest grows over the years using **Steps** to define years of growth. 
+At each Step in the simulation new virtual seeds are cast into the level to act as spawn locations for other Foliage Type Actors. 
+By tweaking the various Foliage Types options like **Spread**, **Size**, and **Priority Settings**, you will affect how those Foliage Types spread and compete for spawn locations through the level the Procedural Foliage Spawner in placed in.
+To learn more about how to use the Procedural Foliage Tool, check out the Procedural Foliage Quick Start that is linked below.
+
+[REGION:topics third]
+[OBJECT:Topic]
+	[PARAM:image]
+		![](T_PFT_Trees_00.png)
+	[/PARAM]
+	[PARAM:icon]
+		![](%ROOT%/editor_icon.png)(convert:false)
+	[/PARAM]
+	[PARAM:title]
+		Procedural Foliage Tool Quick Start
+	[/PARAM]
+	[PARAM:description]
+		A comprehensive guide for getting the Procedural Foliage Tool setup and working.
+	[/PARAM]
+	[PARAM:path]
+		[RELATIVE:Engine/OpenWorldTools/ProceduralFoliage/QuickStart]
+	[/PARAM]
+[/OBJECT]
+[/region]
+
+#### Foliage Type Actors
+The Procedural Foliage Tools makes use of Actors called **Foliage Types**. 
+Foliage Types allow you to define various properties the Foliage should have like what Static Mesh to use or how big Static Mesh should be spawned.
+The chart below gives a very high level explanation of the four main sets of properties that the Foliage Type Actor are broken down into and what those options do.
+
+![](T_Foliage_Type_Actor.png)
+
+| Property Name| Description |
+| ------------ | ----------- |
+|**Mesh**| The Mesh section is where you define which Static Mesh to use for the foliage.|
+|**Placement**| The Placement section is where you can adjust how the Static Mesh specified in the Mesh section are placed on objects through out the level.|
+|**Growth**| The Growth section helps define how big the trees will be at the beginning of the simulation and how big they will "Grow" to by the end of the simulation.|
+|**Instance Settings**| The Instance Settings allows you to adjust how the Static Mesh used for the Foliage Type will be displayed in the level.|
+
+### Grass Tool
+[region:note]
+Currently the Grass Tools will only work in conjunction with the Landscape Terrain system. 
+[/region]
+The Grass Tools are used to densely cover the Landscape Terrain in things Static Mesh to simulate things like grass, flowers, and small rocks. 
+You can use the Grass Tools to define a Static Mesh that will spawn when a certain Landscape Layer is painted on the Landscape Terrain. 
+This tool does not use any kind of simulation like the Procedural Foliage Spawner does, it simply fills the specified area at the specified density with the specified Static Mesh. 
+To lean more about how to use the Grass Tool, check out the Grass Tool Quick Start that is lined below.
+
+[REGION:topics third]
+[OBJECT:Topic]
+	[PARAM:image]
+		![](T_PFT_Grass_00.png)
+	[/PARAM]
+	[PARAM:icon]
+		![](%ROOT%/editor_icon.png)(convert:false)
+	[/PARAM]
+	[PARAM:title]
+		Grass Tool Quick Start
+	[/PARAM]
+	[PARAM:description]
+		A comprehensive guide for getting the Grass Tool setup and working.
+	[/PARAM]
+	[PARAM:path]
+		[RELATIVE:Engine/OpenWorldTools/Grass/QuickStart]
+	[/PARAM]
+[/OBJECT]
+[/region]
+
+#### Grass Type Actors
+The Grass Tool makes use of Actors called **Landscape Grass Types**.
+Landscape Grass Types allow you to define various properties the Grass should have like what Static Mesh to use for the grass and how many of those Static Meshes should be spawned.
+The chart below gives a very high level explanation of the main sets of properties that the Landscape Grass Type Actor is broken down into and what those options do.
+
+![](T_LS_Grass_Type_Actor.png)
+
+| Property Name| Description |
+| ------------ | ----------- |
+|**Grass Varieties**| Inside of Grass Varieties you can specify what type of Static Mesh to use for the grass as well as grass density and cull distance settings.|
+
+
+
+
+
+
+
+
+
+
+
+
+
+
+
+
+
+
+
+
+
+
+
+
+
+
+
+
+
+
+
+
+
+
+
+
+
+
+
+
+
+
+
+
+
+
+
+
+
+
+
+
+
+
+
+
+
+
+
+
+
+
+
+
+
+
diff --git a/Engine/Documentation/Source/Engine/OpenWorldTools/ProceduralFoliage/QuickStart/1/PFQS_1.INT.udn b/Engine/Documentation/Source/Engine/OpenWorldTools/ProceduralFoliage/QuickStart/1/PFQS_1.INT.udn
new file mode 100644
--- /dev/null
+++ b/Engine/Documentation/Source/Engine/OpenWorldTools/ProceduralFoliage/QuickStart/1/PFQS_1.INT.udn
@@ -0,0 +1,116 @@
+Availability:Public
+Title: 1. Procedural Foliage Tool Actors Creation
+Crumbs:
+Description:In this step we create all of the required Actors.
+Navigation:topic
+SkillLevel:Intermediate
+Version: 4.8
+
+[VAR:Steps]
+[OBJECT:Navigation]
+	[PARAM:previous]
+	[/PARAM]
+	[PARAM:current]
+	[/PARAM]
+	[PARAM:home]
+		[How-To Home](Engine\OpenWorldTools\ProceduralFoliage\QuickStart "Engine\OpenWorldTools\ProceduralFoliage\QuickStart:Title%")
+	[/PARAM]
+	[PARAM:next]
+	[Next Step](Engine\OpenWorldTools\ProceduralFoliage\QuickStart\2 "%Engine\OpenWorldTools\ProceduralFoliage\QuickStart\2:title%") 
+	[/PARAM]
+[/OBJECT]
+[/VAR]
+
+%Steps%
+
+In the following How To we will create a new Level, Landscape Terrain and all of Actors the Procedural Foliage Tool requires to function correctly inside UE4.
+
+
+1. First create a new level inside UE4 using the **Default Template** as a base.
+
+	![](T_New_Level.png)
+
+1. After the new level has loaded, add a new Landscape actor to the level by first clicking on the **Landscape Tool** icon and then pressing the **Create** button to add the Landscape Actor to the level.
+	
+	![](T_Create_Landscape.gif)(convert:false)
+
+1. Once the Landscape has been added to the level exit the Landscape tool by clicking on the **Place** icon in the **Modes Panel**.
+
+	![](T_Exit_Landscape_Mode.png)
+
+	[region:note]
+	The reason we are using a Landscape Terrain this is because the Landscape Terrain Actor quickly provides us with a very large area to spawn lots of Static Meshes on.
+	[/region]
+
+1. Next create a new Procedural Foliage Spawner by **Right- Clicking** in the **Content Browser** with the mouse and going to the **Create Advanced Asset** menu option then **Miscellaneous** > **Procedural Foliage Spawner**. 
+	
+	![](T_Create_PFS.png)
+
+1. Now name the Procedural Foliage Spawner **PFS_Example**.
+
+	![](T_Name_PFS.png)
+
+1. Then drag the Procedural Foliage Spawner from the **Content Browser** into the level and position it so that it is at **0,0,200** in the level then scale Procedural Foliage Spawner to **100,100,10** in the X, Y and Z Axis.
+
+	![](T_Add_PFS_To_Level.gif)(convert:false)
+
+1. Next  create the **Foliage Type** Actors by  **Right- Clicking**  with the Mouse in the Content Browser and then go to the **Create Advanced Asset** section and then **Miscellaneous** > **Foliage Type** and name it **Tree_00**.
+
+	![](T_Create_FT.png)
+
+1. Finally save your work and level by pressing the **Save All** button to save all content and the ** Save** button to save the level. 
+	When prompted for a level name use the name **PFT_00**.
+
+	![](T_Save_All.png)
+
+%Steps%
+
+
+
+
+
+
+
+
+
+
+
+
+
+
+
+
+
+
+
+
+
+
+
+
+
+
+
+
+
+
+
+
+
+
+
+
+
+
+
+
+
+
+
+
+
+
+
+
+
+
diff --git a/Engine/Documentation/Source/Engine/OpenWorldTools/ProceduralFoliage/QuickStart/2/PFQS_2.INT.udn b/Engine/Documentation/Source/Engine/OpenWorldTools/ProceduralFoliage/QuickStart/2/PFQS_2.INT.udn
new file mode 100644
--- /dev/null
+++ b/Engine/Documentation/Source/Engine/OpenWorldTools/ProceduralFoliage/QuickStart/2/PFQS_2.INT.udn
@@ -0,0 +1,63 @@
+Availability:Public
+Title: 2. Procedural Foliage Tool Actors Setup
+Crumbs:
+Description:In this step we setup all of the required Actors.
+Navigation:topic
+SkillLevel:Intermediate
+Version: 4.8
+
+[VAR:Steps]
+[OBJECT:Navigation]
+	[PARAM:previous]
+	[Previous Step](Engine\OpenWorldTools\ProceduralFoliage\QuickStart\1 "%Engine\OpenWorldTools\ProceduralFoliage\QuickStart\1:title%") 
+	[/PARAM]
+	[PARAM:current]
+	[/PARAM]
+	[PARAM:home]
+		[How-To Home](Engine\OpenWorldTools\ProceduralFoliage\QuickStart "Engine\OpenWorldTools\ProceduralFoliage\QuickStart:Title%")
+	[/PARAM]
+	[PARAM:next]
+	[Next Step](Engine\OpenWorldTools\ProceduralFoliage\QuickStart\3 "%Engine\OpenWorldTools\ProceduralFoliage\QuickStart\3:title%") 
+	[/PARAM]
+[/OBJECT]
+[/VAR]
+
+%Steps%
+In the following section we will cover how to set up the **Foliage Type Actors** to work with the Procedural Foliage Spawner. 
+We will continue to work with the **PFT_00** level that was created in the last step so if you have closed down this level please re-open it before going any further.
+
+1. First open up the **Procedural Foliage Spawner** by **Double - Clicking** on it with the **Left Mouse Button** inside of the Content Browser. 
+
+	![](T_PFS_Opened.png)
+
+1. Next add a new item to the **Foliage Types** array by pressing the **Plus Sign** icon that is to the right of the **Foliage Types** menu option.
+
+	![](T_PFS_Add_FT.png)
+
+1. Then in the Content Browser select the **Tree_00** Foliage Type and in the **Foliage Types** section of the Procedural Foliage Spawner press the **Arrow** icon to load the selected Foliage Type into the Procedural Foliage Spawner.
+
+	![](T_Assign_PFT_To_PFS.gif)(convert:false)
+
+1. With a Foliage Type now set for Procedural Foliage Spawner open up the **Tree_00** Foliage Type so that we can do a very basic setup to ensure everything is working.
+
+	![](T_PFT_Window.png)
+
+1. At the top of the Tree_00 Foliage Type locate the **Mesh** section and then click on the drop down menu that says **None** with the **Left Mouse Button**. 
+	Use the search menu locate  **HillTree_02** Static Mesh from the **Open World Demo Collection** by inputting **HillTree_02** as the search term. 
+	Once the **HillTree_02** is located, load it by clicking on it with the **Left Mouse Button**. 
+
+	![](T_Set_Tree_Mesh.gif)(convert:false)
+
+1. Next inside the viewport select the **Procedural Foliage Spawner** that was placed in the level and then expand the **Procedural Foliage** section. 
+	
+	![](T_PFV_Select_In_Level.png)
+
+1. Under the **Procedural Foliage** section press the **Resimulate** button and you should now have something that looks like this in your level.
+
+	![](T_Final_Results.png)
+
+1. Finally save your work and level by pressing the **Save All** button to save all content and the ** Save** button to save the level.  
+
+	![](Engine\OpenWorldTools\ProceduralFoliage\QuickStart\T_Uni_Save_Icons.png)
+
+%Steps%
\ No newline at end of file
diff --git a/Engine/Documentation/Source/Engine/OpenWorldTools/ProceduralFoliage/QuickStart/3/PFQS_3.INT.udn b/Engine/Documentation/Source/Engine/OpenWorldTools/ProceduralFoliage/QuickStart/3/PFQS_3.INT.udn
new file mode 100644
--- /dev/null
+++ b/Engine/Documentation/Source/Engine/OpenWorldTools/ProceduralFoliage/QuickStart/3/PFQS_3.INT.udn
@@ -0,0 +1,124 @@
+Availability:Public
+Title: 3. Procedural Foliage Tool tweaking Actor properties
+Crumbs:
+Description:In this step we tweak the settings of the Procedural Foliage Actors.
+Navigation:topic
+SkillLevel:Intermediate
+Version: 4.8
+
+[VAR:Steps]
+[OBJECT:Navigation]
+	[PARAM:previous]
+	[Previous Step](Engine\OpenWorldTools\ProceduralFoliage\QuickStart\2 "%Engine\OpenWorldTools\ProceduralFoliage\QuickStart\2:title%") 
+	[/PARAM]
+	[PARAM:current]
+	[/PARAM]
+	[PARAM:home]
+		[How-To Home](Engine\OpenWorldTools\ProceduralFoliage\QuickStart "Engine\OpenWorldTools\ProceduralFoliage\QuickStart:Title%")
+	[/PARAM]
+	[PARAM:next]
+	[Next Step](Engine\OpenWorldTools\ProceduralFoliage\QuickStart\4 "%Engine\OpenWorldTools\ProceduralFoliage\QuickStart\4:title%") 
+	[/PARAM]
+[/OBJECT]
+[/VAR]
+
+%Steps%
+In the following section we will take a look at the various options of the **Foliage Type Actor**, what those options do, and how you can use them to tweak how your Foliage Meshes are placed through the level.
+We will continue to work with the **PFT_00** level that was used in the last step.
+If you have closed down this level please re-open it inside of UE4 before going any further.
+
+1. First open up the **Tree_00** Foliage Type if not already open and expand the **Placement** section to make sure that both **Align to Normal** and **Random Yaw** are enabled.
+
+	![](T_Placement_Options.png)
+
+	[region:note]
+	The Placement section is where you can adjust how the Static Meshes used in the Foliage Type are placed on objects in the level. 
+	[/region]
+
+1. Next expanded the **Procedural** section of the Foliage Type and then expand the **Collision** section and set the **Shade Radius** to **50** then press the **Resimulate** button on the Procedural Foliage Spawner that was placed in the level. 
+
+	![](T_Shade_Radius.png)
+
+	[region:note]
+	The Collision section determines which Foliage Types should be removed when two Foliage Types are competing for the same spawn location or relative space. When a Virtual Seed's collision radius overlaps an existing collision or shade radius from another Foliage Type's Virtual Seed, the Virtual Seed will be replace or killed based which Foliage Type has the higher propriety.
+	[/region]
+
+1. Now collapse the **Collision** section and expand the **Clustering** section, then set the **Num Steps**  to **0** and press the **Resimulate** button. 
+
+	![](T_Num_Steps_0.png)
+
+	[region:note]
+	Clustering uses a number of properties such as, density, age, and proximity to help determine how the specified Static Meshes should placed and grouped in the level. 
+	[/region]
+
+1. While we now have some space in between our trees, the over all density is still a little to high. 
+	To fix this set the **Initial Seed Density** to **0.25** and then press the **Resimulate** button.
+
+	![](T_Num_ISD_0.125.png)
+
+1. As you can see, setting the **Initial Seed Density** to **0.25** greatly reduced the density of our forest, all most a little to much. 
+	To address this set the **Num Steps** back to **three** and then press the **Resimulate** button.
+
+	![](T_Num_Steps_3.png)
+
+1. Now collapse the **Clustering** section and expand the **Growth** section then 
+	set the **Max Age** to **20.0** and under **Procedural Scale** set the Max value to **10.0** then press the **Resimulate** button. 
+
+	![](T_Set_Growth.png)
+
+1. Next collapse the **Growth** section and expand the ** Instance Settings** options and then under the **Cull Distance** option set the **Max** value to **20,000**. 
+	
+	![](T_Cull_Distance.png)
+
+	[region:note]
+	The Instance Settings allows you to adjust how the Static Mesh used for the Foliage type will be displayed in the level. Inside the Instance Settings you can set or adjust many different properties like Cull Distance, Shadowing, and Collision. Which properties you adjust here depend on how you are using the Foliage in your level. 
+	[/region]
+
+1. Finally make sure to save the level and content by pressing the **Save** button on the **Main Tool Bar** and **Save All** button on the **Content Browser**.
+
+	![](Engine\OpenWorldTools\ProceduralFoliage\QuickStart\T_Uni_Save_Icons.png)
+
+
+
+
+
+
+
+
+
+
+
+
+
+
+
+
+
+
+
+
+
+
+
+
+
+
+
+
+
+
+
+
+
+
+
+
+
+
+
+
+
+
+
+
+%Steps%
\ No newline at end of file
diff --git a/Engine/Documentation/Source/Engine/OpenWorldTools/ProceduralFoliage/QuickStart/4/PFQS_4.INT.udn b/Engine/Documentation/Source/Engine/OpenWorldTools/ProceduralFoliage/QuickStart/4/PFQS_4.INT.udn
new file mode 100644
--- /dev/null
+++ b/Engine/Documentation/Source/Engine/OpenWorldTools/ProceduralFoliage/QuickStart/4/PFQS_4.INT.udn
@@ -0,0 +1,75 @@
+Availability:Public
+Title: 4. Using Multiple Foliage Types
+Crumbs:
+Description:In this step we will set up our Foliage Spawner to work with multiple Foliage Types.
+Navigation:topic
+SkillLevel:Intermediate
+Version: 4.8
+
+[VAR:Steps]
+[OBJECT:Navigation]
+	[PARAM:previous]
+	[Previous Step](Engine\OpenWorldTools\ProceduralFoliage\QuickStart\3 "%Engine\OpenWorldTools\ProceduralFoliage\QuickStart\3:title%") 
+	[/PARAM]
+	[PARAM:current]
+	[/PARAM]
+	[PARAM:home]
+		[How-To Home](Engine\OpenWorldTools\ProceduralFoliage\QuickStart "Engine\OpenWorldTools\ProceduralFoliage\QuickStart:Title%")
+	[/PARAM]
+	[PARAM:next]
+	[Next Step](Engine\OpenWorldTools\ProceduralFoliage\QuickStart\5 "%Engine\OpenWorldTools\ProceduralFoliage\QuickStart\5:title%") 
+	[/PARAM]
+[/OBJECT]
+[/VAR]
+
+%Steps%
+Having a single species of tree does not make our virtual forest look very interesting nor realistic. 
+Luckily for us the **Procedural Foliage Spawner** allows for you to spawn multiple **Foliage Types** resulting in one single **Procedural Foliage Spawner** being able to spawn an entire forest with a multitude of different trees. 
+In the following section we will take a look at how you can set up a Procedural Foliage Spawner to work with multiple Foliage Types.
+We will be continuing to work with the **PFT_00** level that was used in the last step.
+If you have closed down this level please re-open before going any further.
+
+1. First inside of the **Content Browser** select the **Tree_00** Foliage Type by clicking on it with the mouse then press **Ctrl + W** on the Keyboard to duplicate it and name it **Tree_01**.
+
+	![](T_Dup_FT.png)
+
+1. Then open the newly created **Tree_01** Foliage Type up and under the **Mesh** section change the mesh to the **ScotsPineTall_01** Static Mesh.
+
+	![](T_New_Mesh.png)
+
+1. Next open up the **Procedural Foliage Spawner** from the **Content Browser** and expand the **Foliage Types** section.
+
+	![](T_Expand_FT.png)
+
+1. Now click on the **Plus** sign icon to add the option to input another Foliage Type.
+
+	![](T_Add_New_FT.png)
+
+1. Next from the **Content Browser** select the **Tree_01** Foliage Type and then in the Procedural Foliage Spawner under the **Foliage Types** add the **Tree_01** Foliage Type to the empty **Foliage Type Object** by pressing the white **Arrow** icon.
+
+	![](T_Add_New_FT_Type.gif)(convert:false)
+
+1. Then select the **Procedural Foliage Spawner** that was placed in the level and press the **Resimulate** button and you should now have something that looks like this in your level.
+
+	![](T_2_FT_In_Level.png)
+
+1. While this does look very interesting the results could be better. Open up the **Tree_01** Foliage Type and adjust the following settings.
+	* **Num Steps:** 4
+	* **Initial Seed Density:** 0.125
+	* **Average Spread Distance:** 100
+	* **Can Grow in Shade:** Enabled
+	* **Spawns in Shade:** Enabled
+	* **Max Age:** 15
+	* **Overlap Priority:** 1
+	* **Procedural Scale:** Max 5.0
+
+1. Once the settings have been adjusted press the **Re-Simulate** button on the Procedural Foliage Spawner.
+
+	![](T_FT01_Adjust_Settings.png)
+
+1. Finally make sure to save all of you work by pressing both the **Save All** and **Save** buttons.
+
+	![](Engine\OpenWorldTools\ProceduralFoliage\QuickStart\T_Uni_Save_Icons.png)
+
+
+%Steps%
\ No newline at end of file
diff --git a/Engine/Documentation/Source/Engine/OpenWorldTools/ProceduralFoliage/QuickStart/5/PFQS_5.INT.udn b/Engine/Documentation/Source/Engine/OpenWorldTools/ProceduralFoliage/QuickStart/5/PFQS_5.INT.udn
new file mode 100644
--- /dev/null
+++ b/Engine/Documentation/Source/Engine/OpenWorldTools/ProceduralFoliage/QuickStart/5/PFQS_5.INT.udn
@@ -0,0 +1,65 @@
+Availability:Public
+Title: 5. Using the Procedural Foliage Blocking Volume
+Crumbs:
+Description:In this step we will set up a blocking volume to restrict where foliage can spawn.
+Navigation:topic
+SkillLevel:Intermediate
+Version: 4.8
+
+[VAR:Steps]
+[OBJECT:Navigation]
+	[PARAM:previous]
+	[Previous Step](Engine\OpenWorldTools\ProceduralFoliage\QuickStart\4 "%Engine\OpenWorldTools\ProceduralFoliage\QuickStart\4:title%") 
+	[/PARAM]
+	[PARAM:current]
+	[/PARAM]
+	[PARAM:home]
+		[How-To Home](Engine\OpenWorldTools\ProceduralFoliage\QuickStart "Engine\OpenWorldTools\ProceduralFoliage\QuickStart:Title%")
+	[/PARAM]
+	[PARAM:next]
+
+	[/PARAM]
+[/OBJECT]
+[/VAR]
+
+%Steps%
+The **Procedural Foliage Blocking Volume** is an Volume Actor that can be placed in a level and scaled to a desired size which will prevent the Procedural Foliage Spawner 
+form spawning any Foliage Actors inside of the bounds of the Procedural Foliage Blocking Volume.
+In the following section we will go over how you can add a **Procedural Foliage Blocking Volume** to your level and use it to prevent Foliage meshes from spawning.
+We will be continue to work with the **PFT_00** level that was used in the last step.
+If you have closed down this level please re-open it inside of UE4 before going any further.
+
+1. First find the Procedural Foliage Blocking Volume by searching for it in the **Modes** panel using **Proc** as the search term.  
+
+	![](T_Find_ProcF_Blocking_Volume.png)
+
+1. Once the Procedural Foliage Blocking Volume has been located in the Modes panel, select it by clicking on it with the mouse and then drag it from the Modes panel into the level.
+
+	![](T_Adding_PFV_To_Level.gif)(convert:false)
+
+1. Next input the following values for the Procedural Foliage Blocking Volume location and scale. When completed you should have something that looks similar to the image below.
+
+	* **Location X:** 5430.0 cm
+	* **Location Y:** -3900.0 cm
+	* **Location Z:** 200.0 cm
+	* **Scale    X:** 41.75
+	* **Scale    Y:** 65.5
+	* **Scale    X:** 41.75
+	
+	![](T_PFBV_Postion.png)
+
+1. Now select the **Procedural Foliage Spawner** in the level and then in the **Details** panel press the **Resimulate** button. 
+   When the Resimulate is completed you should now be missing the entire back section of trees that where intersecting with the Procedural Foliage Blocking Volume.
+   In the following image we can see the before and after results. 
+
+	   ![](T_PFV_Before_After.png)
+
+	   * 1: Before
+	   * 2: After
+
+1. Finally save your work and level by pressing the **Save All** button to save all content and the ** Save** button to save the level. 
+	
+	![](Engine\OpenWorldTools\ProceduralFoliage\QuickStart\T_Uni_Save_Icons.png)
+
+
+%Steps%
\ No newline at end of file
diff --git a/Engine/Documentation/Source/Engine/OpenWorldTools/ProceduralFoliage/QuickStart/ProceduralFoliageQuickStart.INT.udn b/Engine/Documentation/Source/Engine/OpenWorldTools/ProceduralFoliage/QuickStart/ProceduralFoliageQuickStart.INT.udn
new file mode 100644
--- /dev/null
+++ b/Engine/Documentation/Source/Engine/OpenWorldTools/ProceduralFoliage/QuickStart/ProceduralFoliageQuickStart.INT.udn
@@ -0,0 +1,51 @@
+Availability:Public
+Title: Procedural Foliage Tools Quick Start
+Crumbs:
+Description:How to set up and use the Procedural Foliage tools.
+Navigation:topic
+SkillLevel:Intermediate
+Version: 4.8
+
+
+## Overview
+In this Procedural Foliage Quick Start we will take a look at how the **Procedural Foliage Tools** that are available with the 4.8 release of Unreal Engine 4 (UE4) work. 
+Over the course of this Quick Start you will learn how to use the Procedural Foliage Tools in the manner they are intended to be used as well as how the Procedural Foliage Tools work in conjunction with other tools inside UE4.
+You will also be introduced to key tool properties that will help you achieve constant results that you are happy with.
+When you have completed this Quick Start you will have a new level that should look similar to the image below. 
+
+![](T_PF_QS_Final_Product.png)
+
+After completing this Quick Start you should gain a better understanding of what Actor types are required for the Procedural Foliage Tools and how to set them up to work correctly.
+You will also be introduced to any caveats the tools have and how to adjust workflows accordingly.
+
+
+[region:note]
+If you have not done so already, make sure to enable the Procedural Foliage tools in your project by first going to the Main Toolbar and then going to **Edit** > **Editor Preferences** > **Experimental**. 
+Once you are in the Experimental section enable the Procedural Foliage option by clicking on the checkmark box next to the word **Procedural Foliage**. 
+[/region]
+
+[OBJECT:TopicButtonList]
+	[PARAM:icon]
+		![](%ROOT%/build_icon.png)(convert:false)
+	[/PARAM]
+	[PARAM:title]
+		Steps
+	[/PARAM]
+	[PARAM:description]
+	[/PARAM]
+	[PARAM:links]
+	* [](Engine\OpenWorldTools\ProceduralFoliage\QuickStart\1)
+	* [](Engine\OpenWorldTools\ProceduralFoliage\QuickStart\2)
+	* [](Engine\OpenWorldTools\ProceduralFoliage\QuickStart\3)
+	* [](Engine\OpenWorldTools\ProceduralFoliage\QuickStart\4)
+	* [](Engine\OpenWorldTools\ProceduralFoliage\QuickStart\5)
+	[/PARAM]
+[/OBJECT]
+
+
+
+
+
+
+
+
