#UE4 Docs: Putting this as Docs so that it does not go out before going through QA.

diff --git a/Engine/Documentation/Source/Engine/OpenWorldTools/Grass/QuickStart/1/GTQS_1.INT.udn b/Engine/Documentation/Source/Engine/OpenWorldTools/Grass/QuickStart/1/GTQS_1.INT.udn
--- a/Engine/Documentation/Source/Engine/OpenWorldTools/Grass/QuickStart/1/GTQS_1.INT.udn
+++ b/Engine/Documentation/Source/Engine/OpenWorldTools/Grass/QuickStart/1/GTQS_1.INT.udn
@@ -1,4 +1,4 @@
-Availability:Public
+Availability:Docs
 Title:1. Grass Tool Initial Setup
 Crumbs:
 Description: How To set up the Initial Actors the Grass Tool needs to work.
diff --git a/Engine/Documentation/Source/Engine/OpenWorldTools/Grass/QuickStart/2/GTQS_2.INT.udn b/Engine/Documentation/Source/Engine/OpenWorldTools/Grass/QuickStart/2/GTQS_2.INT.udn
--- a/Engine/Documentation/Source/Engine/OpenWorldTools/Grass/QuickStart/2/GTQS_2.INT.udn
+++ b/Engine/Documentation/Source/Engine/OpenWorldTools/Grass/QuickStart/2/GTQS_2.INT.udn
@@ -1,4 +1,4 @@
-Availability:Public
+Availability:Docs
 Title:2. Grass Tool Actor Creation and Setup
 Crumbs:
 Description: How to create and setup the required Actors for the Grass Tool
diff --git a/Engine/Documentation/Source/Engine/OpenWorldTools/Grass/QuickStart/3/GTQS_3.INT.udn b/Engine/Documentation/Source/Engine/OpenWorldTools/Grass/QuickStart/3/GTQS_3.INT.udn
--- a/Engine/Documentation/Source/Engine/OpenWorldTools/Grass/QuickStart/3/GTQS_3.INT.udn
+++ b/Engine/Documentation/Source/Engine/OpenWorldTools/Grass/QuickStart/3/GTQS_3.INT.udn
@@ -1,4 +1,4 @@
-Availability:Public
+Availability:Docs
 Title:3. Grass Tool Material Setup
 Crumbs:
 Description: How to setup a Material that will work with the Grass system.
diff --git a/Engine/Documentation/Source/Engine/OpenWorldTools/Grass/QuickStart/4/GTQS_4.INT.udn b/Engine/Documentation/Source/Engine/OpenWorldTools/Grass/QuickStart/4/GTQS_4.INT.udn
--- a/Engine/Documentation/Source/Engine/OpenWorldTools/Grass/QuickStart/4/GTQS_4.INT.udn
+++ b/Engine/Documentation/Source/Engine/OpenWorldTools/Grass/QuickStart/4/GTQS_4.INT.udn
@@ -1,4 +1,4 @@
-Availability:Public
+Availability:Docs
 Title:4. Using the Grass Tools
 Crumbs:
 Description: How to use the Grass Tools.
@@ -60,7 +60,7 @@ We will continue to work with the **GT_00** level that was created in the last s
 	[region:note]
 	Depending on how powerful your development PC is when you start to paint on the Landscape you might experience some lag in the editor as the grass is spawned. 
 	This is normal behavior as the grass is dynamically spawned after you have finished painting. 
-	If you find that the lag makes it to hard to work, try turning down the **Grass Density** while painting and then turn it back up when done.
+	If you find that the lag makes it to hard to work, try turning down the **Grass Density** in the **Landscape Grass Type** settings and then turn it back up when done.
 	[/region]
 
 1. To remove grass from the Landscape Terrain select the **Rock Target Layer** then inside of the viewport hold down the **Left Mouse Button** 
diff --git a/Engine/Documentation/Source/Engine/OpenWorldTools/Overview/OpenWorldToolsOverview.INT.udn b/Engine/Documentation/Source/Engine/OpenWorldTools/Overview/OpenWorldToolsOverview.INT.udn
--- a/Engine/Documentation/Source/Engine/OpenWorldTools/Overview/OpenWorldToolsOverview.INT.udn
+++ b/Engine/Documentation/Source/Engine/OpenWorldTools/Overview/OpenWorldToolsOverview.INT.udn
@@ -1,4 +1,4 @@
-Availability:Public
+Availability:Docs
 Title: Open World Tools Overview
 Crumbs:
 Description: A system for procedurally placing Static Meshes in your levels.
diff --git a/Engine/Documentation/Source/Engine/OpenWorldTools/ProceduralFoliage/QuickStart/1/PFQS_1.INT.udn b/Engine/Documentation/Source/Engine/OpenWorldTools/ProceduralFoliage/QuickStart/1/PFQS_1.INT.udn
--- a/Engine/Documentation/Source/Engine/OpenWorldTools/ProceduralFoliage/QuickStart/1/PFQS_1.INT.udn
+++ b/Engine/Documentation/Source/Engine/OpenWorldTools/ProceduralFoliage/QuickStart/1/PFQS_1.INT.udn
@@ -1,4 +1,4 @@
-Availability:Public
+Availability:Docs
 Title: 1. Procedural Foliage Tool Actors Creation
 Crumbs:
 Description:In this step we create all of the required Actors.
diff --git a/Engine/Documentation/Source/Engine/OpenWorldTools/ProceduralFoliage/QuickStart/2/PFQS_2.INT.udn b/Engine/Documentation/Source/Engine/OpenWorldTools/ProceduralFoliage/QuickStart/2/PFQS_2.INT.udn
--- a/Engine/Documentation/Source/Engine/OpenWorldTools/ProceduralFoliage/QuickStart/2/PFQS_2.INT.udn
+++ b/Engine/Documentation/Source/Engine/OpenWorldTools/ProceduralFoliage/QuickStart/2/PFQS_2.INT.udn
@@ -1,4 +1,4 @@
-Availability:Public
+Availability:Docs
 Title: 2. Procedural Foliage Tool Actors Setup
 Crumbs:
 Description:In this step we setup all of the required Actors.
diff --git a/Engine/Documentation/Source/Engine/OpenWorldTools/ProceduralFoliage/QuickStart/3/PFQS_3.INT.udn b/Engine/Documentation/Source/Engine/OpenWorldTools/ProceduralFoliage/QuickStart/3/PFQS_3.INT.udn
--- a/Engine/Documentation/Source/Engine/OpenWorldTools/ProceduralFoliage/QuickStart/3/PFQS_3.INT.udn
+++ b/Engine/Documentation/Source/Engine/OpenWorldTools/ProceduralFoliage/QuickStart/3/PFQS_3.INT.udn
@@ -1,4 +1,4 @@
-Availability:Public
+Availability:Docs
 Title: 3. Procedural Foliage Tool tweaking Actor properties
 Crumbs:
 Description:In this step we tweak the settings of the Procedural Foliage Actors.
diff --git a/Engine/Documentation/Source/Engine/OpenWorldTools/ProceduralFoliage/QuickStart/4/PFQS_4.INT.udn b/Engine/Documentation/Source/Engine/OpenWorldTools/ProceduralFoliage/QuickStart/4/PFQS_4.INT.udn
--- a/Engine/Documentation/Source/Engine/OpenWorldTools/ProceduralFoliage/QuickStart/4/PFQS_4.INT.udn
+++ b/Engine/Documentation/Source/Engine/OpenWorldTools/ProceduralFoliage/QuickStart/4/PFQS_4.INT.udn
@@ -1,4 +1,4 @@
-Availability:Public
+Availability:Docs
 Title: 4. Using Multiple Foliage Types
 Crumbs:
 Description:In this step we will set up our Foliage Spawner to work with multiple Foliage Types.
diff --git a/Engine/Documentation/Source/Engine/OpenWorldTools/ProceduralFoliage/QuickStart/5/PFQS_5.INT.udn b/Engine/Documentation/Source/Engine/OpenWorldTools/ProceduralFoliage/QuickStart/5/PFQS_5.INT.udn
--- a/Engine/Documentation/Source/Engine/OpenWorldTools/ProceduralFoliage/QuickStart/5/PFQS_5.INT.udn
+++ b/Engine/Documentation/Source/Engine/OpenWorldTools/ProceduralFoliage/QuickStart/5/PFQS_5.INT.udn
@@ -1,4 +1,4 @@
-Availability:Public
+Availability:Docs
 Title: 5. Using the Procedural Foliage Blocking Volume
 Crumbs:
 Description:In this step we will set up a blocking volume to restrict where foliage can spawn.
diff --git a/Engine/Documentation/Source/Engine/OpenWorldTools/ProceduralFoliage/QuickStart/ProceduralFoliageQuickStart.INT.udn b/Engine/Documentation/Source/Engine/OpenWorldTools/ProceduralFoliage/QuickStart/ProceduralFoliageQuickStart.INT.udn
--- a/Engine/Documentation/Source/Engine/OpenWorldTools/ProceduralFoliage/QuickStart/ProceduralFoliageQuickStart.INT.udn
+++ b/Engine/Documentation/Source/Engine/OpenWorldTools/ProceduralFoliage/QuickStart/ProceduralFoliageQuickStart.INT.udn
@@ -1,4 +1,4 @@
-Availability:Public
+Availability:Docs
 Title: Procedural Foliage Tools Quick Start
 Crumbs:
 Description:How to set up and use the Procedural Foliage tools.
