"are" to "is" (singular subject) sp: throughout sp: learn

diff --git a/Engine/Documentation/Source/Engine/OpenWorldTools/Overview/OpenWorldToolsOverview.INT.udn b/Engine/Documentation/Source/Engine/OpenWorldTools/Overview/OpenWorldToolsOverview.INT.udn
--- a/Engine/Documentation/Source/Engine/OpenWorldTools/Overview/OpenWorldToolsOverview.INT.udn
+++ b/Engine/Documentation/Source/Engine/OpenWorldTools/Overview/OpenWorldToolsOverview.INT.udn
@@ -60,14 +60,14 @@ To learn more about how to use the Procedural Foliage Tool, check out the Proced
 #### Foliage Type Actors
 The Procedural Foliage Tools makes use of Actors called **Foliage Types**. 
 Foliage Types allow you to define various properties the Foliage should have like what Static Mesh to use or how big Static Mesh should be spawned.
-The chart below gives a very high level explanation of the four main sets of properties that the Foliage Type Actor are broken down into and what those options do.
+The chart below gives a very high level explanation of the four main sets of properties that the Foliage Type Actor is broken down into and what those options do.
 
 ![](T_Foliage_Type_Actor.png)
 
 | Property Name| Description |
 | ------------ | ----------- |
 |**Mesh**| The Mesh section is where you define which Static Mesh to use for the foliage.|
-|**Placement**| The Placement section is where you can adjust how the Static Mesh specified in the Mesh section are placed on objects through out the level.|
+|**Placement**| The Placement section is where you can adjust how the Static Mesh specified in the Mesh section are placed on objects throughout the level.|
 |**Growth**| The Growth section helps define how big the trees will be at the beginning of the simulation and how big they will "Grow" to by the end of the simulation.|
 |**Instance Settings**| The Instance Settings allows you to adjust how the Static Mesh used for the Foliage Type will be displayed in the level.|
 
@@ -78,7 +78,7 @@ Currently the Grass Tools will only work in conjunction with the Landscape Terra
 The Grass Tools are used to densely cover the Landscape Terrain in things Static Mesh to simulate things like grass, flowers, and small rocks. 
 You can use the Grass Tools to define a Static Mesh that will spawn when a certain Landscape Layer is painted on the Landscape Terrain. 
 This tool does not use any kind of simulation like the Procedural Foliage Spawner does, it simply fills the specified area at the specified density with the specified Static Mesh. 
-To lean more about how to use the Grass Tool, check out the Grass Tool Quick Start that is lined below.
+To learn more about how to use the Grass Tool, check out the Grass Tool Quick Start that is lined below.
 
 [REGION:topics third]
 [OBJECT:Topic]
