#UE4doc #docs: sp: "too"

diff --git a/Engine/Documentation/Source/Engine/OpenWorldTools/Grass/QuickStart/4/GTQS_4.INT.udn b/Engine/Documentation/Source/Engine/OpenWorldTools/Grass/QuickStart/4/GTQS_4.INT.udn
--- a/Engine/Documentation/Source/Engine/OpenWorldTools/Grass/QuickStart/4/GTQS_4.INT.udn
+++ b/Engine/Documentation/Source/Engine/OpenWorldTools/Grass/QuickStart/4/GTQS_4.INT.udn
@@ -60,7 +60,7 @@ We will continue to work with the **GT_00** level that was created in the last s
 	[region:note]
 	Depending on how powerful your development PC is when you start to paint on the Landscape you might experience some lag in the editor as the grass is spawned. 
 	This is normal behavior as the grass is dynamically spawned after you have finished painting. 
-	If you find that the lag makes it to hard to work, try turning down the **Grass Density** in the **Landscape Grass Type** settings and then turn it back up when done.
+	If you find that the lag makes it too hard to work, try turning down the **Grass Density** in the **Landscape Grass Type** settings and then turn it back up when done.
 	[/region]
 
 1. To remove grass from the Landscape Terrain select the **Rock Target Layer** then inside of the viewport hold down the **Left Mouse Button** 
