#UE4doc #docs: added word "to"

diff --git a/Engine/Documentation/Source/Engine/OpenWorldTools/ProceduralFoliage/QuickStart/2/PFQS_2.INT.udn b/Engine/Documentation/Source/Engine/OpenWorldTools/ProceduralFoliage/QuickStart/2/PFQS_2.INT.udn
--- a/Engine/Documentation/Source/Engine/OpenWorldTools/ProceduralFoliage/QuickStart/2/PFQS_2.INT.udn
+++ b/Engine/Documentation/Source/Engine/OpenWorldTools/ProceduralFoliage/QuickStart/2/PFQS_2.INT.udn
@@ -43,7 +43,7 @@ We will continue to work with the **PFT_00** level that was created in the last
 	![](T_PFT_Window.png)
 
 1. At the top of the Tree_00 Foliage Type locate the **Mesh** section and then click on the drop down menu that says **None** with the **Left Mouse Button**. 
-	Use the search menu locate  **HillTree_02** Static Mesh from the **Open World Demo Collection** by inputting **HillTree_02** as the search term. 
+	Use the search menu to locate  **HillTree_02** Static Mesh from the **Open World Demo Collection** by inputting **HillTree_02** as the search term. 
 	Once the **HillTree_02** is located, load it by clicking on it with the **Left Mouse Button**. 
 
 	![](T_Set_Tree_Mesh.gif)(convert:false)
