#UE4doc #docs: "placed" to "be placed" sp: "overall" "too"

diff --git a/Engine/Documentation/Source/Engine/OpenWorldTools/ProceduralFoliage/QuickStart/3/PFQS_3.INT.udn b/Engine/Documentation/Source/Engine/OpenWorldTools/ProceduralFoliage/QuickStart/3/PFQS_3.INT.udn
--- a/Engine/Documentation/Source/Engine/OpenWorldTools/ProceduralFoliage/QuickStart/3/PFQS_3.INT.udn
+++ b/Engine/Documentation/Source/Engine/OpenWorldTools/ProceduralFoliage/QuickStart/3/PFQS_3.INT.udn
@@ -48,15 +48,15 @@ If you have closed down this level please re-open it inside of UE4 before going
 	![](T_Num_Steps_0.png)
 
 	[region:note]
-	Clustering uses a number of properties such as, density, age, and proximity to help determine how the specified Static Meshes should placed and grouped in the level. 
+	Clustering uses a number of properties such as, density, age, and proximity to help determine how the specified Static Meshes should be placed and grouped in the level. 
 	[/region]
 
-1. While we now have some space in between our trees, the over all density is still a little to high. 
+1. While we now have some space in between our trees, the overall density is still a little too high. 
 	To fix this set the **Initial Seed Density** to **0.25** and then press the **Resimulate** button.
 
 	![](T_Num_ISD_0.125.png)
 
-1. As you can see, setting the **Initial Seed Density** to **0.25** greatly reduced the density of our forest, all most a little to much. 
+1. As you can see, setting the **Initial Seed Density** to **0.25** greatly reduced the density of our forest, all most a little too much. 
 	To address this set the **Num Steps** back to **three** and then press the **Resimulate** button.
 
 	![](T_Num_Steps_3.png)
@@ -74,7 +74,7 @@ If you have closed down this level please re-open it inside of UE4 before going
 	The Instance Settings allows you to adjust how the Static Mesh used for the Foliage type will be displayed in the level. Inside the Instance Settings you can set or adjust many different properties like Cull Distance, Shadowing, and Collision. Which properties you adjust here depend on how you are using the Foliage in your level. 
 	[/region]
 
-1. Finally make sure to save the level and content by pressing the **Save** button on the **Main Tool Bar** and **Save All** button on the **Content Browser**.
+1. Finally, make sure to save the level and content by pressing the **Save** button on the **Main Tool Bar** and **Save All** button on the **Content Browser**.
 
 	![](Engine\OpenWorldTools\ProceduralFoliage\QuickStart\T_Uni_Save_Icons.png)
 
