#UE4doc #docs: "an" to "a" sp: continuing

diff --git a/Engine/Documentation/Source/Engine/OpenWorldTools/ProceduralFoliage/QuickStart/5/PFQS_5.INT.udn b/Engine/Documentation/Source/Engine/OpenWorldTools/ProceduralFoliage/QuickStart/5/PFQS_5.INT.udn
--- a/Engine/Documentation/Source/Engine/OpenWorldTools/ProceduralFoliage/QuickStart/5/PFQS_5.INT.udn
+++ b/Engine/Documentation/Source/Engine/OpenWorldTools/ProceduralFoliage/QuickStart/5/PFQS_5.INT.udn
@@ -23,10 +23,10 @@ Version: 4.8
 [/VAR]
 
 %Steps%
-The **Procedural Foliage Blocking Volume** is an Volume Actor that can be placed in a level and scaled to a desired size which will prevent the Procedural Foliage Spawner 
+The **Procedural Foliage Blocking Volume** is a Volume Actor that can be placed in a level and scaled to a desired size which will prevent the Procedural Foliage Spawner 
 form spawning any Foliage Actors inside of the bounds of the Procedural Foliage Blocking Volume.
 In the following section we will go over how you can add a **Procedural Foliage Blocking Volume** to your level and use it to prevent Foliage meshes from spawning.
-We will be continue to work with the **PFT_00** level that was used in the last step.
+We will be continuing to work with the **PFT_00** level that was used in the last step.
 If you have closed down this level please re-open it inside of UE4 before going any further.
 
 1. First find the Procedural Foliage Blocking Volume by searching for it in the **Modes** panel using **Proc** as the search term.  
