Removed unnecessary code from VTE tutorial.

diff --git a/Engine/Documentation/Source/Programming/Tutorials/VariablesTimersEvents/1/VariablesTimersEvents_1.INT.udn b/Engine/Documentation/Source/Programming/Tutorials/VariablesTimersEvents/1/VariablesTimersEvents_1.INT.udn
--- a/Engine/Documentation/Source/Programming/Tutorials/VariablesTimersEvents/1/VariablesTimersEvents_1.INT.udn
+++ b/Engine/Documentation/Source/Programming/Tutorials/VariablesTimersEvents/1/VariablesTimersEvents_1.INT.udn
@@ -43,7 +43,6 @@ Description:Starting and stopping a timer that runs code at a fixed rate.
 		CountdownText = CreateDefaultSubobject<UTextRenderComponent>(TEXT("CountdownNumber"));
 		CountdownText->SetHorizontalAlignment(EHTA_Center);
 		CountdownText->SetWorldSize(150.0f);
-		CountdownText->AttachTo(RootComponent);
 		RootComponent = CountdownText;
 
 		CountdownTime = 3;
@@ -168,7 +167,6 @@ At this point, we've already created a simple class that uses a timer. Non-progr
 				CountdownText = CreateDefaultSubobject<UTextRenderComponent>(TEXT("CountdownNumber"));
 				CountdownText->SetHorizontalAlignment(EHTA_Center);
 				CountdownText->SetWorldSize(150.0f);
-				CountdownText->AttachTo(RootComponent);
 				RootComponent = CountdownText;
 
 				CountdownTime = 3;
diff --git a/Engine/Documentation/Source/Programming/Tutorials/VariablesTimersEvents/2/VariablesTimersEvents_2.INT.udn b/Engine/Documentation/Source/Programming/Tutorials/VariablesTimersEvents/2/VariablesTimersEvents_2.INT.udn
--- a/Engine/Documentation/Source/Programming/Tutorials/VariablesTimersEvents/2/VariablesTimersEvents_2.INT.udn
+++ b/Engine/Documentation/Source/Programming/Tutorials/VariablesTimersEvents/2/VariablesTimersEvents_2.INT.udn
@@ -150,7 +150,6 @@ We have now made a variable and a function accessible to, and alterable by, non-
 				CountdownText = CreateDefaultSubobject<UTextRenderComponent>(TEXT("CountdownNumber"));
 				CountdownText->SetHorizontalAlignment(EHTA_Center);
 				CountdownText->SetWorldSize(150.0f);
-				CountdownText->AttachTo(RootComponent);
 				RootComponent = CountdownText;
 
 				CountdownTime = 3;
diff --git a/Engine/Documentation/Source/Programming/Tutorials/VariablesTimersEvents/3/VariablesTimersEvents_3.INT.udn b/Engine/Documentation/Source/Programming/Tutorials/VariablesTimersEvents/3/VariablesTimersEvents_3.INT.udn
--- a/Engine/Documentation/Source/Programming/Tutorials/VariablesTimersEvents/3/VariablesTimersEvents_3.INT.udn
+++ b/Engine/Documentation/Source/Programming/Tutorials/VariablesTimersEvents/3/VariablesTimersEvents_3.INT.udn
@@ -145,7 +145,6 @@ Description:Using Blueprints to supplement or replace C++ functionality.
 				CountdownText = CreateDefaultSubobject<UTextRenderComponent>(TEXT("CountdownNumber"));
 				CountdownText->SetHorizontalAlignment(EHTA_Center);
 				CountdownText->SetWorldSize(150.0f);
-				CountdownText->AttachTo(RootComponent);
 				RootComponent = CountdownText;
 
 				CountdownTime = 3;
