TypoCorrection

diff --git a/Engine/Documentation/Source/Gameplay/HowTo/Networking/ReplicateVariable/Blueprints/ReplicateVariable_BP.INT.udn b/Engine/Documentation/Source/Gameplay/HowTo/Networking/ReplicateVariable/Blueprints/ReplicateVariable_BP.INT.udn
--- a/Engine/Documentation/Source/Gameplay/HowTo/Networking/ReplicateVariable/Blueprints/ReplicateVariable_BP.INT.udn
+++ b/Engine/Documentation/Source/Gameplay/HowTo/Networking/ReplicateVariable/Blueprints/ReplicateVariable_BP.INT.udn
@@ -1,4 +1,4 @@
-Availability: Public
+﻿Availability: Public
 Title:How to Replicate Variables in Blueprints
 Crumbs: %ROOT%, Gameplay, Gameplay\HowTo\Networking\ReplicateVariable
 Description:Guide to using Replicated and RepNotify Variables in Blueprints.
@@ -26,7 +26,7 @@ In this example, we will create a **Replicated** variable that is used for a cou
 
 	![](HT_0.png)
 
-	This will ensure that this Blueprint it replicated across the network to all connect machines. 
+	This will ensure that this Blueprint is replicated across the network to all connect machines. 
 
 1.	In the **Components** window, click the **Add Component** button and search for and add a **Text Render** component. 
 
