TypoCorrection

diff --git a/Engine/Documentation/Source/Engine/OpenWorldTools/ProceduralFoliage/QuickStart/5/PFQS_5.INT.udn b/Engine/Documentation/Source/Engine/OpenWorldTools/ProceduralFoliage/QuickStart/5/PFQS_5.INT.udn
--- a/Engine/Documentation/Source/Engine/OpenWorldTools/ProceduralFoliage/QuickStart/5/PFQS_5.INT.udn
+++ b/Engine/Documentation/Source/Engine/OpenWorldTools/ProceduralFoliage/QuickStart/5/PFQS_5.INT.udn
@@ -1,4 +1,4 @@
-Availability:Docs
+﻿Availability:Docs
 Title: 5. Using the Procedural Foliage Blocking Volume
 Crumbs:
 Description:In this step we will set up a blocking volume to restrict where foliage can spawn.
@@ -24,7 +24,7 @@ Version: 4.8
 
 %Steps%
 The **Procedural Foliage Blocking Volume** is a Volume Actor that can be placed in a level and scaled to a desired size which will prevent the Procedural Foliage Spawner 
-form spawning any Foliage Actors inside of the bounds of the Procedural Foliage Blocking Volume.
+from spawning any Foliage Actors inside of the bounds of the Procedural Foliage Blocking Volume.
 In the following section we will go over how you can add a **Procedural Foliage Blocking Volume** to your level and use it to prevent Foliage meshes from spawning.
 We will be continuing to work with the **PFT_00** level that was used in the last step.
 If you have closed down this level please re-open it inside of UE4 before going any further.
