Changed all instance of "[RocketInstall]" to "[LauncherInstall]\\[VersionNumber]" and all instances of "RocketEditor.exe/app" to "UE4Editor.exe/app"

diff --git a/Engine/Documentation/Source/GettingStarted/RunningUnrealEngine/RunningUnrealEngine.INT.udn b/Engine/Documentation/Source/GettingStarted/RunningUnrealEngine/RunningUnrealEngine.INT.udn
--- a/Engine/Documentation/Source/GettingStarted/RunningUnrealEngine/RunningUnrealEngine.INT.udn
+++ b/Engine/Documentation/Source/GettingStarted/RunningUnrealEngine/RunningUnrealEngine.INT.udn
@@ -62,11 +62,11 @@ must first compile your modules using the **Debug** configuration so they exist
 	active_button_content
 	[/PARAMLITERAL]
 	[PARAM:content]
-	1.	From a command prompt, navigate to your `[RocketInstall]\Engine\Binaries\Win64` directory.
+	1.	From a command prompt, navigate to your `[LauncherInstall]\[VersionNumber]\Engine\Binaries\Win64` directory.
 
-	1.	Run the `RocketEditor.exe` passing it the path to your project:
+	1.	Run the `UE4Editor.exe` passing it the path to your project:
 
-			RocketEditor.exe "[ProjectPath]\[ProjectName].uproject"
+			UE4Editor.exe "[ProjectPath]\[ProjectName].uproject"
 	[/PARAM]
 [/OBJECT]
 [OBJECT:ToggleButtonContent]
@@ -79,11 +79,11 @@ must first compile your modules using the **Debug** configuration so they exist
 	[PARAMLITERAL:active]
 	[/PARAMLITERAL]
 	[PARAM:content]
-	1.	From a command prompt, navigate to your `[RocketInstall]/Engine/Binaries/Mac` directory.
+	1.	From a command prompt, navigate to your `[LauncherInstall]/[VersionNumber]/Engine/Binaries/Mac` directory.
 
-	1.	Run the `RocketEditor.app` passing it the path to your project:
+	1.	Run the `UE4Editor.app` passing it the path to your project:
 
-			open RocketEditor.app --args "[ProjectPath]/[ProjectName].uproject"
+			open UE4Editor.app --args "[ProjectPath]/[ProjectName].uproject"
 	[/PARAM]
 [/OBJECT]
 [/PUBLISH:Rocket]
@@ -128,9 +128,9 @@ must first compile your modules using the **Debug** configuration so they exist
 	active_button_content
 	[/PARAMLITERAL]
 	[PARAM:content]
-	1.	Navigate to your `[RocketInstall]\Engine\Binaries\Win64` directory.
+	1.	Navigate to your `[LauncherInstall]\[VersionNumber]\Engine\Binaries\Win64` directory.
 
-	1.	**Right-click** on the `RocketEditor.exe` executable and choose **Create shortcut**.
+	1.	**Right-click** on the `UE4Editor.exe` executable and choose **Create shortcut**.
 
 	1.	Rename the shortcut to reflect the game it will run and the fact that it will launch the editor, i.e. **MyProject - Editor.exe**.
 
@@ -138,7 +138,7 @@ must first compile your modules using the **Debug** configuration so they exist
 
 	1.	Add the name of the game to run as an argument at the end of the **Target** property:
 
-			[RocketInstall]\Engine\Binaries\Win64\RocketEditor.exe "[ProjectPath]\[ProjectName].uproject"
+			[LauncherInstall]\[VersionNumber]\Engine\Binaries\Win64\UE4Editor.exe "[ProjectPath]\[ProjectName].uproject"
 
 	1. Press **OK** to save the changes.
 
@@ -213,11 +213,11 @@ When running from the command line, you must pass the name of the project you wa
 	active_button_content
 	[/PARAMLITERAL]
 	[PARAM:content]
-	1.	From a command prompt, navigate to your `[RocketInstall]\Engine\Binaries\Win64` directory.
+	1.	From a command prompt, navigate to your `[LauncherInstall]\[VersionNumber]\Engine\Binaries\Win64` directory.
 
-	1.	Run the **RocketEditor.exe** passing it the path to the project to run, along with the `-game` parameter.
+	1.	Run the **UE4Editor.exe** passing it the path to the project to run, along with the `-game` parameter.
 
-			RocketEditor.exe "[ProjectPath]\[ProjectName].uproject" -game
+			UE4Editor.exe "[ProjectPath]\[ProjectName].uproject" -game
 	[/PARAM]
 [/OBJECT]
 [OBJECT:ToggleButtonContent]
@@ -230,11 +230,11 @@ When running from the command line, you must pass the name of the project you wa
 	[PARAMLITERAL:active]
 	[/PARAMLITERAL]
 	[PARAM:content]
-	1.	From a command prompt, navigate to your `[RocketInstall]/Engine/Binaries/Mac` directory.
+	1.	From a command prompt, navigate to your `[LauncherInstall]\[VersionNumber]/Engine/Binaries/Mac` directory.
 
-	1.	Run the **RocketEditor.exe** passing it the path to the project to run, along with the `-game` parameter
+	1.	Run the **UE4Editor.app** passing it the path to the project to run, along with the `-game` parameter
 
-			open RocketEditor.app --args "[ProjectPath]/[ProjectName].uproject" -game
+			open UE4Editor.app --args "[ProjectPath]/[ProjectName].uproject" -game
 	[/PARAM]
 [/OBJECT]
 [/PUBLISH:Rocket]
@@ -284,9 +284,9 @@ When running from the executable, you must specify the game you want to run as a
 	When running from the executable, you must specify the path to the project you want to run along with the `-game` switch as arguments to the 
 	executable via the **Target** property of a shortcut.
 
-	1.	Navigate to your `[RocketInstall]\Engine\Binaries\Win64` directory.
+	1.	Navigate to your `[LauncherInstall]\[VersionNumber]\Engine\Binaries\Win64` directory.
 
-	1.	**Right-click** on the **RocketEditor.exe** executable and choose **Create shortcut**.
+	1.	**Right-click** on the **UE4Editor.exe** executable and choose **Create shortcut**.
 
 	1.	Rename the shortcut to reflect the game it will run, i.e. **MyProject.exe**.
 
@@ -294,7 +294,7 @@ When running from the executable, you must specify the game you want to run as a
 
 	1.	Add the full path to the project to run as an argument at the end of the **Target** property, and specify the `-game` argument to run as a game:
 
-			[RocketInstall]\Engine\Binaries\Win64\RocketEditor.exe "[ProjectPath]\[ProjectName].uproject" -game
+			[LauncherInstall]\[VersionNumber]\Engine\Binaries\Win64\UE4Editor.exe "[ProjectPath]\[ProjectName].uproject" -game
 
 		Press **OK** to save the changes.
 
@@ -396,7 +396,7 @@ used to load the engine running the `ExampleMap` map:
 	[/PARAMLITERAL]
 	[PARAM:content]
 	~~~
-	RocketEditor.exe "[ProjectPath]\[ProjectName].uproject" ExampleMap -game
+	UE4Editor.exe "[ProjectPath]\[ProjectName].uproject" ExampleMap -game
 	~~~
 	[/PARAM]
 [/OBJECT]
@@ -411,7 +411,7 @@ used to load the engine running the `ExampleMap` map:
 	[/PARAMLITERAL]
 	[PARAM:content]
 	~~~
-	open RocketEditor.app --args "[ProjectPath]/[ProjectName].uproject" ExampleMap -game
+	open UE4Editor.app --args "[ProjectPath]/[ProjectName].uproject" ExampleMap -game
 	~~~
 	[/PARAM]
 [/OBJECT]
@@ -442,7 +442,7 @@ following command line can be used:
 	[/PARAMLITERAL]
 	[PARAM:content]
 	~~~
-	RocketEditor.exe "[ProjectPath]\[ProjectName].uproject" ExampleMap
+	UE4Editor.exe "[ProjectPath]\[ProjectName].uproject" ExampleMap
 	~~~
 	[/PARAM]
 [/OBJECT]
@@ -457,7 +457,7 @@ following command line can be used:
 	[/PARAMLITERAL]
 	[PARAM:content]
 	~~~
-	open RocketEditor.app --args "[ProjectPath]/[ProjectName].uproject" ExampleMap
+	open UE4Editor.app --args "[ProjectPath]/[ProjectName].uproject" ExampleMap
 	~~~
 	[/PARAM]
 [/OBJECT]
