#UE4doc #docs: sp: occurred, unambiguously

diff --git a/Engine/Documentation/Source/GettingStarted/DownloadingUnrealEngine/Perforce/AutoMerge/AutoMerge.INT.udn b/Engine/Documentation/Source/GettingStarted/DownloadingUnrealEngine/Perforce/AutoMerge/AutoMerge.INT.udn
--- a/Engine/Documentation/Source/GettingStarted/DownloadingUnrealEngine/Perforce/AutoMerge/AutoMerge.INT.udn
+++ b/Engine/Documentation/Source/GettingStarted/DownloadingUnrealEngine/Perforce/AutoMerge/AutoMerge.INT.udn
@@ -40,7 +40,7 @@ Finally, when you are ready to check in, the description field will already be f
 The description contains:
 
 
-* [AUTOMERGE] or [SAFEMERGE] to indicate how the merge occured.
+* [AUTOMERGE] or [SAFEMERGE] to indicate how the merge occurred.
 * A copy of the original changelist description.
 * A footer that indicates the originating branch, change#, submitting user, and date.
 
@@ -55,7 +55,7 @@ See the [Perforce Integration](GettingStarted/DownloadingUnrealEngine/Perforce/I
 
 ### NULL Merge
 
-The _NULL Merge_ tool merges a changelist and performs an _accept yours_ resolve, which means that the files will be integrated, but will not be changed. This will unambiguosly indicate that a changelist has been considered for integration, but it was decided not to incorporate those changes.
+The _NULL Merge_ tool merges a changelist and performs an _accept yours_ resolve, which means that the files will be integrated, but will not be changed. This will unambiguously indicate that a changelist has been considered for integration, but it was decided not to incorporate those changes.
 
 For more information on NULL merges, see the [Perforce Integration](GettingStarted/DownloadingUnrealEngine/Perforce/Integration#NullMerge) page.
 
