#UE4doc #docs: sp: integration, Perforce, highly, changelist, submitting

diff --git a/Engine/Documentation/Source/GettingStarted/DownloadingUnrealEngine/Perforce/Integration/Integration.INT.udn b/Engine/Documentation/Source/GettingStarted/DownloadingUnrealEngine/Perforce/Integration/Integration.INT.udn
--- a/Engine/Documentation/Source/GettingStarted/DownloadingUnrealEngine/Perforce/Integration/Integration.INT.udn
+++ b/Engine/Documentation/Source/GettingStarted/DownloadingUnrealEngine/Perforce/Integration/Integration.INT.udn
@@ -20,26 +20,26 @@ _[Note: To help simplify the merge process, an auto-merge tool has been created
 
 ### What does integrate, merge, and branch mean in Perforce?
 
-They all mean the same thing and are accomplished using the _Integrate_ command. A _Branch_ is usually referred to as an integrate that creates a new file in a new location. A _Merge_ is an integrate between two previously-branched files.
+They all mean the same thing and are accomplished using the _Integrate_ command. A _Branch_ is usually referred to as an integration that creates a new file in a new location. A _Merge_ is an integrate between two previously-branched files.
 
 
 ### Why do all these complicated steps when I can just copy the file over and change it?
 
-Perforce keeps a searchable, graphable database of all changes and integrations that have ever taken place. Merging "by hand" severs the database tracking records, makes it impossible to track code history, makes the branches harder to maintain, and makes it more difficult for you and Peforce when an integrate is done correctly. Symptoms could be anything from Perforce refusing to integrate files to ridiculously difficult resolves for simple changes.
+Perforce keeps a searchable, graphable database of all changes and integrations that have ever taken place. Merging "by hand" severs the database tracking records, makes it impossible to track code history, makes the branches harder to maintain, and makes it more difficult for you and Perforce when an integrate is done correctly. Symptoms could be anything from Perforce refusing to integrate files to ridiculously difficult resolves for simple changes.
 
 (#IntegrateSection)
 
 
 ### How do I integrate a single changelist?
 
-_Note: It is hightly recommended to use P4V when integrating, as there have been many improvements to the UI to visualize the process. Everything ultimately resolves to the same underlying P4API calls, so both P4V and P4Win support the same feature set._
+_Note: It is highly recommended to use P4V when integrating, as there have been many improvements to the UI to visualize the process. Everything ultimately resolves to the same underlying P4API calls, so both P4V and P4Win support the same feature set._
 
 To integrate a single change in P4V:
 
 
 * Make sure your active workspace contains the branch you are integrating to.
     * There is no need to include the source branch in your workspace. 
-* **Right-click** on the changelist you want to integrate and select _Integrate using submitted changeslist 'change#'..._. You should see a dialog similar to this:
+* **Right-click** on the changelist you want to integrate and select _Integrate using submitted changelist 'change#'..._. You should see a dialog similar to this:
 
 
 (#IntegrationShot)
@@ -71,7 +71,7 @@ To integrate a single change in P4V:
 * Select _New_ changelist to ensure you do not mix existing changes with this merge.
 * click _Finish_ to create the new changelist. 
 * Resolve the integrated files by following the [instructions below](#Resolve). 
-* Test everything thoroughly before sumbitting.
+* Test everything thoroughly before submitting.
 * Remember to use the [changelist description guidelines](#ChangeDesc) mentioned later in this guide.
 
 
@@ -88,7 +88,7 @@ To integrate a single change in P4Win:
 * Select _New_ changelist to ensure you do not mix existing changes with this merge 
 * Click _Finish_ to create the new changelist. 
 * Resolve the integrated files by following the [instructions below](#Resolve). 
-* Test everything thoroughly before sumbitting.
+* Test everything thoroughly before submitting.
 * Remember to use the [changelist description guidelines](#ChangeDesc) mentioned later in this guide.
 
 
