#UE4doc #docs: added word "you"

diff --git a/Engine/Documentation/Source/GettingStarted/DownloadingUnrealEngine/Perforce/CommonMergeProblems/CommonMergeProblems.INT.udn b/Engine/Documentation/Source/GettingStarted/DownloadingUnrealEngine/Perforce/CommonMergeProblems/CommonMergeProblems.INT.udn
--- a/Engine/Documentation/Source/GettingStarted/DownloadingUnrealEngine/Perforce/CommonMergeProblems/CommonMergeProblems.INT.udn
+++ b/Engine/Documentation/Source/GettingStarted/DownloadingUnrealEngine/Perforce/CommonMergeProblems/CommonMergeProblems.INT.udn
@@ -22,7 +22,7 @@ Here are the tools recommended for merging and integrating Unreal Engine 4 QA-Ap
 
 A tool such as Araxis Merge can be very useful for handling merges. It can be integrated into version control systems such as Perforce, so you can have diff/merge handled by Araxis instead of the built-in tools.
 
-For merges of code drops, the folder comparisons (both two-way and three-way) in Araxis Merge help in comparing a current project with a recent QA-Approved Build; and in some cases, a previous QA-Approved Build. This strategy allows to see what has evolved in a QA-Approved Build, as well as in your current project.
+For merges of code drops, the folder comparisons (both two-way and three-way) in Araxis Merge help in comparing a current project with a recent QA-Approved Build; and in some cases, a previous QA-Approved Build. This strategy allows you to see what has evolved in a QA-Approved Build, as well as in your current project.
 
 Branching in a version control system such as Perforce when merging can be helpful, as well. Working with all merge-related changes in a separate branch keeps the base branch clean until you are ready to integrate the fully-merged changes over - usually after some regression testing and all issues resolved.
 
