#UE4doc #docs: CL: added "Prerequesite" information

diff --git a/Engine/Documentation/Source/Engine/OpenWorldTools/ProceduralFoliage/QuickStart/ProceduralFoliageQuickStart.INT.udn b/Engine/Documentation/Source/Engine/OpenWorldTools/ProceduralFoliage/QuickStart/ProceduralFoliageQuickStart.INT.udn
--- a/Engine/Documentation/Source/Engine/OpenWorldTools/ProceduralFoliage/QuickStart/ProceduralFoliageQuickStart.INT.udn
+++ b/Engine/Documentation/Source/Engine/OpenWorldTools/ProceduralFoliage/QuickStart/ProceduralFoliageQuickStart.INT.udn
@@ -18,11 +18,14 @@ When you have completed this Quick Start you will have a new level that should l
 After completing this Quick Start you should gain a better understanding of what Actor types are required for the Procedural Foliage Tools and how to set them up to work correctly.
 You will also be introduced to any caveats the tools have and how to adjust workflows accordingly.
 
+## Prerequisites
+Before you can use the Procedural Foliage tools in your project, you must first enable them by going to the **Main Toolbar** and then going to **Edit** > **Editor Preferences** > **Experimental**. 
+Once you are in the Experimental section enable Procedural Foliage option by clicking on the checkmark box next to the word **Procedural Foliage**.
+You will also need to download the **Open World Demo Collection** content pack from the Unreal Engine 4 Launcher as some of the content from the collection will be used in the Open World Tools Quick Starts.
+
+Once the Open World Demo Collection is downloaded add it to the project that you are using to follow along with this Quick Start by first locating the Open World Demo Collection in the UE4 Launcher **Vault** section and then click on the button that says **Add to Project**. You will then be shown a list of projects that you can add this to. Select the project that you are using to follow along with this Quick Start and then press the **Add to Project** button to add the content. When you re-open the project you will now see a folder called **Kite Demo** and inside this folder you will find the content you need.
+
 
-[region:note]
-If you have not done so already, make sure to enable the Procedural Foliage tools in your project by first going to the Main Toolbar and then going to **Edit** > **Editor Preferences** > **Experimental**. 
-Once you are in the Experimental section enable the Procedural Foliage option by clicking on the checkmark box next to the word **Procedural Foliage**. 
-[/region]
 
 [OBJECT:TopicButtonList]
 	[PARAM:icon]
