#UE4doc #docs: CL: added "Prerequesite" information

diff --git a/Engine/Documentation/Source/Engine/OpenWorldTools/Overview/OpenWorldToolsOverview.INT.udn b/Engine/Documentation/Source/Engine/OpenWorldTools/Overview/OpenWorldToolsOverview.INT.udn
--- a/Engine/Documentation/Source/Engine/OpenWorldTools/Overview/OpenWorldToolsOverview.INT.udn
+++ b/Engine/Documentation/Source/Engine/OpenWorldTools/Overview/OpenWorldToolsOverview.INT.udn
@@ -21,6 +21,12 @@ Once you are in the Experimental section enable Procedural Foliage option by cli
 
 ![](T_Enable_PFT.gif)(convert:false)
 
+You will also need to download the **Open World Demo Collection** content pack from the Unreal Engine 4 Launcher as some of the content from the collection will be used in the Open World Tools Quick Starts.
+
+Once the Open World Demo Collection is downloaded add it to the project that you are using to follow along with this Quick Start by first locating the Open World Demo Collection in the UE4 Launcher **Vault** section and then click on the button that says **Add to Project**. You will then be shown a list of projects that you can add this to. Select the project that you are using to follow along with this Quick Start and then press the **Add to Project** button to add the content. When you re-open the project you will now see a folder called **Kite Demo** and inside this folder you will find the content you need.
+
+
+
 ## Types of Procedural Tools 
 The Procedural Tools that come with the 4.8 release of Unreal Engine 4 allow you to place Static Meshes inside your levels using a computer algorithm to determine placement instead having to do this work manually. 
 These tools are great if you want to quickly generate a natural looking forest comprised of many different species of trees, or have a certain species of bush spawn whenever a certain Texture is painted on the Landscape Terrain.
