No changes rechecking in the file.

diff --git a/Engine/Documentation/Source/Engine/OpenWorldTools/Grass/QuickStart/GrassToolsQuickStart.INT.udn b/Engine/Documentation/Source/Engine/OpenWorldTools/Grass/QuickStart/GrassToolsQuickStart.INT.udn
--- a/Engine/Documentation/Source/Engine/OpenWorldTools/Grass/QuickStart/GrassToolsQuickStart.INT.udn
+++ b/Engine/Documentation/Source/Engine/OpenWorldTools/Grass/QuickStart/GrassToolsQuickStart.INT.udn
@@ -8,8 +8,8 @@ Version: 4.8
 
 ## Overview
 
-While the Procedural Foliage tools are great a simulating a huge forest that has many different species of trees and bushes in it, the one thing that it can not do very well is spawn the thousands and thousands of Static Meshes that would be needed to make the Landscape look like it was densely covered in grass or rocks. 
-To accomplish this a new tool was introduced to UE4, the **Grass Tool**. 
+While the Procedural Foliage tools are great a simulating a huge forest that has many different species of trees and bushes in it, the one thing that it cannot do very well is spawn the thousands and thousands of Static Meshes that would be needed to make the Landscape look like it was densely covered in grass or rocks. 
+To accomplish this, a new tool was introduced to UE4, the **Grass Tool**. 
 The Grass Tool allows you to define a Static Mesh and amount to spawn on the Landscape Terrain whenever a certain Texture is painted on the Landscape. 
 In the following How - To we will cover all you need to know about using the Grass Tool in your levels.
 
@@ -19,6 +19,14 @@ In the following How - To we will cover all you need to know about using the Gra
 Currently the Grass system will only work with the Landscape Terrain actor. Trying to use the Grass system to spawn grass on any other actor type will not work.
 [/region]
 
+## Prerequisites
+Before you can use the Procedural Foliage tools in your project, you must first enable them by going to the **Main Toolbar** and then going to **Edit** > **Editor Preferences** > **Experimental**. 
+Once you are in the Experimental section enable Procedural Foliage option by clicking on the checkmark box next to the word **Procedural Foliage**.
+You will also need to download the **Open World Demo Collection** content pack from the Unreal Engine 4 Launcher as some of the content from the collection will be used in the Open World Tools Quick Starts.
+
+Once the Open World Demo Collection is downloaded add it to the project that you are using to follow along with this Quick Start by first locating the Open World Demo Collection in the UE4 Launcher **Vault** section and then click on the button that says **Add to Project**. You will then be shown a list of projects that you can add this to. Select the project that you are using to follow along with this Quick Start and then press the **Add to Project** button to add the content. When you re-open the project you will now see a folder called **Kite Demo** and inside this folder you will find the content you need.
+
+
 [OBJECT:TopicButtonList]
 	[PARAM:icon]
 		![](%ROOT%/build_icon.png)(convert:false)
