#UE4doc #docs: minor punctuation changes for clarity.

diff --git a/Engine/Documentation/Source/Engine/OpenWorldTools/Grass/QuickStart/4/GTQS_4.INT.udn b/Engine/Documentation/Source/Engine/OpenWorldTools/Grass/QuickStart/4/GTQS_4.INT.udn
--- a/Engine/Documentation/Source/Engine/OpenWorldTools/Grass/QuickStart/4/GTQS_4.INT.udn
+++ b/Engine/Documentation/Source/Engine/OpenWorldTools/Grass/QuickStart/4/GTQS_4.INT.udn
@@ -63,9 +63,8 @@ We will continue to work with the **GT_00** level that was created in the last s
 	If you find that the lag makes it too hard to work, try turning down the **Grass Density** in the **Landscape Grass Type** settings and then turn it back up when done.
 	[/region]
 
-1. To remove grass from the Landscape Terrain select the **Rock Target Layer** then inside of the viewport hold down the **Left Mouse Button** 
-	to begin painting the Rock Texture with will remove the grass.
-	Adjusting the **Brush Size** and **Tool Strength** will help to better define how the grass is placed or removed while painting on the Landscape.
+1. To remove grass from the Landscape Terrain, select the Rock Target Layer, then inside of the viewport hold down the Left Mouse Button to begin painting. The Rock Texture with will replace the grass.
+   Adjusting the **Brush Size** and **Tool Strength** will help to better define how the grass is placed or removed while painting on the Landscape.
 
 	![](T_Painting_Removing_Grass.gif)(convert:false)
 
