no changes made

diff --git a/Engine/Documentation/Source/Engine/OpenWorldTools/ProceduralFoliage/QuickStart/1/PFQS_1.INT.udn b/Engine/Documentation/Source/Engine/OpenWorldTools/ProceduralFoliage/QuickStart/1/PFQS_1.INT.udn
--- a/Engine/Documentation/Source/Engine/OpenWorldTools/ProceduralFoliage/QuickStart/1/PFQS_1.INT.udn
+++ b/Engine/Documentation/Source/Engine/OpenWorldTools/ProceduralFoliage/QuickStart/1/PFQS_1.INT.udn
@@ -42,7 +42,7 @@ In the following How To we will create a new Level, Landscape Terrain and all of
 	The reason we are using a Landscape Terrain this is because the Landscape Terrain Actor quickly provides us with a very large area to spawn lots of Static Meshes on.
 	[/region]
 
-1. Next create a new Procedural Foliage Spawner by **Right- Clicking** in the **Content Browser** with the mouse and going to the **Create Advanced Asset** menu option then **Miscellaneous** > **Procedural Foliage Spawner**. 
+1. Next create a new Procedural Foliage Spawner by **Right- Clicking** in the **Content Browser** and going to the **Create Advanced Asset** menu option then **Miscellaneous** > **Procedural Foliage Spawner**. 
 	
 	![](T_Create_PFS.png)
 
