Minor updates to C++ QS tutorial.

diff --git a/Engine/Documentation/Source/Programming/QuickStart/1/QuickStart_1.INT.udn b/Engine/Documentation/Source/Programming/QuickStart/1/QuickStart_1.INT.udn
--- a/Engine/Documentation/Source/Programming/QuickStart/1/QuickStart_1.INT.udn
+++ b/Engine/Documentation/Source/Programming/QuickStart/1/QuickStart_1.INT.udn
@@ -27,6 +27,6 @@ Description:Creating a new project with the specifications we need
 
 ----
 
-The **Unreal Editor** will now open our new project.
+The **Unreal Editor** will now open our new project. **Visual Studio** will also open and load the solution file that our project has created.
 
 %Steps%
diff --git a/Engine/Documentation/Source/Programming/QuickStart/2/QuickStart_2.INT.udn b/Engine/Documentation/Source/Programming/QuickStart/2/QuickStart_2.INT.udn
--- a/Engine/Documentation/Source/Programming/QuickStart/2/QuickStart_2.INT.udn
+++ b/Engine/Documentation/Source/Programming/QuickStart/2/QuickStart_2.INT.udn
@@ -36,6 +36,6 @@ Version: 4.7
 
 ----
 
-**Unreal Engine** will now open **Visual Studio** for you and load your new project.
+Now that we have created a C++ class, we can switch to **Visual Studio** to program it. `FloatingActor.cpp` will be opened for us automatically, and **Unreal Engine** will automatically compile and reload the code with our new class.
 
 %Steps%
diff --git a/Engine/Documentation/Source/Programming/QuickStart/3/QuickStart_3.INT.udn b/Engine/Documentation/Source/Programming/QuickStart/3/QuickStart_3.INT.udn
--- a/Engine/Documentation/Source/Programming/QuickStart/3/QuickStart_3.INT.udn
+++ b/Engine/Documentation/Source/Programming/QuickStart/3/QuickStart_3.INT.udn
@@ -25,7 +25,7 @@ Description:Writing a C++ class and compiling it with Unreal Engine
 
 	[![](SolutionExplorer_Files.png)](SolutionExplorer_Files.png)
 
-	[REGION:note]We're about to add code. All code used in this tutorial will be found at bottom of the page in its current state as of completing that page's instructions.[/REGION]
+	[REGION:note]We're about to write code. All code used in this tutorial will be found at bottom of the page in its current state as of completing that page's instructions.[/REGION]
 
 2. In `FloatingActor.h`, we'll add the following code just before the closing brace and semicolon at the end of the file:
 
