#UE4 Docs: First draft update of Linux cross compiling page

diff --git a/Engine/Documentation/Source/Platforms/Linux/GettingStarted/LinuxCrossGettingStarted.INT.udn b/Engine/Documentation/Source/Platforms/Linux/GettingStarted/LinuxCrossGettingStarted.INT.udn
--- a/Engine/Documentation/Source/Platforms/Linux/GettingStarted/LinuxCrossGettingStarted.INT.udn
+++ b/Engine/Documentation/Source/Platforms/Linux/GettingStarted/LinuxCrossGettingStarted.INT.udn
@@ -3,11 +3,25 @@ Title: Getting Started: Developing Linux Projects
 Crumbs:%ROOT%, Platforms, Platforms/Linux
 Description:Setting up for development for the Linux platform.
 
+%Globals:BinarySourceSelectVar%
+
 [TOC(start:2 end:3)]
 
-## Disclaimer
+## Why cross-compilation
 
-%Globals:BinarySourceSelectVar%
+Cross-compilation was chosen in order to make it easier for game developers (who tend to have Windows-centric workflow) to target Linux. At the moment, this is only supported for Windows, Mac users currently have to resort to [native compilation](https://wiki.unrealengine.com/Building_On_Linux). We support and test Linux-x86_64, but we provide the libraries and the toolchain (clang-3.5.0-arm) which allow you to compile for Linux ARM as well (original Raspberry Pi and up), although this will require you to make a (minor) code change in UBT.
+
+From this point on we assume that you are targeting x86_64 Linux; although the most of below will apply to compiling for ARM as well (except for the different toolchain). Note that you can [build your own](https://wiki.unrealengine.com/Building_Linux_cross-toolchain) toolchain if you want different versions of the tools or want to target a different architecture.
+
+## Setting up the toolchain 
+
+Add an environment variable (Control Panel->System->Advanced system settings->Advanced->Environment variables) named LINUX_ROOT, value of which is the absolute path to your toolchain, without trailing backslash:
+
+![](image00.png)
+
+After you made sure that variable is set, re-regenerate UE4 project files (using GenerateProjectFiles.bat) and restart Visual Studio. After this, you should have "Linux" available among Win32/Win64 configurations, and you should be able to cross-compile for it.
+
+## Setting up packaging for Linux
 
 [OBJECT:ToggleButtonContent]
 	[PARAMLITERAL:category]
@@ -37,59 +51,56 @@ Description:Setting up for development for the Linux platform.
 	[/PARAM]
 [/OBJECT]
 
-
 [VAR:source]
-## Source and Github Users
+For packaging for Linux platform you will need to build the following targets:
 
-Linux as a target platform is supported on Windows host platform and on Linux itself (but not yet on Mac). This guide is geared towards Windows users trying to cross-compile for Linux; if you run Unreal Editor on Linux itself, no separate step is needed at the moment
-(note that on Linux, we currently use the compiler available on the system, although in the future we may bundle it to provide a stable build environment which will allow you to prepare the binaries in a distro-agnostic way).
+For Linux:
 
-On Windows, we use the [cross-compiler toolchain](http://cdn.unrealengine.com/qfe/v4_clang-3.5.0_ld-2.24_glibc-2.12.2.zip) to cross compile a Linux binary. 
+* CrashReportClient
+* UE4Game (if you project is content-only, otherwise just build your project)
 
-### Installing from Source: 
+For Windows:
 
-As this guide assumes that you are running Windows, you will need to install the toolchain first:
+* The editor itself (this is needed to get Linux target modules built, so that Unreal Editor and Unreal Frontend know how to cook/package games for Linux).
+* UnrealPak and ShaderCompileWorker (you probably have those built anyway, but mentioning it for completeness).
 
-* [Clang-3.5.0 based tool chain](http://cdn.unrealengine.com/qfe/v4_clang-3.5.0_ld-2.24_glibc-2.12.2.zip)
-
-### Toolchain set up
-
-Unzip the above file to a path that does not contain spaces or non-ASCII characters (e.g. C:\CrossToolchain350 is a good example) and set LINUX_ROOT environment variable so it contains the absolute path to the "toolchain" directory of the toolchain (e.g. C:\CrossToolchain350\toolchain\).
-More information: https://wiki.unrealengine.com/Compiling_For_Linux
-
-### Unreal Engine 4 Setup. 
+[/VAR]
 
-* Fork and clone the Unreal Engine 4 repository \[[UnrealEngine](https://github.com/EpicGames/UnrealEngine)\]
+[VAR:binary]
+Note that binary release can only package content-only projects. If you want to package a code-based project for Linux (or a project that has non-default third party plugins), you will have to set up a source build. Otherwise, if you are fine with content-only, there are no need to build anything as Linux binaries of UE4Game and CrashReportClient are bundled.
+[/VAR]
 
-* Within a shell run `Setup.bat -include=Linux`
+## Packaging for Linux
 
-Follow the instuctions on [https://github.com/EpicGames/UnrealEngine](https://github.com/EpicGames/UnrealEngine) to build the editor.  Set your solution configuration to **Development** and your solution platform to **Linux** to build Linux targets.
+The easiest way to package a project is to open it in the editor and then using File->Package To->Linux (assuming that you have a cross-toolchain installed in the previous step and Linux target modules are built for the editor in question - if you don't see Linux in the list, then it's likely that one of the former is not true - see above). After some time (which depends on the project in question and is rather short for sample project) you will have game assets and binaries in the directory you chose to package to.
 
-[/VAR]
+The details of the packaging process can be seen by clicking the "Show Output Log" link. If this process fails with error message "unable to find UnrealPak or ShaderCompileWorker" see above about building them for the host (windows) platform.
 
+Scp or otherwise copy it to a target machine (mounting a Samba share - if you know how to do that - may be better if target machine is low on disk space, this also reduces iteration times), chmod +x the target executable (which will be located in LinuxNoEditor/<ProjectName>/Binaries/Linux/ directory) and run it.
 
+### Changes for Linux-ARM platforms
 
+If you are using the Linux-ARM cross compile toolchain before running GenerateProjectFiles.bat below edit the following file in the UE4 source code: 
 
+.../UnrealEngine/Engine/Source/Programs/UnrealBuildTool/Linux/UEBuildLinux.cs 
 
+by commenting out the following line: 
 
-[VAR:binary]
+static private string DefaultArchitecture = "x86_64-unknown-linux-gnu"; 
 
-## Binary and Launcher Users
+and un-commenting the linux-arm architecture line just below it: 
 
-Linux as a target platform is supported on Windows host platform and on Linux itself (but not yet on Mac). This guide is geared towards Windows users trying to cross-compile for Linux; if you run Unreal Editor on Linux itself, no separate step is needed at the moment
-(note that on Linux, we currently use the compiler available on the system, although in the future we may bundle it to provide a stable build environment which will allow you to prepare the binaries in a distro-agnostic way).
+//static private string DefaultArchitecture = "arm-unknown-linux-gnueabihf";
 
-On Windows, we use the [cross-compiler toolchain](http://cdn.unrealengine.com/qfe/v4_clang-3.5.0_ld-2.24_glibc-2.12.2.zip) to cross compile a Linux binary. 
+An additional step is required for the packaged project to be runnable on Linux-ARM platforms. Assuming the packaged project is located at \foo\bar\SunTemple_Linux open the following location:
 
-### Installing 
+\foo\bar\SunTemple_Linux\LinuxNoEditor\Engine\Binaries\Linux 
 
-As this guide assumes that you are running Windows, you will need to install the toolchain first:
+and replace libopenal.so.1 with the version from: 
 
-* [Clang-3.5.0 based tool chain](http://cdn.unrealengine.com/qfe/v4_clang-3.5.0_ld-2.24_glibc-2.12.2.zip)
+...\UnrealEngine\Engine\Source\ThirdParty\OpenAL\1.15.1\lib\Linux\arm-unknown-linux-gnueabihf\libopenal.so 
 
-### Toolchain set up
+where ...\UnrealEngine is the location of the UE4 source code. Please be sure to rename libopenal.so to libopenal.so.1.
 
-Unzip the above file to a path that does not contain spaces or non-ASCII characters (e.g. C:\CrossToolchain350 is a good example) and set LINUX_ROOT environment variable so it contains the absolute path to the "toolchain" directory of the toolchain (e.g. C:\CrossToolchain350\toolchain\).
-More information: https://wiki.unrealengine.com/Compiling_For_Linux
+Now the project can be copied over to the target machine, chmod+x the target executable and run it as instructed above.
 
-[/VAR]
