Updated sentence for clarity

diff --git a/Engine/Documentation/Source/Platforms/Android/GettingStarted/2/AndroidQS_2.INT.udn b/Engine/Documentation/Source/Platforms/Android/GettingStarted/2/AndroidQS_2.INT.udn
--- a/Engine/Documentation/Source/Platforms/Android/GettingStarted/2/AndroidQS_2.INT.udn
+++ b/Engine/Documentation/Source/Platforms/Android/GettingStarted/2/AndroidQS_2.INT.udn
@@ -35,7 +35,7 @@ Description:Set up your Android device for development.
 		[REGION: note]
 		Note: This may be in a **More** section on newer Android versions.
 		[/REGION]
-	1.	Tap the **Build Number** 7 times. After the first few presses, you will see a **You are now # steps away from becoming a developer.** message.
+	1.	Tap the **Build Number** 7 times. After the first few presses, you will see a message stating **You are now # steps away from becoming a developer.**
 	1.	You will get a success message once **Developer Mode** has been activated.
 	1.	Navigate back to the **Settings** menu and select **Developer Options** which is now available.
 1.	In the **Developer Options** menu, enable USB debugging. 
