#UE4doc #docs: sp: performance

diff --git a/Engine/Documentation/Source/Platforms/DeviceProfiles/DeviceProfiles.INT.udn b/Engine/Documentation/Source/Platforms/DeviceProfiles/DeviceProfiles.INT.udn
--- a/Engine/Documentation/Source/Platforms/DeviceProfiles/DeviceProfiles.INT.udn
+++ b/Engine/Documentation/Source/Platforms/DeviceProfiles/DeviceProfiles.INT.udn
@@ -32,7 +32,7 @@ For more on setting up configuration files, see the main [Configuration Files](P
 
 Device profiles can be used to set your game's resolution for iOS devices. The `r.MobileContentScaleFactor` property is relative to the 'nominal iOS resolution' for a device. For example, a Retina iPad 
 has actual resolution 2048x1536, but nominal resolution 1024x768, so you would use 2.0 to get native resolution there. You can also use fractional values like 1.5 if you want more resolution but not the 
-full peformance hit from running at native resolution.
+full performance hit from running at native resolution.
 
 	[iPad3 DeviceProfile]
 	+CVars=r.MobileContentScaleFactor=1
