#UE4doc #docs: Multiple edits.

diff --git a/Engine/Documentation/Source/Engine/OpenWorldTools/Overview/OpenWorldToolsOverview.INT.udn b/Engine/Documentation/Source/Engine/OpenWorldTools/Overview/OpenWorldToolsOverview.INT.udn
--- a/Engine/Documentation/Source/Engine/OpenWorldTools/Overview/OpenWorldToolsOverview.INT.udn
+++ b/Engine/Documentation/Source/Engine/OpenWorldTools/Overview/OpenWorldToolsOverview.INT.udn
@@ -40,8 +40,8 @@ If you want to densely populate the Landscape Terrain used in that pine and birc
 The Procedural Foliage Tool is used for simulating huge forests that are filled with lots of different kinds of trees and bushes. 
 This tool works by simulating how a forest grows over the years using **Steps** to define years of growth. 
 At each Step in the simulation new virtual seeds are cast into the level to act as spawn locations for other Foliage Type Actors. 
-By tweaking the various Foliage Types options like **Spread**, **Size**, and **Priority Settings**, you will affect how those Foliage Types spread and compete for spawn locations through the level the Procedural Foliage Spawner in placed in.
-To learn more about how to use the Procedural Foliage Tool, check out the Procedural Foliage Quick Start that is linked below.
+By tweaking the various Foliage Type options like **Spread**, **Size**, and **Priority Settings**, you will affect how those Foliage Types spread and compete for spawn locations through the level the Procedural Foliage Spawner in placed in.
+To learn more about how to use the Procedural Foliage Tool check out the Procedural Foliage Quick Start linked below.
 
 [REGION:topics third]
 [OBJECT:Topic]
@@ -64,7 +64,7 @@ To learn more about how to use the Procedural Foliage Tool, check out the Proced
 [/region]
 
 #### Foliage Type Actors
-The Procedural Foliage Tools makes use of Actors called **Foliage Types**. 
+The Procedural Foliage Tools make use of Actors called **Foliage Types**. 
 Foliage Types allow you to define various properties the Foliage should have like what Static Mesh to use or how big Static Mesh should be spawned.
 The chart below gives a very high level explanation of the four main sets of properties that the Foliage Type Actor is broken down into and what those options do.
 
@@ -73,18 +73,18 @@ The chart below gives a very high level explanation of the four main sets of pro
 | Property Name| Description |
 | ------------ | ----------- |
 |**Mesh**| The Mesh section is where you define which Static Mesh to use for the foliage.|
-|**Placement**| The Placement section is where you can adjust how the Static Mesh specified in the Mesh section are placed on objects throughout the level.|
-|**Growth**| The Growth section helps define how big the trees will be at the beginning of the simulation and how big they will "Grow" to by the end of the simulation.|
+|**Placement**| The Placement section is where you can adjust how the Static Meshes specified in the Mesh section are placed on objects throughout the level.|
+|**Growth**| The Growth section helps define how big the trees will be at the beginning of the simulation and how big they will "Grow" by the end of the simulation.|
 |**Instance Settings**| The Instance Settings allows you to adjust how the Static Mesh used for the Foliage Type will be displayed in the level.|
 
 ### Grass Tool
 [region:note]
 Currently the Grass Tools will only work in conjunction with the Landscape Terrain system. 
 [/region]
-The Grass Tools are used to densely cover the Landscape Terrain in things Static Mesh to simulate things like grass, flowers, and small rocks. 
+The Grass Tools are used to densely cover the Landscape Terrain in Static Meshes to simulate things like grass, flowers, and small rocks. 
 You can use the Grass Tools to define a Static Mesh that will spawn when a certain Landscape Layer is painted on the Landscape Terrain. 
-This tool does not use any kind of simulation like the Procedural Foliage Spawner does, it simply fills the specified area at the specified density with the specified Static Mesh. 
-To learn more about how to use the Grass Tool, check out the Grass Tool Quick Start that is lined below.
+This tool does not use any kind of simulation like the Procedural Foliage Spawner does. Instead, it simply fills the specified area at the specified density with the specified Static Mesh. 
+To learn more about how to use the Grass Tool, check out the Grass Tool Quick Start that is linked below.
 
 [REGION:topics third]
 [OBJECT:Topic]
