#UE4 Docs: fixed typo

diff --git a/Engine/Documentation/Source/Platforms/Linux/GettingStarted/LinuxCrossGettingStarted.INT.udn b/Engine/Documentation/Source/Platforms/Linux/GettingStarted/LinuxCrossGettingStarted.INT.udn
--- a/Engine/Documentation/Source/Platforms/Linux/GettingStarted/LinuxCrossGettingStarted.INT.udn
+++ b/Engine/Documentation/Source/Platforms/Linux/GettingStarted/LinuxCrossGettingStarted.INT.udn
@@ -19,7 +19,7 @@ Add an environment variable (Control Panel->System->Advanced system settings->Ad
 
 ![](image00.png)
 
-After you made sure that variable is set, re-regenerate UE4 project files (using GenerateProjectFiles.bat) and restart Visual Studio. After this, you should have "Linux" available among Win32/Win64 configurations, and you should be able to cross-compile for it.
+After you made sure that variable is set, regenerate UE4 project files (using GenerateProjectFiles.bat) and restart Visual Studio. After this, you should have "Linux" available among Win32/Win64 configurations, and you should be able to cross-compile for it.
 
 ## Setting up packaging for Linux
 
