#UE4 Docs: set Linux pages public

 - added link to Linux docs from Platforms landing page

diff --git a/Engine/Documentation/Source/Platforms/Linux/GettingStarted/LinuxCrossGettingStarted.INT.udn b/Engine/Documentation/Source/Platforms/Linux/GettingStarted/LinuxCrossGettingStarted.INT.udn
--- a/Engine/Documentation/Source/Platforms/Linux/GettingStarted/LinuxCrossGettingStarted.INT.udn
+++ b/Engine/Documentation/Source/Platforms/Linux/GettingStarted/LinuxCrossGettingStarted.INT.udn
@@ -1,4 +1,4 @@
-﻿Availability:Docs
+﻿Availability:Public
 Title: Getting Started: Developing Linux Projects
 Crumbs:%ROOT%, Platforms, Platforms/Linux
 Description:Setting up for development for the Linux platform.
diff --git a/Engine/Documentation/Source/Platforms/Linux/Linux.INT.udn b/Engine/Documentation/Source/Platforms/Linux/Linux.INT.udn
--- a/Engine/Documentation/Source/Platforms/Linux/Linux.INT.udn
+++ b/Engine/Documentation/Source/Platforms/Linux/Linux.INT.udn
@@ -1,4 +1,4 @@
-Availability: Docs
+Availability: Public
 Title: Linux Game Development
 Crumbs:%ROOT%, Platforms
 Description:Developing games for the Linux platform.
diff --git a/Engine/Documentation/Source/Platforms/Platforms.INT.udn b/Engine/Documentation/Source/Platforms/Platforms.INT.udn
--- a/Engine/Documentation/Source/Platforms/Platforms.INT.udn
+++ b/Engine/Documentation/Source/Platforms/Platforms.INT.udn
@@ -94,6 +94,20 @@ platforms
 	[/PARAM]
 [/OBJECT]
 
+[OBJECT:TopicButtonList]
+	[PARAM:icon]
+		![](%ROOT%/platform_icon.png)(convert:false)
+	[/PARAM]
+	[PARAM:title]
+		[](Platforms/Linux)
+	[/PARAM]
+	[PARAM:description]
+	[/PARAM]
+	[PARAM:links]
+		* [](Platforms/Linux/GettingStarted "%Platforms/Linux/GettingStarted:description%")
+	[/PARAM]
+[/OBJECT]
+
 [OBJECT:TopicButtonList]
 	[PARAM:icon]
 		![](%ROOT%/platform_icon.png)(convert:false)
