#UE4doc #docs: sp: variables

diff --git a/Engine/Documentation/Source/Platforms/HTML5/BuildingEmscripten/HTML5BuildingEmscripten.INT.udn b/Engine/Documentation/Source/Platforms/HTML5/BuildingEmscripten/HTML5BuildingEmscripten.INT.udn
--- a/Engine/Documentation/Source/Platforms/HTML5/BuildingEmscripten/HTML5BuildingEmscripten.INT.udn
+++ b/Engine/Documentation/Source/Platforms/HTML5/BuildingEmscripten/HTML5BuildingEmscripten.INT.udn
@@ -72,7 +72,7 @@ The CMake command will output a `LLVM.sln`. Open it in Visual Studio 2010 SP1.
 
 All projects should build successfully and there should be 0 errors in the final tally of the output log.
 
-### Update environment varibles
+### Update environment variables
 
 As with the normal installation process for Emscripten, you will need to verify that all environment variables are set to correct values. Because of the manual nature of building Emscripten from source, you will likely need to modify most of these variables.
 
