#UE4 Docs: removed Engine/Basics from crumbs

diff --git a/Engine/Documentation/Source/Engine/Basics/AssetsAndPackages/AssetsAndPackages.INT.udn b/Engine/Documentation/Source/Engine/Basics/AssetsAndPackages/AssetsAndPackages.INT.udn
--- a/Engine/Documentation/Source/Engine/Basics/AssetsAndPackages/AssetsAndPackages.INT.udn
+++ b/Engine/Documentation/Source/Engine/Basics/AssetsAndPackages/AssetsAndPackages.INT.udn
@@ -1,6 +1,6 @@
 Availability:Public
 Title: Assets and Packages
-Crumbs: %ROOT%, Engine, Engine/Basics
+Crumbs: %ROOT%, Engine
 Description:Overview of the asset and package system used in Unreal Engine 4.
 
 [VAR:Topic]
diff --git a/Engine/Documentation/Source/Engine/Basics/AssetsAndPackages/Saving/Saving.INT.udn b/Engine/Documentation/Source/Engine/Basics/AssetsAndPackages/Saving/Saving.INT.udn
--- a/Engine/Documentation/Source/Engine/Basics/AssetsAndPackages/Saving/Saving.INT.udn
+++ b/Engine/Documentation/Source/Engine/Basics/AssetsAndPackages/Saving/Saving.INT.udn
@@ -1,6 +1,6 @@
 Availability: Public
 Title:Package Saving Procedure in UnrealEd
-Crumbs: %ROOT%, Engine, Engine/Basics, Engine/Basics/AssetsAndPackages
+Crumbs: %ROOT%, Engine, Engine/Basics/AssetsAndPackages
 Description:Guide to saving content packages inside of Unreal Editor.
 
 [TOC (start:2 end:3)]
diff --git a/Engine/Documentation/Source/Engine/Basics/DerivedDataCache/DerivedDataCache.INT.udn b/Engine/Documentation/Source/Engine/Basics/DerivedDataCache/DerivedDataCache.INT.udn
--- a/Engine/Documentation/Source/Engine/Basics/DerivedDataCache/DerivedDataCache.INT.udn
+++ b/Engine/Documentation/Source/Engine/Basics/DerivedDataCache/DerivedDataCache.INT.udn
@@ -1,6 +1,6 @@
 Availability:Public
 Title:Derived Data Cache
-Crumbs:%ROOT%, Engine, Engine/Basics
+Crumbs:%ROOT%, Engine
 Description: Overview of the Derived Data Cache and the storage of assets in formats used by UE4 and its target platforms. 
 
 The **Derived Data Cache** (DDC) stores versions of assets in the formats used by the engine and its 
diff --git a/Engine/Documentation/Source/Engine/Basics/DirectoryStructure/DirectoryStructure.INT.udn b/Engine/Documentation/Source/Engine/Basics/DirectoryStructure/DirectoryStructure.INT.udn
--- a/Engine/Documentation/Source/Engine/Basics/DirectoryStructure/DirectoryStructure.INT.udn
+++ b/Engine/Documentation/Source/Engine/Basics/DirectoryStructure/DirectoryStructure.INT.udn
@@ -1,6 +1,6 @@
 Availability:Public
 Title:Directory Structure
-Crumbs: %ROOT%, Engine, Engine/Basics
+Crumbs: %ROOT%, Engine
 Description:Overview of the directories that make up the engine and game projects.
 
 [TOC(start:2)]
diff --git a/Engine/Documentation/Source/Engine/Basics/Distributions/Distributions.INT.udn b/Engine/Documentation/Source/Engine/Basics/Distributions/Distributions.INT.udn
--- a/Engine/Documentation/Source/Engine/Basics/Distributions/Distributions.INT.udn
+++ b/Engine/Documentation/Source/Engine/Basics/Distributions/Distributions.INT.udn
@@ -1,6 +1,6 @@
 Availability: Public
 Title: Distributions
-Crumbs: %ROOT%, Engine, Engine/Basics
+Crumbs: %ROOT%, Engine
 Description:Data types that provide flexibility by allowing for constant values, random values within a range, values interpolated along a curve, and values driven by parameters.
 
 
diff --git a/Engine/Documentation/Source/Engine/Basics/Projects/CreatingTemplates/CreatingTemplates.INT.udn b/Engine/Documentation/Source/Engine/Basics/Projects/CreatingTemplates/CreatingTemplates.INT.udn
--- a/Engine/Documentation/Source/Engine/Basics/Projects/CreatingTemplates/CreatingTemplates.INT.udn
+++ b/Engine/Documentation/Source/Engine/Basics/Projects/CreatingTemplates/CreatingTemplates.INT.udn
@@ -1,6 +1,6 @@
 Availability:Public
 Title:Converting a Project to a Template
-Crumbs: %ROOT%, Engine, Engine/Basics, Engine/Basics/Projects
+Crumbs: %ROOT%, Engine, Engine/Basics/Projects
 Description:Steps for converting an existing project to a template.
 
 [EXCERPT:ConvertDesc]
diff --git a/Engine/Documentation/Source/Engine/Basics/Redirectors/Redirectors.INT.udn b/Engine/Documentation/Source/Engine/Basics/Redirectors/Redirectors.INT.udn
--- a/Engine/Documentation/Source/Engine/Basics/Redirectors/Redirectors.INT.udn
+++ b/Engine/Documentation/Source/Engine/Basics/Redirectors/Redirectors.INT.udn
@@ -1,6 +1,6 @@
 Availability: Public
 Title:Redirectors
-Crumbs:%ROOT%, Engine, Engine/Basics
+Crumbs:%ROOT%, Engine
 Description:Objects that are placed in packages to redirect references to moved assets to their current locations.
 
 [TOC (start:2 end:3)]
diff --git a/Engine/Documentation/Source/Engine/Basics/Screenshots/Screenshots.INT.udn b/Engine/Documentation/Source/Engine/Basics/Screenshots/Screenshots.INT.udn
--- a/Engine/Documentation/Source/Engine/Basics/Screenshots/Screenshots.INT.udn
+++ b/Engine/Documentation/Source/Engine/Basics/Screenshots/Screenshots.INT.udn
@@ -1,6 +1,6 @@
 Availability: Public
 Title:Taking Screenshots
-Crumbs:%ROOT%, Engine, Engine/Basics
+Crumbs:%ROOT%, Engine
 Description:Guide to taking in-game screenshots of environments and gameplay.
 
 [TOC (start:2 end:3)]
diff --git a/Engine/Documentation/Source/Engine/Performance/SystemSettings/SystemSettings.INT.udn b/Engine/Documentation/Source/Engine/Performance/SystemSettings/SystemSettings.INT.udn
--- a/Engine/Documentation/Source/Engine/Performance/SystemSettings/SystemSettings.INT.udn
+++ b/Engine/Documentation/Source/Engine/Performance/SystemSettings/SystemSettings.INT.udn
@@ -1,6 +1,6 @@
 Availability: Public
 Title:System Settings
-Crumbs: %ROOT%, Engine, Engine/Basics
+Crumbs: %ROOT%, Engine
 Description:Platform-specific collection of settings that control certain aspects of the game and engine based on the hardware capabilities.
 
 
