TypoCorrection

diff --git a/Engine/Documentation/Source/Engine/Rendering/Materials/ExpressionReference/Utility/UtilityExpressions.INT.udn b/Engine/Documentation/Source/Engine/Rendering/Materials/ExpressionReference/Utility/UtilityExpressions.INT.udn
--- a/Engine/Documentation/Source/Engine/Rendering/Materials/ExpressionReference/Utility/UtilityExpressions.INT.udn
+++ b/Engine/Documentation/Source/Engine/Rendering/Materials/ExpressionReference/Utility/UtilityExpressions.INT.udn
@@ -1,4 +1,4 @@
-Availability:Public
+﻿Availability:Public
 Title:Utility Expressions
 Crumbs: %ROOT%, Engine, Engine/Rendering/Materials, Engine/Rendering/Materials/ExpressionReference
 Description:Expressions that perform various utility operations on one or more inputs.
@@ -202,7 +202,7 @@ The LightmassReplace expression simply passes through the Realtime input when co
 [EXCERPT:LinearInterpolate]
 ## LinearInterpolate
 
-The LinearInterpolate expression blends between two input value(s) based on a third input value used as a mask. This can be thought of as a mask to define transitions between two textures, like a layer mask in Photoshop.  The intensity of the mask Alpha determines the ratio of color to take from the two input values.  If Alpha is 1.0/white, the first input is used.  If Alpha is 0.0/black, the second input is used.  If Alpha is grey (somewhere between 0.0 and 1.0), the output is a blend between the two inputs. Keep in mind that the blend happens per channel.  So, if Alpha is an RGB color, Alpha's red channel value defines the blend between A and B's red channels **independently** of Alpha's green channel, which defines the blend between A and B's green channels.
+The LinearInterpolate expression blends between two input value(s) based on a third input value used as a mask. This can be thought of as a mask to define transitions between two textures, like a layer mask in Photoshop.  The intensity of the mask Alpha determines the ratio of color to take from the two input values.  If Alpha is 0.0/black, the first input is used.  If Alpha is 1.0/white, the second input is used.  If Alpha is grey (somewhere between 0.0 and 1.0), the output is a blend between the two inputs. Keep in mind that the blend happens per channel.  So, if Alpha is an RGB color, Alpha's red channel value defines the blend between A and B's red channels **independently** of Alpha's green channel, which defines the blend between A and B's green channels.
 
 | Item | Description |
 | -------- | ----------- |
