#UE4doc #docs: (Updated text to: [Create a] Reference to TriggerBox + image) x 2

diff --git a/Engine/Documentation/Source/Gameplay/HowTo/ReferenceAssets/Blueprints/GHT_5B.INT.udn b/Engine/Documentation/Source/Gameplay/HowTo/ReferenceAssets/Blueprints/GHT_5B.INT.udn
--- a/Engine/Documentation/Source/Gameplay/HowTo/ReferenceAssets/Blueprints/GHT_5B.INT.udn
+++ b/Engine/Documentation/Source/Gameplay/HowTo/ReferenceAssets/Blueprints/GHT_5B.INT.udn
@@ -27,7 +27,7 @@ The most common method of creating a reference is by referencing an Actor that i
 
 	![](GHT5B_2.png)
 
-1.	With the **Blueprint Editor** window open, **Right-click** in the graph window and select **Add Reference to TriggerBox**. 
+1.	With the **Blueprint Editor** window open, **Right-click** in the graph window and select **Create a Reference to TriggerBox**. 
 
 	![](GHT5B_3.png)
 
@@ -45,7 +45,7 @@ The most common method of creating a reference is by referencing an Actor that i
 
 Another method of getting reference to an Actor can be done when spawning an Actor in your levels. Take a look at the example below. 
 
-1.	Inside your project, open your player character's Blueprint (we are using **MyCharacter**). 
+1.	Inside your project, open your player character's Blueprint (we are using **ThirdPersonCharacter**). 
 
 1.	Inside the graph, **Right-click** and add an **Event Begin Play** node and connect it to a **Spawn Actor from Class** node. 
 
@@ -65,7 +65,7 @@ Another method of getting reference to an Actor can be done when spawning an Act
 
 	![](GHT5B_7.png)
 
-	Here we are accessing the Particle Effect Component (**P_Fire** denoted in the yellow box above) from the Spawned Actor and **Deactivating** it from our **MyCharacter** Blueprint when the player presses the **F** key.
+	Here we are accessing the Particle Effect Component (**P_Fire** denoted in the yellow box above) from the Spawned Actor and **Deactivating** it from our **ThirdPersonCharacter** Blueprint when the player presses the **F** key.
 
 
 ## Direct Blueprint Communication Referencing
@@ -147,7 +147,7 @@ While the **Direct Blueprint Communication** method works for getting reference
 
 	![](GHT5B_2.png)
 
-1.	**Right-click** and select the **Add Reference to...** option. 
+1.	**Right-click** and select the **Create a Reference to...** option. 
 
 	![](GHT5B_19.png)
 
@@ -157,7 +157,7 @@ While the **Direct Blueprint Communication** method works for getting reference
 
 	![](GHT5B_21.png)
 
-	Here we created a variable of the **Blueprint_Effect_Fire_C** type called **TargetBlueprint** (as the fire effect is what we want to reference) inside the **MyCharacter** Blueprint. The **Level Blueprint** when the game is started will get the reference to the Actor in the level, then pass it to the **MyCharacter** Blueprint so that we may access its properties. 
+	Here we created a variable of the **Blueprint_Effect_Fire_C** type called **TargetBlueprint** (as the fire effect is what we want to reference) inside the **ThirdPersonCharacter** Blueprint. When the game is started the **Level Blueprint** will get the reference to the Actor in the level then pass it to the **ThirdPersonCharacter** Blueprint so that we may access its properties. 
 
 1.	**Control** drag the variable into the graph, then drag off it to access its properties. 
 
