#UE4doc #docs: Replaced all references to MyCharacter w ThirdPersonCharacter and updated images accordingly

diff --git a/Engine/Documentation/Source/Gameplay/HowTo/RespawnPlayer/Blueprints/GHT_3B.INT.udn b/Engine/Documentation/Source/Gameplay/HowTo/RespawnPlayer/Blueprints/GHT_3B.INT.udn
--- a/Engine/Documentation/Source/Gameplay/HowTo/RespawnPlayer/Blueprints/GHT_3B.INT.udn
+++ b/Engine/Documentation/Source/Gameplay/HowTo/RespawnPlayer/Blueprints/GHT_3B.INT.udn
@@ -15,29 +15,29 @@ Below you will find the steps needed to respawn the player in a single player ga
 
 _The steps below are based off using the **Blueprint Third Person Template** project, however you can use any project you wish._
 
-1.	Inside your project, open up your **Game Mode** Blueprint (we are using the **MyGame** Blueprint located in **Game/Blueprints**). 
+1.	Inside your project, open up your **Game Mode** Blueprint (we are using the **ThirdPersonGameMode** Blueprint). 
 
 1.	**Right-click** in the graph and add an **Event Begin Play** node, then **Right-click** again and add a **Get Player Character** node. 
 
-1.	Drag off the **Return Value** of the **Get Player Character** node and add a **Cast To MyCharacter** node and connect the nodes. 
+1.	Drag off the **Return Value** of the **Get Player Character** node and add a **Cast To ThirdPersonCharacter** node and connect the nodes. 
 
 	![](GHT3B_1.png)
 
-	You want to **Cast To** your Player Character Blueprint, which is **MyCharacter** in this example. This will allow us to access it and determine when the player character is killed by using an **On Destroyed** event, at which point we can spawn a new player character. 
+	You want to **Cast To** your Player Character Blueprint, which is **ThirdPersonCharacter** in this example. This will allow us to access it and determine when the player character is killed by using an **On Destroyed** event, at which point we can spawn a new player character. 
 
-1.	Drag off the **As My Character C** pin and add a **Assign On Destroyed** node. 
+1.	Drag off the **As ThirdPersonCharacter C** pin and add a **Assign On Destroyed** node. 
 
 1.	Off the newly added **OnDestroyed_Event** node, add a **Delay** node set to whatever time you wish to delay the respawn (optional). 
 
 	![](GHT3B_2.png)
 
-	Here we are Binding an Event to another Event from another Blueprint. This allows the Event in this Blueprint to fire off when the Event occurs from the other Blueprint (the OnDestroy Event from **MyCharacter**).  This is called an **Event Dispatcher** and can also be used to communicate with the **Level Blueprint** to tell the world an Event has occurred. 
+	Here we are Binding an Event to another Event from another Blueprint. This allows the Event in this Blueprint to fire off when the Event occurs from the other Blueprint (the OnDestroy Event from **ThirdPersonCharacter**).  This is called an **Event Dispatcher** and can also be used to communicate with the **Level Blueprint** to tell the world an Event has occurred. 
 	
 	[REGION:note]
 	Refer to the [](Engine/Blueprints/UserGuide/EventDispatcher) documentation for more information.
 	[/REGION] 
 
-1.	Off the **Delay** node (if you added it), add a **SpawnActorFromClass** node and set the **Class** drop down to **MyCharacter**. 
+1.	Off the **Delay** node (if you added it), add a **SpawnActorFromClass** node and set the **Class** drop down to **ThirdPersonCharacter**. 
 
 1.	Create a new **Transform** variable and call it **Spawn Transform**, then connect it to the **Spawn Transform** of the **Spawn Actor** node. 
 
@@ -61,7 +61,7 @@ _The steps below are based off using the **Blueprint Third Person Template** pro
 
 	 You may need to re-check **Context Sensitive** for the node to appear in the menu. 
 
-1.	Connect the custom event to the **Cast To MyCharacter** node. 
+1.	Connect the custom event to the **Cast To ThirdPersonCharacter** node. 
 
 1.	Drag off the **Possess** node and call the **Respawn** event, your full graph should look similar to below. 
 
@@ -73,7 +73,7 @@ _The steps below are based off using the **Blueprint Third Person Template** pro
 
 1.	**Compile** and **Save** then close the Blueprint.
 
-1.	Open the **MyCharacter** Blueprint, **Right-click** in the graph and add an **Event Begin Play** node. 
+1.	Open the **ThirdPersonCharacter** Blueprint, **Right-click** in the graph and add an **Event Begin Play** node. 
 
 1.	**Right-click** again and add a **Get Game Mode** node, then drag off it and **Cast To MyGame** (or your Game Mode). 
 
