#UE4doc #docs: updated 1 image to reflect ThirdPersonCharacter from MyCharacter change

diff --git a/Engine/Documentation/Source/Gameplay/HowTo/SettingUpAGameMode/Blueprints/GameMode_BP.INT.udn b/Engine/Documentation/Source/Gameplay/HowTo/SettingUpAGameMode/Blueprints/GameMode_BP.INT.udn
--- a/Engine/Documentation/Source/Gameplay/HowTo/SettingUpAGameMode/Blueprints/GameMode_BP.INT.udn
+++ b/Engine/Documentation/Source/Gameplay/HowTo/SettingUpAGameMode/Blueprints/GameMode_BP.INT.udn
@@ -27,7 +27,7 @@ In the **Pick Parent Class** window, select the **Game Mode** Class.
 
 	![](GameMode_5.png)
 
-	Here we are assigning the Character Blueprint called **MyCharacter** as the **Default Pawn Class** for players to use in the game. 
+	Here we are assigning the Character Blueprint called **ThirdPersonCharacter** as the **Default Pawn Class** for players to use in the game. 
 
 	[REGION:note]
 	The **Game Mode** Blueprint points to existing Blueprints of the **Character**, **HUD**, **PlayerController**, **Spectator**, and **Game State** Classes. You will need to create these separately then specify them for use in the **Game Mode** Blueprint in order to actually use them in your game. 
@@ -56,7 +56,7 @@ In the previous section, we created a **Game Mode** Blueprint. Once you have a *
 
 	![](GameMode_9.png)
 
-	Here we can see that **MyCharacter** is being used as the **Default Pawn Class**. 
+	Here we can see that **ThirdPersonCharacter** is being used as the **Default Pawn Class**. 
 
 	
 ## Overriding the Default Game Mode
