#UE4doc #docs: Updated "MyCharacter" ref and image

diff --git a/Engine/Documentation/Source/Gameplay/HowTo/UseRaycasts/Blueprints/GHT_2B.INT.udn b/Engine/Documentation/Source/Gameplay/HowTo/UseRaycasts/Blueprints/GHT_2B.INT.udn
--- a/Engine/Documentation/Source/Gameplay/HowTo/UseRaycasts/Blueprints/GHT_2B.INT.udn
+++ b/Engine/Documentation/Source/Gameplay/HowTo/UseRaycasts/Blueprints/GHT_2B.INT.udn
@@ -19,7 +19,7 @@ Below you will find steps for setting up a **Single Line Trace by Channel**.
 
 1.  Create a new project using the **Blueprint First Person** template with **Include Starter Content** and open the project. 
 
-1.  In the **Game/Blueprints** folder, open the **MyCharacter** Blueprint. 
+1.  In the **FirstPersonBP/Blueprints** folder, open the **FirstPersonCharacter** Blueprint. 
 
 1.  **Right-click** in the graph and search for and add an **Event Tick** node. 
 
