#UE4doc #docs: Replaced all references to MyCharacter w ThirdPersonCharacter and updated images accordingly

diff --git a/Engine/Documentation/Source/Gameplay/HowTo/UseTimers/Blueprints/UseTimers_BP.INT.udn b/Engine/Documentation/Source/Gameplay/HowTo/UseTimers/Blueprints/UseTimers_BP.INT.udn
--- a/Engine/Documentation/Source/Gameplay/HowTo/UseTimers/Blueprints/UseTimers_BP.INT.udn
+++ b/Engine/Documentation/Source/Gameplay/HowTo/UseTimers/Blueprints/UseTimers_BP.INT.udn
@@ -15,7 +15,7 @@ The steps below will show you how you can create/set a Timer in Blueprints.
 
 _For this example, we are using the **Blueprint Third Person Template** and will be causing fire damage to the player over time._
 
-1.	Inside your project, open the **MyCharacter** Blueprint. 
+1.	Inside your project, open the **ThirdPersonCharacter** Blueprint. 
 
 1.	**Right-click** in the graph and add two **Custom Events** called **PlayerInFire** and **PlayerNotInFire**.
 
@@ -65,15 +65,15 @@ _For this example, we are using the **Blueprint Third Person Template** and will
 
 	![](GHT6B_7.png)
 
-	When the player enters/exits this Trigger, we will call the Events in the **MyCharacter** Blueprint. 
+	When the player enters/exits this Trigger, we will call the Events in the **ThirdPersonCharacter** Blueprint. 
 
 1.	Return to the **Graph** tab and **Right-click** on the **Trigger** and add the **OnComponentBegin** and **OnComponentEnd** Overlap Events. 
 
 	![](GHT6B_8.png)
 
-1.	Drag off the **Other Actor** pin of each and add the **Cast To MyCharacter** node. 
+1.	Drag off the **Other Actor** pin of each and add the **Cast To ThirdPersonCharacter** node. 
 
-1.	Drag off the **As My Character C** pin and (for the Begin node) add **PlayerInFire** and (for the End node) **PlayerIsNotInFire** Events. 
+1.	Drag off the **As Third Person Character C** pin and (for the Begin node) add **PlayerInFire** and (for the End node) **PlayerIsNotInFire** Events. 
 
 	![](GHT6B_8b.png)
 
