#UE4doc #docs: minor grammar adjustment

diff --git a/Engine/Documentation/Source/Gameplay/SaveGame/Blueprints/BP_SaveGame.INT.udn b/Engine/Documentation/Source/Gameplay/SaveGame/Blueprints/BP_SaveGame.INT.udn
--- a/Engine/Documentation/Source/Gameplay/SaveGame/Blueprints/BP_SaveGame.INT.udn
+++ b/Engine/Documentation/Source/Gameplay/SaveGame/Blueprints/BP_SaveGame.INT.udn
@@ -65,7 +65,7 @@ can be read from it. Again, this may not apply in all game implementations.
 ![](LoadGameBP_1.png)
 
 You can use the **Load Game From Slot** node to create an instance of your SaveGame object Blueprint. Just like when you created a SaveGame object for saving information, you need to cast to your specific 
-SaveGame object Blueprint, and you can save the result to a variable for easy access later.
+SaveGame object Blueprint so that you can save the result to a variable for easy access later.
 
 ![](LoadGameBP_2.png)
 
