#UE4doc #docs: sp: Updated links

diff --git a/Engine/Documentation/Source/Programming/Introduction/IntroToProgramming.INT.udn b/Engine/Documentation/Source/Programming/Introduction/IntroToProgramming.INT.udn
--- a/Engine/Documentation/Source/Programming/Introduction/IntroToProgramming.INT.udn
+++ b/Engine/Documentation/Source/Programming/Introduction/IntroToProgramming.INT.udn
@@ -365,11 +365,11 @@ There are too many specifiers to list here, so I’ll link to the docs for them:
 
 [List of UCLASS Specifiers](Programming/UnrealArchitecture/Reference/Classes/Specifiers)
 
-[List of UPROPERTY Specifiers](https://docs.unrealengine.com/latest/INT/API/Runtime/CoreUObject/UObject/UP___367/index.html)
+[List of UPROPERTY Specifiers](https://docs.unrealengine.com/latest/INT/Programming/UnrealArchitecture/Reference/Properties/Specifiers/index.html)
 
-[List of UFUNCTION Specifiers](https://docs.unrealengine.com/latest/INT/API/Runtime/CoreUObject/UObject/UF___366/index.html)
+[List of UFUNCTION Specifiers](https://docs.unrealengine.com/latest/INT/Programming/UnrealArchitecture/Reference/Functions/Specifiers/index.html)
 
-[List of USTRUCT Specifiers](https://docs.unrealengine.com/latest/INT/API/Runtime/CoreUObject/UObject/US___368/index.html)
+[List of USTRUCT Specifiers](https://docs.unrealengine.com/latest/INT/Programming/UnrealArchitecture/Reference/Structs/Specifiers/index.html)
 
 ### Object/Actor Iterators
 
