#UE4doc #docs: sp: previously, your... you may also [be] prompted

diff --git a/Engine/Documentation/Source/Platforms/iOS/QuickStart/6/iOSQuickStart_6.INT.udn b/Engine/Documentation/Source/Platforms/iOS/QuickStart/6/iOSQuickStart_6.INT.udn
--- a/Engine/Documentation/Source/Platforms/iOS/QuickStart/6/iOSQuickStart_6.INT.udn
+++ b/Engine/Documentation/Source/Platforms/iOS/QuickStart/6/iOSQuickStart_6.INT.udn
@@ -86,11 +86,11 @@ to use with the Unreal Engine.
 
 	![](provision_app_id.png)
 
-1. Select the certificate you created previsouly and click **Continue**.
+1. Select the certificate you created previously and click **Continue**.
 
 	![](provision_certificate.png)
 
-1. Select all of the devices you want associated with the profile. Only devices selected here will be able to have yoru game launch on them.
+1. Select all of the devices you want associated with the profile. Only devices selected here will be able to have your game launch on them.
 
 	![](provision_device.png)
 
@@ -148,7 +148,7 @@ to use with the Unreal Engine.
 	![](import_certificate.png)
 
 	[REGION:note]
-	When importing a certificate, you may also prompted to import the key pair used to generate the certificate.
+	When importing a certificate, you may also be prompted to import the key pair used to generate the certificate.
 
 	![](import_keypair.png)
 
