#UE4doc #docs: fixed link "Oculus Documentation"

diff --git a/Engine/Documentation/Source/Platforms/Oculus/QuickStart/6/OculusQuickStart_6.INT.udn b/Engine/Documentation/Source/Platforms/Oculus/QuickStart/6/OculusQuickStart_6.INT.udn
--- a/Engine/Documentation/Source/Platforms/Oculus/QuickStart/6/OculusQuickStart_6.INT.udn
+++ b/Engine/Documentation/Source/Platforms/Oculus/QuickStart/6/OculusQuickStart_6.INT.udn
@@ -32,7 +32,7 @@ Availability:Public
 	* [](Platforms\VR)
 		* [](Platforms\VR\CheatSheet)
 		* [](Platforms\Oculus\BestPractices)
-	* [Oculus Documentation](https://developer.oculus.com/documents/)
+	* [Oculus Documentation](https://developer.oculus.com/documentation/)
 		* [User Guide](http://static.oculus.com/sdk-downloads/documents/Oculus_User_Guide_0.4.4.pdf)
 		* [Developer Guide](http://static.oculus.com/sdk-downloads/documents/Oculus_Developer_Guide_0.4.4.pdf)
 	[/PARAM]
