#UE4doc #docs: Deleted Link: Building Emscription from Source

diff --git a/Engine/Documentation/Source/Platforms/Platforms.INT.udn b/Engine/Documentation/Source/Platforms/Platforms.INT.udn
--- a/Engine/Documentation/Source/Platforms/Platforms.INT.udn
+++ b/Engine/Documentation/Source/Platforms/Platforms.INT.udn
@@ -119,7 +119,7 @@ platforms
 	[/PARAM]
 	[PARAM:links]
 		* [](Platforms/HTML5/GettingStarted "%Platforms/HTML5/GettingStarted:description%")
-		* [](Platforms/HTML5/BuildingEmscripten "%Platforms/HTML5/BuildingEmscripten:description%")
+	
 	[/PARAM]
 [/OBJECT]
 
