#UE4doc #docs: "more clear" to "clearer"

diff --git a/Engine/Documentation/Source/Programming/Development/CodingStandard/CodingStandard.INT.udn b/Engine/Documentation/Source/Programming/Development/CodingStandard/CodingStandard.INT.udn
--- a/Engine/Documentation/Source/Programming/Development/CodingStandard/CodingStandard.INT.udn
+++ b/Engine/Documentation/Source/Programming/Development/CodingStandard/CodingStandard.INT.udn
@@ -166,7 +166,7 @@ We use a system based on JavaDoc to automatically extract comments from the code
 
 The following example demonstrates the format of class, state, method, and variable comments. Remember that comments should augment the code. The code documents the implementation and the comments document the intent. Make sure to update comments when you change the intent of a piece of code.
 
-Note that two different parameter comment styles are supported, as embodied by the Steep and Sweeten methods.  The @param style used by Steep is the traditional style, but for simple functions it can be more clear to integrate the parameter documentation into the descriptive comment for the function, as in the Sweeten example.
+Note that two different parameter comment styles are supported, as embodied by the Steep and Sweeten methods.  The @param style used by Steep is the traditional style, but for simple functions it can be clearer to integrate the parameter documentation into the descriptive comment for the function, as in the Sweeten example.
 
 Unlike the UE3 coding standard, the method comments should only be included once, where the method is publically declared.  The method comments should only contain information relevant to callers of the method, including any information about overrides of the method that may be relevant to the caller.  Details about the implementation of the method and its overrides that are not relevant to callers should be commented within the method implementation.
 
