#UE4doc #docs: extra word, "are"

diff --git a/Engine/Documentation/Source/Programming/Development/ManagingGameCode/CppClassWizard/CppClassWizard.INT.udn b/Engine/Documentation/Source/Programming/Development/ManagingGameCode/CppClassWizard/CppClassWizard.INT.udn
--- a/Engine/Documentation/Source/Programming/Development/ManagingGameCode/CppClassWizard/CppClassWizard.INT.udn
+++ b/Engine/Documentation/Source/Programming/Development/ManagingGameCode/CppClassWizard/CppClassWizard.INT.udn
@@ -13,7 +13,7 @@ The _C++ Class Wizard_ provides a fast and easy way to add native C++ code class
 with your own functionality, if you wish. This converts a content-only project into a code project. You can access the C++ Class Wizard like so:
 
 [REGION:note]
-On Windows, please make sure you are have Visual Studio 2013 Professional or Express for Windows Desktop installed before proceeding.  On Mac, make sure to have Xcode installed.
+On Windows, please make sure you have Visual Studio 2013 Professional or Express for Windows Desktop installed before proceeding.  On Mac, make sure to have Xcode installed.
 [/REGION:note]
 
 
