#UE4doc #docs: sp: doing, distortion, losing  and edited sentence for clarification

diff --git a/Engine/Documentation/Source/Programming/Development/Tools/ConsoleManager/ConsoleManager.INT.udn b/Engine/Documentation/Source/Programming/Development/Tools/ConsoleManager/ConsoleManager.INT.udn
--- a/Engine/Documentation/Source/Programming/Development/Tools/ConsoleManager/ConsoleManager.INT.udn
+++ b/Engine/Documentation/Source/Programming/Development/Tools/ConsoleManager/ConsoleManager.INT.udn
@@ -31,12 +31,12 @@ The console variable has a name and the user gets auto completion when he starts
 
 ## Creating / registering a console variable
 
-The variable needs to be registered early when the engine is created. This example shows the best way of dong that (in any C++ file):
+The variable needs to be registered early when the engine is created. This example shows the best way of doing that (in any C++ file):
 
 	static TAutoConsoleVariable<int32> CVarRefractionQuality(
 		TEXT("r.RefractionQuality"),
 		2,
-		TEXT("Defines the distorion/refraction quality, adjust for quality or performance.\n")
+		TEXT("Defines the distortion/refraction quality, adjust for quality or performance.\n")
 		TEXT("<=0: off (fastest)\n")
 		TEXT("  1: low quality (not yet implemented)\n")
 		TEXT("  2: normal quality (default)\n")
@@ -51,7 +51,7 @@ If needed you also can generate a console variable inside a function:
 
 	IConsoleManager::Get().RegisterConsoleVariable(TEXT("r.RefractionQuality"),
        2,
-       TEXT("Defines the distorion/refraction quality, adjust for quality or performance.\n")
+       TEXT("Defines the distortion/refraction quality, adjust for quality or performance.\n")
 		TEXT("<=0: off (fastest)\n")
 		TEXT("  1: low quality (not yet implemented)\n")
 		TEXT("  2: normal quality (default)\n")
@@ -87,9 +87,9 @@ Getting the state of console variables created with `RegisterConsoleVariableRef`
 	// get the value on the game thread
 	int32 MyVar = CVarRefractionQuality.GetValueOnGameThread();
 
-Only slightly slower (virtual function call, possibly cache miss) is it to get the state by using one of the get functions: !GetInt(), !GetFloat(), !GetString()
-For best performance, you should used same type the variable was registered with. In order to get the pointer to the variable, you can either store the return
-argument of the register function or call `FindConsoleVariable` just before you needs the variable. Example:
+Using Getter functions (ie. !GetInt(), !GetFloat(), !GetString()) to determine a console variable’s state results in a slightly slower implementation (virtual function call, possibly cache miss, etc.).
+For best performance you should used same type the variable was registered with. In order to get the pointer to the variable, you can either store the return
+argument of the register function or call `FindConsoleVariable` just before you need the variable. Example:
 
     static const auto CVar = IConsoleManager::Get().FindConsoleVariable(TEXT("TonemapperType")); 
 	int32 Value = CVar->GetInt();
@@ -181,7 +181,7 @@ Please read `IConsoleManager.h` for find more details on this.
 On engine startup, the state of console variables can be loaded from the file "Engine/Config/ConsoleVariables.ini". This place is reverved for the local developer - it should nto be used for project settings.
 More details can be found in the file itself:
 
-    ; This file allows to set console variables on engine startup (In undefined order).
+    ; This file allows you to set console variables on engine startup (In undefined order).
     ; Currently there is no other file overriding this one.
     ; This file should be in the source control database (for the comments and to know where to find it)
     ; but kept empty from variables.
@@ -273,7 +273,7 @@ The priority is also helpful to see who set the variable the last time. You can
 
 The `UnregisterConsoleVariable` method allows to remove the console variable. At least this is what is happening from the users perspective.
 The variable is still kept (with the unregistered flags) to not crash when pointers access the data. If a new variable is registered with the same name,
-the old variable is restored and help and flags get copied from the new variable. This way DLL loading and unloading can work even without lossing the variable state.
+the old variable is restored and help and flags get copied from the new variable. This way DLL loading and unloading can work even without losing the variable state.
 Note this will not work for console variable references. This can be fixed by giving up on one: do not store pointers, do not unregistered or not use references.
 
 
