minor edits

diff --git a/Engine/Documentation/Source/Programming/Development/Tools/ConsoleManager/ConsoleManager.INT.udn b/Engine/Documentation/Source/Programming/Development/Tools/ConsoleManager/ConsoleManager.INT.udn
--- a/Engine/Documentation/Source/Programming/Development/Tools/ConsoleManager/ConsoleManager.INT.udn
+++ b/Engine/Documentation/Source/Programming/Development/Tools/ConsoleManager/ConsoleManager.INT.udn
@@ -88,7 +88,7 @@ Getting the state of console variables created with `RegisterConsoleVariableRef`
 	int32 MyVar = CVarRefractionQuality.GetValueOnGameThread();
 
 Using Getter functions (ie. !GetInt(), !GetFloat(), !GetString()) to determine a console variable’s state results in a slightly slower implementation (virtual function call, possibly cache miss, etc.).
-For best performance you should used same type the variable was registered with. In order to get the pointer to the variable, you can either store the return
+For best performance you should use same type the variable was registered with. In order to get the pointer to the variable, you can either store the return
 argument of the register function or call `FindConsoleVariable` just before you need the variable. Example:
 
     static const auto CVar = IConsoleManager::Get().FindConsoleVariable(TEXT("TonemapperType")); 
@@ -178,7 +178,7 @@ Please read `IConsoleManager.h` for find more details on this.
 
 ## Loading console variables
 
-On engine startup, the state of console variables can be loaded from the file "Engine/Config/ConsoleVariables.ini". This place is reverved for the local developer - it should nto be used for project settings.
+On engine startup, the state of console variables can be loaded from the file "Engine/Config/ConsoleVariables.ini". This place is reserved for the local developer - it should not be used for project settings.
 More details can be found in the file itself:
 
     ; This file allows you to set console variables on engine startup (In undefined order).
@@ -271,7 +271,7 @@ The priority is also helpful to see who set the variable the last time. You can
 
 ## Unregistering console variables
 
-The `UnregisterConsoleVariable` method allows to remove the console variable. At least this is what is happening from the users perspective.
+The `UnregisterConsoleVariable` method allows you to remove the console variable. At least this is what is happening from the user's perspective.
 The variable is still kept (with the unregistered flags) to not crash when pointers access the data. If a new variable is registered with the same name,
 the old variable is restored and help and flags get copied from the new variable. This way DLL loading and unloading can work even without losing the variable state.
 Note this will not work for console variable references. This can be fixed by giving up on one: do not store pointers, do not unregistered or not use references.
