#UE4doc #docs: "can by" to "can be"

diff --git a/Engine/Documentation/Source/Programming/Development/VisualStudioSetup/UnrealVS/UnrealVS.INT.udn b/Engine/Documentation/Source/Programming/Development/VisualStudioSetup/UnrealVS/UnrealVS.INT.udn
--- a/Engine/Documentation/Source/Programming/Development/VisualStudioSetup/UnrealVS/UnrealVS.INT.udn
+++ b/Engine/Documentation/Source/Programming/Development/VisualStudioSetup/UnrealVS/UnrealVS.INT.udn
@@ -54,7 +54,7 @@ The UnrealVS extension **does not work** with the **Visual Studio Express** edit
 
 	![UnrealVS Toolbar](unrealvs_toolbar_cmd.png)
 
-1. By default, the toolbar appears as above but its contents can by customized via the **Tools->Customize...** by choosing the **UnrealVS** toolbar in the **Commands** tab.
+1. By default, the toolbar appears as above but its contents can be customized via the **Tools->Customize...** by choosing the **UnrealVS** toolbar in the **Commands** tab.
 
 	![UnrealVS Customize Toolbar](toolbar_customize.png)
 	
