#UE4doc #docs: sp: behavior it's to it is

diff --git a/Engine/Documentation/Source/Programming/Introduction/IntroToProgramming.INT.udn b/Engine/Documentation/Source/Programming/Introduction/IntroToProgramming.INT.udn
--- a/Engine/Documentation/Source/Programming/Introduction/IntroToProgramming.INT.udn
+++ b/Engine/Documentation/Source/Programming/Introduction/IntroToProgramming.INT.udn
@@ -114,7 +114,7 @@ Before continuing to the section below, let’s add a couple of properties to th
 		...
 	};
 
-DamageTimeInSeconds is a property the designer can modify. The DamagePerSecond property is a calculated value using the designer’s settings (see the next section). The VisibleAnywhere flag marks that property as viewable, but not editable in the Unreal Editor. The Transient flag means that it won’t be saved or loaded from disk; it’s meant to be a derived, non-persistent value. The image below shows the properties as part of the class defaults.
+DamageTimeInSeconds is a property the designer can modify. The DamagePerSecond property is a calculated value using the designer’s settings (see the next section). The VisibleAnywhere flag marks that property as viewable, but not editable in the Unreal Editor. The Transient flag means that it won’t be saved or loaded from disk; it is meant to be a derived, non-persistent value. The image below shows the properties as part of the class defaults.
 
 ![image alt text](image_3.png)
 
@@ -182,7 +182,7 @@ This is the first class that we’re going to customize with our designer hats o
 
 ![image alt text](image_10.png)
 
-Wait a second...our calculated value doesn’t match what we’d expect. It should be 150 but it’s still at the default value of 200. The reason for this is that we are only calculating our damage per second value after the properties have been initialized from the loading process. Runtime changes in the Unreal Editor aren’t accounted for. There’s a simple solution to this problem because the engine notifies the target object when it has been changed in the editor. The code below shows the added hooks needed to calculate the derived value as it changes in the editor.
+Wait a second...our calculated value doesn’t match what we’d expect. It should be 150 but it is still at the default value of 200. The reason for this is that we are only calculating our damage per second value after the properties have been initialized from the loading process. Runtime changes in the Unreal Editor aren’t accounted for. There’s a simple solution to this problem because the engine notifies the target object when it has been changed in the editor. The code below shows the added hooks needed to calculate the derived value as it changes in the editor.
 
 	void AMyActor::PostInitProperties()
 	{
@@ -245,7 +245,7 @@ This version still generates the thunking method to call into the Blueprint VM.
 
 Now this version of the function is called when the Blueprint in question doesn’t override the method. One thing to note, is that in future versions of the build tools the auto generated _Implementation() declaration will go away and you’ll be expected to explicitly add that to the header. As of version 4.7, the auto generation of that declaration still occurs.
 
-Now that we have walked through the common gameplay programmer workflow and methods to work with designers to build out gameplay features, it’s time for you to choose your own adventure. You can either continue with this document to read more about how we use C++ in the engine or you can jump right into one of our samples that we include in the launcher to get a more hands on experience. 
+Now that we have walked through the common gameplay programmer workflow and methods to work with designers to build out gameplay features, it is time for you to choose your own adventure. You can either continue with this document to read more about how we use C++ in the engine or you can jump right into one of our samples that we include in the launcher to get a more hands on experience. 
 
 ## Diving Deeper
 
@@ -266,7 +266,7 @@ The base building block in the Unreal Engine is called **UObject**. This class,
 * Configurable values for properties
 * Networking support for properties and methods
 
-Each class that derives from **UObject** has a singleton **UClass** created for it that contains all of the meta data about the class instance. **UObject** and **UClass** together are at the root of everything that a gameplay object does during its lifetime. The best way to think of the difference between a **UClass** and a **UObject** is that the **UClass** *describes* what an instance of a **UObject** will look like, what properties are available for serialization, networking, etc. Most gameplay development doesn’t involve directly deriving from **UObjects**, but instead from **AActor** and **UActorComponent**. You don’t need to know the details of how **UClass**/**UObject** works in order to write gameplay code, but it’s good to know that these systems exist.
+Each class that derives from **UObject** has a singleton **UClass** created for it that contains all of the meta data about the class instance. **UObject** and **UClass** together are at the root of everything that a gameplay object does during its lifetime. The best way to think of the difference between a **UClass** and a **UObject** is that the **UClass** *describes* what an instance of a **UObject** will look like, what properties are available for serialization, networking, etc. Most gameplay development doesn’t involve directly deriving from **UObjects**, but instead from **AActor** and **UActorComponent**. You don’t need to know the details of how **UClass**/**UObject** works in order to write gameplay code, but it is good to know that these systems exist.
 
 #### AActor
 
@@ -315,11 +315,11 @@ This tree of components is attached to the one actor class. As you can see from
 
 To use a **UStruct**, you don’t have to extend from any particular class, you just have mark the struct with **USTRUCT()** and our build tools will do the base work for you. Unlike a **UObject**, **UStructs** are *not* garbage collected. If you create dynamic instances of them, you must manage their lifecycle yourself. **UStructs** are meant to be plain old data types that have the **UObject** reflection support for editing within the Unreal Editor, Blueprint manipulation, serialization, networking, etc.
 
-Now that we’ve talked about the basic hierarchy used in our gameplay class construction, it’s time to choose your path again. You can read about our gameplay classes [here](Programming/UnrealArchitecture/Reference/Classes), head out to our samples in the launcher armed with more information, or continue digging deeper into our C++ features for building games.
+Now that we’ve talked about the basic hierarchy used in our gameplay class construction, it is time to choose your path again. You can read about our gameplay classes [here](Programming/UnrealArchitecture/Reference/Classes), head out to our samples in the launcher armed with more information, or continue digging deeper into our C++ features for building games.
 
 ## Diving Deeper Still
 
-Alright, it’s clear you want to know more. Let’s keep on going deeper into how the engine works.
+Alright, it is clear you want to know more. Let’s keep on going deeper into how the engine works.
 
 ### Unreal Reflection System
 
@@ -354,7 +354,7 @@ Here is an example declaration of a **UCLASS**:
 
 You’ll first notice the inclusion of "MyClass.generated.h".  Unreal will generate all the reflection data and put it into this file.  You must include this file as the last include in the header file that declares your type.
 
-You’ll also have noticed that we can also add additional specifiers to the markup.  I’ve added some of the more common ones for demonstration.  These allow us to specify certain behaviour that our types have.
+You’ll also have noticed that we can also add additional specifiers to the markup.  I’ve added some of the more common ones for demonstration.  These allow us to specify certain behavior that our types have.
 
 * **Blueprintable** - This class can be extended by a Blueprint.
 * **BlueprintReadOnly **- This property can only be read from a Blueprint, and not written to.
@@ -373,7 +373,7 @@ There are too many specifiers to list here, so I’ll link to the docs for them:
 
 ### Object/Actor Iterators
 
-Object iterators are a very useful tool to iterate over all instances of a particular **UObject** type and it’s subclasses.
+Object iterators are a very useful tool to iterate over all instances of a particular **UObject** type and its subclasses.
 
 	// Will find ALL current UObjects instances
 	for (TObjectIterator<UObject> It; It; ++It)
@@ -475,7 +475,7 @@ When a **UObject** is garbage collected, all **UPROPERTY** references to it will
 		// Use SafeObject
 	}
 
-This is important since, as mentioned before, actors that have had **Destroy() **called on them are not removed until the garbage collector runs again.  You can check the **IsPendingKill()** method to see if a **UObject** is awaiting it’s deletion.  If that method returns **true**, you should consider the object dead and not use it.
+This is important since, as mentioned before, actors that have had **Destroy() **called on them are not removed until the garbage collector runs again.  You can check the **IsPendingKill()** method to see if a **UObject** is awaiting its deletion.  If that method returns **true**, you should consider the object dead and not use it.
 
 ### UStructs
 
@@ -684,7 +684,7 @@ Let’s say you were creating a grid-based board game and needed to store and qu
 	{
 	private:
 
-		// Using a TMap, we can refer to each piece by it’s position
+		// Using a TMap, we can refer to each piece by its position
 		TMap<FIntPoint, FPiece> Data;
 
 	public:
