#UE4doc #docs: changed "an" to "at" "It's" to "It is"

diff --git a/Engine/Documentation/Source/Programming/Plugins/Plugins.INT.udn b/Engine/Documentation/Source/Programming/Plugins/Plugins.INT.udn
--- a/Engine/Documentation/Source/Programming/Plugins/Plugins.INT.udn
+++ b/Engine/Documentation/Source/Programming/Plugins/Plugins.INT.udn
@@ -45,7 +45,7 @@ next to a category indicates how many plugins are available in that category.
 Plugins are displayed in the main list, along with their name and icon, the current version of the plugin, a helpful 
 description, the plugin's author (and optional web hyperlink), as well as whether the plugin is currently enabled or not.
 
-The search control an the top will allow you to search plugins displayed in the list by name.
+The search control at the top will allow you to search plugins displayed in the list by name.
 
 ![SearchingPlugins.png](SearchingPlugins.png)
 
@@ -229,7 +229,7 @@ This example plugin descriptor is from the **UObjectPlugin** [downloadable examp
 
 | Field name						| Info			| Description |
 | ---------------------------------	| -------------	| ----------- |
-| **FileVersion**					| _Required_	| Version of this plugin descriptor file itself. Its used for backwards compatibility as new features are added to the plugin system. You should usually set this to the latest version that is allowed by the version of the engine you are using. The latest version is currently **3**, and is the version of format that is documented here. We do not expect this version to change very frequently. In source code, you can look at **EProjectFileVersion** to see the actual value. If you require maximum compatibility with older versions of the engine, then you can use an older version of the format, but it is not recommended.
+| **FileVersion**					| _Required_	| Version of this plugin descriptor file itself. It is used for backwards compatibility as new features are added to the plugin system. You should usually set this to the latest version that is allowed by the version of the engine you are using. The latest version is currently **3**, and is the version of format that is documented here. We do not expect this version to change very frequently. In source code, you can look at **EProjectFileVersion** to see the actual value. If you require maximum compatibility with older versions of the engine, then you can use an older version of the format, but it is not recommended.
 | **Version**						| _Optional_    | Current version number of this build of your plugin. This value should always increase with future versions. This version number is not usually displayed to end-users. |
 | **VersionName**					| _Optional_    | Version of the plugin displayed in the editor UI. This is never used for any versioning checks and can be in whatever format you would like, however we suggest a simple Major.Minor format. You should always update the VersionName whenever the Version number has increased. |
 | **EngineVersion**					| _Optional_    | Minimum engine version necessary to be able to load this plugin. The plugin will fail to load when installed to projects using an engine version less than the value specified here. |
