#UE4doc #docs: sp: relevancies  made "states" plural

diff --git a/Engine/Documentation/Source/Programming/Rendering/Programming.INT.udn b/Engine/Documentation/Source/Programming/Rendering/Programming.INT.udn
--- a/Engine/Documentation/Source/Programming/Rendering/Programming.INT.udn
+++ b/Engine/Documentation/Source/Programming/Rendering/Programming.INT.udn
@@ -61,7 +61,7 @@ The primary classes are:
 |**USceneComponent**|Base class of any object that needs to be added to an FScene, like lights, meshes, fog, etc.|
 |**UPrimitiveComponent**| Base class of anything that can be rendered or interact with physics.  Also acts as the granularity of visibility culling and rendering property specification (casts shadows, etc).  Just like all UObjects, the game thread owns all variables and state and the rendering thread should not access it directly. |
 |**ULightComponent**| Represents a light source.  The Renderer is responsible for computing and adding its contribution to the scene. |
-|**FScene**|Renderer version of the UWorld.  Objects only exist to the renderer once they are added to the FScene, which is called registering a component.  The rendering thread owns all state on the FScene, the game thread cannot modify it directly.|
+|**FScene**|Renderer version of the UWorld.  Objects only exist to the renderer once they are added to the FScene, which is called registering a component.  The rendering thread owns all states in the FScene -the game thread cannot modify it directly.|
 |**FPrimitiveSceneProxy**|Renderer version of UPrimitiveComponent, mirrors UPrimitiveComponent state for the rendering thread.  Exists in the engine module and intended to be subclassed to support different types of primitives (skeletal, rigid, BSP, etc).  Implements some very important functions like GetViewRelevance, DrawDynamicElements, etc.|
 |**FPrimitiveSceneInfo**|Internal renderer state (specific to the FRendererModule implementation) that corresponds to a UPrimitiveComponent and FPrimitiveSceneProxy.  Exists in the renderer module, so the engine cannot see it.|
 |**FSceneView**| Engine representation of a single view into an FScene.  A scene can be rendered with different views in different calls to FSceneRenderer::Render (multiple editor viewports) or with multiple views in the same call to FSceneRenderer::Render (splitscreen in game).  A new View is constructed for each frame.|
@@ -124,7 +124,7 @@ FPrimitiveSceneProxy is the rendering thread version of UPrimitiveComponent that
 The renderer processes the scene in the order that it wants to composite data to the render targets. For example, the Depth only pass is rendered before the Base pass, so that HiZ will be populated to reduce shading cost in the base pass. This order is statically defined by the order pass functions are called in C++.
 #### Relevance
 
-FPrimitiveViewRelevance is the information on what effects (and therefore passes) are relevant to the primitive. A primitive may have multiple elements with different relevance, so FPrimitiveViewRelevance is effectively a logical OR of all the element's relevances. This means that a primitive can have both opaque and translucent relevance, or dynamic and static relevance; they are not mutually exclusive.
+FPrimitiveViewRelevance is the information on what effects (and therefore passes) are relevant to the primitive. A primitive may have multiple elements with different relevance, so FPrimitiveViewRelevance is effectively a logical OR of all the element's relevancies. This means that a primitive can have both opaque and translucent relevance, or dynamic and static relevance; they are not mutually exclusive.
 
 FPrimitiveViewRelevance also indicates whether a primitive needs to use the dynamic and/or static rendering path with bStaticRelevance and bDynamicRelevance.
 #### Drawing Policies
@@ -155,7 +155,7 @@ The static rendering path is implemented through static draw lists. Meshes are i
 
 The static rendering path moves a lot of work to attach time, which significantly speeds up scene traversal at rendering time. Static draw list rendering is about 3x faster on the rendering thread for Static Meshes, which allows a lot more Static Meshes in the scene. Because static draw lists cache data at attach time, they can only cache view independent state. Primitives that are rarely reattached but often rendered are good candidates for the static draw lists. 
 
-The static rendering path can expose bugs because of the way it only calls DrawShared once per state bucket. These bugs can be difficult to detect, since they depend on the rendering order and the attach order of meshes in the scene. Special viewmodes such as lighting only, unlit, etc will force all primitives to use the dynamic path, so if a bug goes away when forcing the dynamic rendering path, there is a good chance it is due to an incorrect implementation of a drawing policy's DrawShared and/or the Matches function.
+The static rendering path can expose bugs because of the way it only calls DrawShared once per state bucket. These bugs can be difficult to detect, since they depend on the rendering order and the attach order of meshes in the scene. Special view modes such as lighting only, unlit, etc will force all primitives to use the dynamic path, so if a bug goes away when forcing the dynamic rendering path, there is a good chance it is due to an incorrect implementation of a drawing policy's DrawShared and/or the Matches function.
 ## High level Rendering order
 
 Here is a description of the control flow when rendering a frame starting from FDeferredShadingSceneRenderer::Render:
