sp: view modes, successfully, sentence structure

diff --git a/Engine/Documentation/Source/Programming/Rendering/ShaderDevelopment/ShaderDevelopment.INT.udn b/Engine/Documentation/Source/Programming/Rendering/ShaderDevelopment/ShaderDevelopment.INT.udn
--- a/Engine/Documentation/Source/Programming/Rendering/ShaderDevelopment/ShaderDevelopment.INT.udn
+++ b/Engine/Documentation/Source/Programming/Rendering/ShaderDevelopment/ShaderDevelopment.INT.udn
@@ -101,7 +101,7 @@ If you want to step into the shader compiler DLL's directly from UE4 (CompileD3D
 
 #### Retrying on compile errors
 
-With r.ShaderDevelopmentMode enabled, you will get the opportunity to retry on shader compile error.  This is especially important for global shaders since it is a fatal error if they do not compile sucessfully.
+With r.ShaderDevelopmentMode enabled, you will get the opportunity to retry on shader compile error.  This is especially important for global shaders since it is a fatal error if they do not compile successfully.
 
 In debug, with the debugger attached, you will hit a breakpoint and get the compile error in the Visual Studio output window.  **You can then double-click the error log to be taken directly to the offending line.**
 
@@ -153,6 +153,6 @@ The [HLSL Cross Compiler](Programming/Rendering/ShaderDevelopment/HLSLCrossCompi
 
 ## AsyncCompute
 
-[AsyncCompute](Programming/Rendering/ShaderDevelopment/AsyncCompute) is a hardware feature that is available on some API and only exists one some GPUs.
-It allows to interleave different GPU work to better utilize the hardware units on the GPU more efficiently.
+[AsyncCompute](Programming/Rendering/ShaderDevelopment/AsyncCompute) is a hardware feature that is available on some APIs that work with certain GPUs.
+It allows interleaving to better utilize the hardware units in the GPU more efficiently.
 
