#UE4doc #docs: fix: "function" singular sp: elimination

diff --git a/Engine/Documentation/Source/Programming/Rendering/ShaderDevelopment/HLSLCrossCompiler/HLSLCrossCompiler.INT.udn b/Engine/Documentation/Source/Programming/Rendering/ShaderDevelopment/HLSLCrossCompiler/HLSLCrossCompiler.INT.udn
--- a/Engine/Documentation/Source/Programming/Rendering/ShaderDevelopment/HLSLCrossCompiler/HLSLCrossCompiler.INT.udn
+++ b/Engine/Documentation/Source/Programming/Rendering/ShaderDevelopment/HLSLCrossCompiler/HLSLCrossCompiler.INT.udn
@@ -13,17 +13,17 @@ In addition to producing GLSL code, the compiler packs global uniforms in to an
 
 ## Overview
 
-The main library entry point is HLSLCrossCompile. This functions performs all steps needed to generate GLSL code from the source HLSL with the requested options. A summary of each stage follows:
+The main library entry point is HLSLCrossCompile. This function performs all steps needed to generate GLSL code from the source HLSL with the requested options. A summary of each stage follows:
 
 |Operation| Description|
 |--|--|
 |Preprocessing|The code is run through a C-like preprocessor. This stage is optional and may be omitted by using the NoPreprocess flag. Unreal performs preprocessing using MCPP before compilation and therefore skips this step.|
 |Parsing|The HLSL source is parsed in to an abstract syntax tree. This is done in the function _mesa_hlsl_parse. The lexer and parser are generated by flex and bison respectively. See the section on parsing for more information.|
 |Compilation|The abstract syntax tree is compiled in to Mesa intermediate representation. This process happens in the function _mesa_ast_to_hir. During this stage, the compiler performs functions such as implicit conversions, function overload resolution, generating instructions for intrinsics, and so on.  A GLSL main entry point is generated. See GenerateGlslMain. This stage will add global declarations for input and output variables to the IR, compute the inputs for the HLSL entry point, call the HLSL entry point, and write outputs to the global output variables.|
-|Optimization| Several optimization passes are performed on the IR including function inlining, dead code elimination, constant propagation, elmination of common subexpressions, and so on. See OptimizeIR and especially do_optimization_pass for details.|
+|Optimization| Several optimization passes are performed on the IR including function inlining, dead code elimination, constant propagation, elimination of common subexpressions, and so on. See OptimizeIR and especially do_optimization_pass for details.|
 |Uniform packing| Global uniforms are packed in to arrays with mapping information retained so the engine may bind parameters to the relevant portion of the uniform array. See PackUniforms for details.|
 |Final optimization| After uniforms have been packed, a second round of optimizations is run on the IR to simplify the code generated when packing uniforms.|
-|Generate GLSL| Finally the optimized IR is converted to GLSL source code. The conversion from IR to GLSL is relatively straight-forward. In addition to producing defintions of all structs and uniform buffers and the source itself, a mapping table is written out in comments at the top of the file. This mapping table is parsed by Unreal to allow the binding of parameters. See GenerateGlsl and especially the ir_gen_glsl_visitor class for details.|
+|Generate GLSL| Finally the optimized IR is converted to GLSL source code. The conversion from IR to GLSL is relatively straight-forward. In addition to producing definitions of all structs and uniform buffers and the source itself, a mapping table is written out in comments at the top of the file. This mapping table is parsed by Unreal to allow the binding of parameters. See GenerateGlsl and especially the ir_gen_glsl_visitor class for details.|
 
 ## Parsing
 
