#UE4doc #docs: sp: utilized, synchronize,

diff --git a/Engine/Documentation/Source/Programming/Rendering/ShaderDevelopment/AsyncCompute/AsyncCompute.INT.udn b/Engine/Documentation/Source/Programming/Rendering/ShaderDevelopment/AsyncCompute/AsyncCompute.INT.udn
--- a/Engine/Documentation/Source/Programming/Rendering/ShaderDevelopment/AsyncCompute/AsyncCompute.INT.udn
+++ b/Engine/Documentation/Source/Programming/Rendering/ShaderDevelopment/AsyncCompute/AsyncCompute.INT.udn
@@ -6,11 +6,11 @@ Version: 4.9
 
 ## Introduction
 
-The RHI now supports asynchronous compute for Xbox One. This is a good way to utilise unused GPU resources (CUs, registers and bandwidth),
+The RHI now supports asynchronous compute for Xbox One. This is a good way to utilize unused GPU resources (CUs, registers and bandwidth),
 by running dispatch() calls asynchronously with the rendering. Async compute uses a separate context, and we provide RHI functions to
-synchronise the rendering and compute contexts.
+synchronize the rendering and compute contexts.
 Dr PIX is useful for identifying areas which might benefit from async compute. For example, if half the CUs are unused during a particular
-rendering pass, those CUs could potentially be utilised by an asynchronous compute job.
+rendering pass, those CUs could potentially be utilized by an asynchronous compute job.
 Async compute has some restrictions:
 
  * Buffers with UAV counters are not supported (this is a limitation of the XDK, and will generate a D3D warning)
@@ -26,7 +26,7 @@ Async compute has some restrictions:
    for the job (via ID3D11XComputeContext::SetLimits). There are currently two priorities available, AsyncComputePriority_High and
    AsyncComputePriority_Default.
  * **RHIEndAsyncComputeJob_DrawThread**
-   Ends an async compute job. Returns a 32-bit Fence Index which can be used for synchronisation, or -1 if compute is disabled or
+   Ends an async compute job. Returns a 32-bit Fence Index which can be used for synchronization, or -1 if compute is disabled or
    there is no active compute job
  * **RHIGraphicsWaitOnAsyncComputeJob**
    Inserts an command in the graphics immediate context to block until an async job is complete. Passing -1 causes this to early-out.
