#UE4doc #docs: sp: conjures, handle

diff --git a/Engine/Documentation/Source/Programming/Slate/Architecture/SlateArchitecture.INT.udn b/Engine/Documentation/Source/Programming/Slate/Architecture/SlateArchitecture.INT.udn
--- a/Engine/Documentation/Source/Programming/Slate/Architecture/SlateArchitecture.INT.udn
+++ b/Engine/Documentation/Source/Programming/Slate/Architecture/SlateArchitecture.INT.udn
@@ -87,7 +87,7 @@ The STextBlock uses a delegate called Text for this purpose.
 [/REGION]
 
 Consider also that in the example above the Text is a string, while the `framerate` is most likely stored as a `float` or an `integer`. The use of delegates offers us the flexibility to perform the conversion whenever the value is read.
-This immediately conjurs up performance concerns, which are addressed in the [Performance Considerations](#PerformanceConsiderations) section below.
+This immediately conjures up performance concerns which are addressed in the [Performance Considerations](#PerformanceConsiderations) section below.
 
 
 
@@ -128,7 +128,7 @@ Consider the following observations:
 * UI complexity is bounded by the number of live widgets.
 * Scrolling content is virtualized whenever possible; this mostly avoids live widgets off-screen.
 	* Large numbers of off-screen widgets can easily tank Slate performance.
-* Assumption: users with big screens have beefy machines to drive those screens; they can hangle a large number of widgets.
+* Assumption: users with big screens have beefy machines to drive those screens; they can handle a large number of widgets.
 
 
 ## Invalidation vs Polling
@@ -141,7 +141,7 @@ Invalidation is - as a rule - reserved for infrequent, low-granularity events.
 
 Consider the example of the Blueprint Editor, which displays nodes on a graph.
 When an update is requested, all the **Graph** Panel widgets are cleared and re-created.
-This is preferrable to fine-grain invalidation because it is simpler and more maintainable.
+This is preferable to fine-grain invalidation because it is simpler and more maintainable.
 
 ## Child Slots
 
