#UE4doc #docs: instances of its to it is where appropriate

diff --git a/Engine/Documentation/Source/Programming/Slate/DetailsCustomization/DetailsCustomization.INT.udn b/Engine/Documentation/Source/Programming/Slate/DetailsCustomization/DetailsCustomization.INT.udn
--- a/Engine/Documentation/Source/Programming/Slate/DetailsCustomization/DetailsCustomization.INT.udn
+++ b/Engine/Documentation/Source/Programming/Slate/DetailsCustomization/DetailsCustomization.INT.udn
@@ -104,7 +104,7 @@ The result:
 
 #### Notes About Multibox Style Layout
 
-* It is not very powerful, but more features will be added as needed.  Its currently designed to be for quick reorganization.
+* It is not very powerful, but more features will be added as needed.  It is currently designed to be for quick reorganization.
 * The slate layout will require more advanced access to properties, specifically property handles, especially if you need to customize their look.
 
 
@@ -134,7 +134,7 @@ From here, you can read and write the value of that property and/or pass it to a
 Other notes:
 
 * The handle returned from `GetProperty` may be invalid if the property could not be found or is not going to appear in the details view.  Check `IsValidHandle()` to be sure.  Calling functions on invalid handles will not crash.
-* You should not store property handles outside of your layout class unless they are weak pointers.  Internally the data they access is a weak pointer so it will not crash if you try to set or get a value on an invalid property but you have a reference to a useless object if you store it and its not cleaned up.
+* You should not store property handles outside of your layout class unless they are weak pointers.  Internally the data they access is a weak pointer so it will not crash if you try to set or get a value on an invalid property but you have a reference to a useless object if you store it and it is not cleaned up.
 * If you try to read / write access a value type for an unsupported property (e.g., a `float` for a `String` property ), it will fail but no data will be corrupted.
 
 
@@ -331,7 +331,7 @@ The result:
 ### Customization Dos and Dont's
 
 * Do check error cases when customizing properties and reading/writing values.  Remember that details views can often be viewing multiple objects at once where each object has different values.  Customized properties should be able to handle the common case of multiple values.
-* Do store any data about the selection on your customization class.  Some non-property details will need the selected Actors for a customization.  You can get the selected Actors from `IDetailLayoutBuilder`.  You can store this selection set or anything selection sensitive on your customization class.  Its guaranteed to be around while the selection remains the same in its details view.
+* Do store any data about the selection on your customization class.  Some non-property details will need the selected Actors for a customization.  You can get the selected Actors from `IDetailLayoutBuilder`.  You can store this selection set or anything selection sensitive on your customization class.  It is guaranteed to be around while the selection remains the same in its details view.
 * Do not use `FActorIterator`, `FSelectedActorIterator`, or `GEditor->GetSelectedActorIterator`.  Remember that the **Details** panel can be locked and these things operate on global selection sets or lists of Actors which is not the same as the selected Actors in a **Details** panel if it is locked!  Using these will access different data.  You can get the list of selected Actors valid for you from `IDetailLayoutBuilder`.
 * Do not hold a strong reference to your layout class or property handles (you should not need to anyway).  Remember that details views (especially level editor ones) can change at any time based on user selection so any references you have to layout classes can easily become invalid.  The shared pointers to layout classes are checked for uniqueness when customizing details to prevent this from happening.
 
