#UE4doc #docs: added ' to "user's"

diff --git a/Engine/Documentation/Source/Programming/Slate/DetailsCustomization/DetailsCustomization.INT.udn b/Engine/Documentation/Source/Programming/Slate/DetailsCustomization/DetailsCustomization.INT.udn
--- a/Engine/Documentation/Source/Programming/Slate/DetailsCustomization/DetailsCustomization.INT.udn
+++ b/Engine/Documentation/Source/Programming/Slate/DetailsCustomization/DetailsCustomization.INT.udn
@@ -267,7 +267,7 @@ In addition to `SProperty` there are some other customization widgets that can b
 
 `SFilterableDetail` is a widget that does not draw anything but filters everything in its content slot when a user types in the search box of the details view.  This widget is useful for non-property based details. `SProperty` is already filtered so you do not need to set up an `SFilterableDetail` for those unless you want to group their filtering.
 
-	// Create a widget which will filter out everything in the content slot when "Create Blocking Volume" is not matched with a users search term
+	// Create a widget which will filter out everything in the content slot when "Create Blocking Volume" is not matched with a user's search term
 	// Note: The second parameter is the localized search term that matches the filter and the third parameter is the category where this filter should reside
 	SNew( SFilterableDetail, NSLOCTEXT("StaticMeshDetails", "BlockingVolumeMenu", "Create Blocking Volume"), &StaticMeshCategory )
 	.Content()
