#UE4doc #docs: sp: particular, overlaid

diff --git a/Engine/Documentation/Source/Programming/Slate/InGameUI/SlateInGameUI.INT.udn b/Engine/Documentation/Source/Programming/Slate/InGameUI/SlateInGameUI.INT.udn
--- a/Engine/Documentation/Source/Programming/Slate/InGameUI/SlateInGameUI.INT.udn
+++ b/Engine/Documentation/Source/Programming/Slate/InGameUI/SlateInGameUI.INT.udn
@@ -8,7 +8,7 @@ Version:4.4
 
 Slate widgets can be used in-game to create heads up displays (HUDs) or other user interface (UI) elements, such 
 as menus. You will generally create one or more _container_ widgets, each of which may be composed of several 
-other types of widgets, that are responsible for a partcular aspect of the UI. 
+other types of widgets that are responsible for a particular aspect of the UI. 
 
 <!-- HUD widget w/ menu overlay -->
 
@@ -33,7 +33,7 @@ Each of these container widgets can then be added or removed depending on the ci
 
 ## Displaying Widgets
 
-In order to display a Slate widget in your game, it must be added to the game's viewport. Overlayed widgets 
+In order to display a Slate widget in your game, it must be added to the game's viewport. Overlaid widgets 
 are ordered according to the Z-order specified when adding them, with larger Z-order values appearing on top 
 of smaller Z-orders.
 
