#UE4doc #docs: duplicate word: the replaced "an" with "on" sp: cannot

diff --git a/Engine/Documentation/Source/Programming/Slate/Widgets/SlateWidgets.INT.udn b/Engine/Documentation/Source/Programming/Slate/Widgets/SlateWidgets.INT.udn
--- a/Engine/Documentation/Source/Programming/Slate/Widgets/SlateWidgets.INT.udn
+++ b/Engine/Documentation/Source/Programming/Slate/Widgets/SlateWidgets.INT.udn
@@ -23,7 +23,7 @@ The following arguments are not on every single widget, but they are on most wid
 * Content - This will specify what widget should be placed in the content section of the widget, if applicable.
 * ReadOnly - This will prevent this widget from being editable if _true_.
 * Style - This will specify the style of images or text font used by the widget. How this is applicable varies by widget.
-* [...]ColorAndOpacity - The color and opacity of the widget or of the the specified system.
+* [...]ColorAndOpacity - The color and opacity of the widget or of the specified system.
 * Padding - The padding of a widget amount of spacing in slate units around the left, top, right, and bottom parts of the widget within its parent. These can be specified as a single value for all four parts, or as a horizontal and vertical value, or as four separate values.
 * HAlign - The horizontal alignment of content within the widget.
 * VAlign - The vertical alignment of content within the widget.
@@ -138,7 +138,7 @@ SUniformGridPanel is a panel that distributes its child widgets evenly both vert
 
 ### Wrap Boxes
 
-SWrapBox is a box that lays out widgets horizontally until these widgets exceed a width, at which point they will be placed an the next line, and so on. Below is an example of such.
+SWrapBox is a box that lays out widgets horizontally until these widgets exceed a width, at which point they will be placed on the next line, and so on. Below is an example of such.
 
 	SNew(SWrapBox)
 	.PreferredWidth( 300.f )
@@ -196,7 +196,7 @@ Radio buttons are check boxes in slate, because check boxes require a delegate f
 
 ### Menus
 
-To create menus or toolbars, use multiboxes. Commands should use Slate's UI_COMMAND system, though this can not be done with dynamically generated buttons/controls.
+To create menus or toolbars, use multiboxes. Commands should use Slate's UI_COMMAND system, though this cannot be done with dynamically generated buttons/controls.
 
 To create a menu, create an FMenuBarBuilder with an FUICommandList passed in. Finally, you can call MakeWidget() on the menu bar builder to get a widget to place.
 
