#UE4doc #docs: sp: user's

diff --git a/Engine/Documentation/Source/Programming/Slate/Overview/SlateOverview.INT.udn b/Engine/Documentation/Source/Programming/Slate/Overview/SlateOverview.INT.udn
--- a/Engine/Documentation/Source/Programming/Slate/Overview/SlateOverview.INT.udn
+++ b/Engine/Documentation/Source/Programming/Slate/Overview/SlateOverview.INT.udn
@@ -179,7 +179,7 @@ The layout above creates the UI below:
 
 ## User Driven Layout
 
-Slate's **docking framework** puts the power into the users hands by allowing its tabbed panels to be rearranged and docked into virtually any layout imaginable. The ability to customize their environment lets the user work with the tools they want to work with the way they want to work with them.
+Slate's **docking framework** puts the power into the user's hands by allowing its tabbed panels to be rearranged and docked into virtually any layout imaginable. The ability to customize their environment lets the user work with the tools they want to work with the way they want to work with them.
 
 **Floating Tab**
 
