#UE4doc #docs: Fixed title sp: "PhysicsConstraints"

diff --git a/Engine/Documentation/Source/Resources/ContentExamples/Physics/1_5/1_5.INT.udn b/Engine/Documentation/Source/Resources/ContentExamples/Physics/1_5/1_5.INT.udn
--- a/Engine/Documentation/Source/Resources/ContentExamples/Physics/1_5/1_5.INT.udn
+++ b/Engine/Documentation/Source/Resources/ContentExamples/Physics/1_5/1_5.INT.udn
@@ -1,5 +1,5 @@
 Availability: Public
-Title:1.5 - PhysicsContraints
+Title:1.5 - PhysicsConstraints
 Crumbs:%ROOT%, Resources, Resources/ContentExamples, Resources/ContentExamples/Physics
 Description:An overview of the Physics example level, example 1.5: PhysicsContraints
 
