#UE4doc #docs: sp: Tessellation in the title

diff --git a/Engine/Documentation/Source/Resources/ContentExamples/MaterialNodes/1_12/1_12.INT.udn b/Engine/Documentation/Source/Resources/ContentExamples/MaterialNodes/1_12/1_12.INT.udn
--- a/Engine/Documentation/Source/Resources/ContentExamples/MaterialNodes/1_12/1_12.INT.udn
+++ b/Engine/Documentation/Source/Resources/ContentExamples/MaterialNodes/1_12/1_12.INT.udn
@@ -1,7 +1,7 @@
 Availability: Public
-Title:1.12 - Tesselation Multiplier
+Title:1.12 - Tessellation Multiplier
 Crumbs:%ROOT%, Resources, Resources/ContentExamples, Resources/ContentExamples/Materials
-Description:An overview of the Material Nodes example level, example 1.12: Tesselation Multiplier.
+Description:An overview of the Material Nodes example level, example 1.12: Tessellation Multiplier.
 Related: Engine/Rendering/Materials/PhysicallyBased
 Related: Engine/Rendering/Materials
 Related: Engine/Rendering/Materials/ExpressionReference
