#UE4doc #docs: sp: dependent, dependency

diff --git a/Engine/Documentation/Source/Programming/UnrealArchitecture/Actors/Ticking/ActorTicking.INT.udn b/Engine/Documentation/Source/Programming/UnrealArchitecture/Actors/Ticking/ActorTicking.INT.udn
--- a/Engine/Documentation/Source/Programming/UnrealArchitecture/Actors/Ticking/ActorTicking.INT.udn
+++ b/Engine/Documentation/Source/Programming/UnrealArchitecture/Actors/Ticking/ActorTicking.INT.udn
@@ -132,9 +132,9 @@ need to be in different groups are added to a list that manages when they will b
 ### Why Is It So Important?
 
 Understanding the ticking and rendering order is essential to avoid off-by-1-frame issues and to work around dependencies in calling order.  For instance, the camera is updated after all ticking is
- done, so anything that is camera relative (such as a first-person attachment or a laser dot) must be done in PreRender.  But if there is, for instance, an _animation_ dependancy on camera position, 
- there **will be** an off-by-1-frame issue, because animation occurs during Pre-physics, Pawn movement occurs during physics, camera (typically) is dependant on Pawn position, and thus the animation 
- dependancy creates a loop.  To get around these issues, it may be necessary to calculate something more than once per frame.  Moving the camera tick as well as the animation tick to post physics 
+ done, so anything that is camera relative (such as a first-person attachment or a laser dot) must be done in PreRender.  But if there is, for instance, an _animation_ dependency on camera position, 
+ there **will be** an off-by-1-frame issue, because animation occurs during Pre-physics, Pawn movement occurs during physics, camera (typically) is dependent on Pawn position, and thus the animation 
+ dependency creates a loop.  To get around these issues, it may be necessary to calculate something more than once per frame.  Moving the camera tick as well as the animation tick to post physics 
  will allow the animation to have correctly updated camera position (but the camera will be updated twice per frame unless coded specifically not to).
 
  - Jordan Weitz
