#UE4doc #docs: sp: capitilized first word of sentence. sp: situations

diff --git a/Engine/Documentation/Source/Programming/UnrealArchitecture/Reference/Classes/GameplayClasses.INT.udn b/Engine/Documentation/Source/Programming/UnrealArchitecture/Reference/Classes/GameplayClasses.INT.udn
--- a/Engine/Documentation/Source/Programming/UnrealArchitecture/Reference/Classes/GameplayClasses.INT.udn
+++ b/Engine/Documentation/Source/Programming/UnrealArchitecture/Reference/Classes/GameplayClasses.INT.udn
@@ -81,7 +81,7 @@ Some constructors may be located in a special "constructors" file on a per-modul
 result of an automatic conversion process from the previous use of a `DEFAULTS` block to the use of constructors. Over time, these will be moved to their respective class source files.
 [/REGION]
 
-It is also possible to place the constructor inline in the class header file. however, if the constructor is in the class header file, the UClass must be declared with the `CustomConstructor` 
+It is also possible to place the constructor inline in the class header file. However, if the constructor is in the class header file, the UClass must be declared with the `CustomConstructor` 
 specifier as this prevents the automatic code generator from creating a constructor declaration in the header.
 
 #### Constructor Format
@@ -253,7 +253,7 @@ Creating sub-objects, especially components, and modifying sub-objects belonging
 		bNoDelete = true;
 	}
 
-Finding components belonging to the parent class is generally not necessary. The component should be created and assigned to a property in the parent's constructor, so the component can be accessed using that property. If this is not the case, the best solution is probably to update the parent's constructor to make it so. The `ConstructorHelpers::FindComponent` helper is provided as a backup for sitations where the standard methods are not desirable.
+Finding components belonging to the parent class is generally not necessary. The component should be created and assigned to a property in the parent's constructor, so the component can be accessed using that property. If this is not the case, the best solution is probably to update the parent's constructor to make it so. The `ConstructorHelpers::FindComponent` helper is provided as a backup for situations where the standard methods are not desirable.
 
 #### Array Manipulation
 
