#UE4doc #docs: sp: Non Transient

diff --git a/Engine/Documentation/Source/Programming/UnrealArchitecture/Reference/Classes/Specifiers/NonTransient/ClassSpecifier_NonTransient.INT.udn b/Engine/Documentation/Source/Programming/UnrealArchitecture/Reference/Classes/Specifiers/NonTransient/ClassSpecifier_NonTransient.INT.udn
--- a/Engine/Documentation/Source/Programming/UnrealArchitecture/Reference/Classes/Specifiers/NonTransient/ClassSpecifier_NonTransient.INT.udn
+++ b/Engine/Documentation/Source/Programming/UnrealArchitecture/Reference/Classes/Specifiers/NonTransient/ClassSpecifier_NonTransient.INT.udn
@@ -1,4 +1,4 @@
-Title: NonTranisent
+Title: NonTransient
 Crumbs:%ROOT%, Programming, Programming/UnrealArchitecture/Reference, Programming/UnrealArchitecture/Reference/Classes, Programming/UnrealArchitecture/Reference/Classes/Specifiers
 Description:
 Availability:Public
