#UE4doc #docs: removed unneeded comma

diff --git a/Engine/Documentation/Source/Programming/UnrealArchitecture/Reference/Functions/Specifiers/Unreliable/FunctionSpecifier_Unreliable.INT.udn b/Engine/Documentation/Source/Programming/UnrealArchitecture/Reference/Functions/Specifiers/Unreliable/FunctionSpecifier_Unreliable.INT.udn
--- a/Engine/Documentation/Source/Programming/UnrealArchitecture/Reference/Functions/Specifiers/Unreliable/FunctionSpecifier_Unreliable.INT.udn
+++ b/Engine/Documentation/Source/Programming/UnrealArchitecture/Reference/Functions/Specifiers/Unreliable/FunctionSpecifier_Unreliable.INT.udn
@@ -4,4 +4,4 @@ Description:
 Availability:Public
 Version: 4.5
 
-The function is replicated over the network, but can fail due to bandwidth limitations or network errors. Only valid when used in conjunction with Client or Server.
\ No newline at end of file
+The function is replicated over the network but can fail due to bandwidth limitations or network errors. Only valid when used in conjunction with Client or Server.
\ No newline at end of file
