#UE4doc #docs: duplicate words: of, the & with, sp: "check-in" following

diff --git a/Engine/Documentation/Source/Programming/UnrealArchitecture/StringHandling/CharacterEncoding/CharacterEncoding.INT.udn b/Engine/Documentation/Source/Programming/UnrealArchitecture/StringHandling/CharacterEncoding/CharacterEncoding.INT.udn
--- a/Engine/Documentation/Source/Programming/UnrealArchitecture/StringHandling/CharacterEncoding/CharacterEncoding.INT.udn
+++ b/Engine/Documentation/Source/Programming/UnrealArchitecture/StringHandling/CharacterEncoding/CharacterEncoding.INT.udn
@@ -18,13 +18,13 @@ Assumed knowledge: [The Absolute Minimum Every Software Developer Absolutely, Po
 
 There are several formats that can be used to represent text and strings. Understanding these formats and their inherent pros and cons can help in making decisions on what formats to use in your projects.
 
-These are not the the technical definitions of formats, but rather simplified versions suitable for this page.
+These are not the technical definitions of formats, but rather simplified versions suitable for this page.
 
 $ **ASCII** : Characters between 32 and 126 inclusive, and 0, 9, 10, and 13. (P4 type text) (This is validated with a P4 trigger on check-in)
-$ **ANSI** : ASCII and the current codepage. (e.g. Western European high ASCII) (needs to be stored as binary on the P4 server)
+$ **ANSI** : ASCII and the current codepage (e.g. Western European high ASCII) needs to be stored as binary on the P4 server.
 $ **UTF-8** : A string made up of single bytes which can use special character sequences to get non-ANSI characters. (a superset of ASCII) (P4 type Unicode)
 $ **UTF-16** : 
-A string made up of 2 bytes per character with a [BOM](http://en.wikipedia.org/wiki/Byte-order_mark). (although can go to 4 bytes with astral characters) (P4 type UTF-16) (This is validated with a P4 trigger on checkin)
+A string made up of 2 bytes per character with a [BOM](http://en.wikipedia.org/wiki/Byte-order_mark). (although can go to 4 bytes with astral characters) (P4 type UTF-16) (This is validated with a P4 trigger on check-in)
 
 
 ### The Case for Binary
@@ -50,7 +50,7 @@ A string made up of 2 bytes per character with a [BOM](http://en.wikipedia.org/w
 | Simple access to all characters we will ever need. | Has a different memory profile for Asian languages. |
 | Uses less memory. | P4 type Unicode is not enabled on our Perforce server. |
 | Is a superset of ASCII; a plain ASCII string is a perfectly valid UTF-8 string. | String operations more complicated; have to parse the string to do something as simple as a length calculation. |
-| Still works when the game detects the string is ASCII and outputs it as such. | MSDev does not handle anything other than ASCII very well in Asian regions. This is why we validate text as ASCII during checkin. |
+| Still works when the game detects the string is ASCII and outputs it as such. | MSDev does not handle anything other than ASCII very well in Asian regions. This is why we validate text as ASCII during check-in. |
 | If we did have a Unicode enabled server, the files would be merge-able and exclusive checkout would not be required. | |
 | Can detect whether a string is UTF-8 by parsing it (with or without a BOM). | |
 
@@ -61,8 +61,8 @@ A string made up of 2 bytes per character with a [BOM](http://en.wikipedia.org/w
 | --- | --- |
 | Simple access to all characters we will ever need. | Uses more memory. |
 | Simple. Memory usage is twice the number of characters (for characters we use, which are all in the [Basic Multilingual Plane](http://en.wikipedia.org/wiki/Mapping_of_Unicode_character_planes)). | Difficult to detect this format if it does not have a BOM. |
-| Simple. String operations can split/combine without having to parse the strings. | Does not work when the game detects the string is ASCII and outputs it as such (this is now detected on checkin with the UTF-16 validator). |
-| Same as the format used in game, no translation, parsing or memory operations required. | MSDev does not handle anything other than ASCII very well in Asian regions. This is why we validate text as ASCII during checkin. |
+| Simple. String operations can split/combine without having to parse the strings. | Does not work when the game detects the string is ASCII and outputs it as such (this is now detected on check-in with the UTF-16 validator). |
+| Same as the format used in game, no translation, parsing or memory operations required. | MSDev does not handle anything other than ASCII very well in Asian regions. This is why we validate text as ASCII during check-in. |
 | Merge-able. Exclusive checkout is not required. | |
 | C# uses UTF-16 internally. | |
 
@@ -114,7 +114,7 @@ In general, we do not recommend string literals inside C++ source code and we re
 
 #### C++ Source Code
 
-UTF-8 or default Windows encoding. MSVC, the Xbox360 compiler and gcc should all be happy with with UTF-8 encoded source files. Latin-1 encoded files with characters with the high bit set, for example copyright, trademark or degree symbols should be avoided in source code where possible because the encoding will break on systems with different locales. Some instances of this in 3rd party software are unavoidable (eg copyright notices) so for MSVC we disable warning 4819, which would otherwise occur when compiling on Asian Windows.
+UTF-8 or default Windows encoding. MSVC, the Xbox360 compiler and gcc should all be happy with UTF-8 encoded source files. Latin-1 encoded files with characters with the high bit set, for example copyright, trademark or degree symbols should be avoided in source code where possible because the encoding will break on systems with different locales. Some instances of this in 3rd party software are unavoidable (eg copyright notices) so for MSVC we disable warning 4819, which would otherwise occur when compiling on Asian Windows.
 
 
 
@@ -127,7 +127,7 @@ UTF-8 or default Windows encoding. MSVC, the Xbox360 compiler and gcc should all
     * People can check in ASCII, UTF-8, UTF-16 and it will work in engine.
     * However, binary files cannot be merged, so if the files are not marked as exclusive checkout, changes will be stomped upon.
 * If you use 'UTF-16', make sure no one checks in a file that is not UTF-16
-    * We have a Perforce trigger that disallows checking in a non UTF-16 as UTF-16
+    * We have a Perforce trigger that disallows check-ing in a non UTF-16 as UTF-16
         * //depot/UnrealEngine3/Development/Tools/P4Utils/CheckUTF16/
 * The 'Unicode' type is UTF-8, and of no use to us here.
 
@@ -231,7 +231,7 @@ __UTF-8 without BOM__ (Some text editors describe BOM as signature)
   
 Take care for compiling C++ source code on East Asian code page CP949 (Korean), CP932 (Japanese), CP936 (Simplified Chinese) or CP950 (Traditional Chinese) Windows, if that source code has an East Asian character stored as UTF-8.
 
-UTF-8 character encoding uses three bytes for East Asian characters: 0xE0-0xEF for the first byte, 0x80-0xBF for the second byte and 0x80-0xBF for the third byte. Without the BOM, East Asian Windows' default encoding recognizes the three UTF-8 encoded bytes and the following byte as two 2-byte East Asian encoded characters, pairing of first and second bytes for one first East Asian character, and the third byte and folowing byte paired to form the second East Asian character.  
+UTF-8 character encoding uses three bytes for East Asian characters: 0xE0-0xEF for the first byte, 0x80-0xBF for the second byte and 0x80-0xBF for the third byte. Without the BOM, East Asian Windows' default encoding recognizes the three UTF-8 encoded bytes and the following byte as two 2-byte East Asian encoded characters, pairing of first and second bytes for one first East Asian character, and the third byte and following byte paired to form the second East Asian character.  
 Problems can occur if the character following the UTF-8 encoded three bytes has special meaning in the string literals or comments.
 
 Eg In in-line comment:  
@@ -248,7 +248,7 @@ The compiler on East Asian code page Windows recognizes the last byte of the UTF
 This behavior is very dangerous and it is difficult to find the missing code.
 
 In single-line comment:  
-Using backslash '\' at the end of of an East Asian comment causes hard-to-find bugs or errors without missing lines.
+Using backslash '\' at the end of an East Asian comment causes hard-to-find bugs or errors without missing lines.
 
     
         // OddNumberOfEastAsianCharacterComment\
