removed unneccessary comma, "a" to "an"

diff --git a/Engine/Documentation/Source/Programming/UnrealArchitecture/TMap/TMap.INT.udn b/Engine/Documentation/Source/Programming/UnrealArchitecture/TMap/TMap.INT.udn
--- a/Engine/Documentation/Source/Programming/UnrealArchitecture/TMap/TMap.INT.udn
+++ b/Engine/Documentation/Source/Programming/UnrealArchitecture/TMap/TMap.INT.udn
@@ -22,7 +22,7 @@ TMap can also take an optional allocator to control the memory allocation behavi
 
 The final TMap template parameter is KeyFuncs, which tells the map how to retrieve the key from the element type, how to compare two keys for equality and how to hash the key.  These have defaults which will just return a reference to the key, use operator== for equality and the non-member GetTypeHash function for hashing.  If your key type supports these functions, it is usable as a map key without the need to provide a custom KeyFuncs.
 
-Unlike TArray, the relative order of TMap elements in memory cannot be relied upon and iterating over the elements is likely to return them in a different order from that which they were added.  Elements are also unlikely to be laid out contiguously in memory.  The backing data structure of a map is a sparse array, which is an array with 'holes'.  As elements are removed from the map, holes in the sparse array will appear, which are then filled by the addition of future elements.  However, even though TMap doesn't shuffle elements to fill holes, pointers to map elements may still be invalidated, as the entire storage can be reallocated when it is full and new elements are added.
+Unlike TArray, the relative order of TMap elements in memory cannot be relied upon and iterating over the elements is likely to return them in a different order from that which they were added.  Elements are also unlikely to be laid out contiguously in memory.  The backing data structure of a map is a sparse array, which is an array with 'holes'.  As elements are removed from the map, holes in the sparse array will appear which are then filled by the addition of future elements.  However, even though TMap doesn't shuffle elements to fill holes, pointers to map elements may still be invalidated, as the entire storage can be reallocated when it is full and new elements are added.
 
 ## Creating and filling a map
 
