#UE4doc #docs: sp: cool down, Defense

diff --git a/Engine/Documentation/Source/Resources/SampleGames/StrategyGame/StrategyGame.INT.udn b/Engine/Documentation/Source/Resources/SampleGames/StrategyGame/StrategyGame.INT.udn
--- a/Engine/Documentation/Source/Resources/SampleGames/StrategyGame/StrategyGame.INT.udn
+++ b/Engine/Documentation/Source/Resources/SampleGames/StrategyGame/StrategyGame.INT.udn
@@ -161,7 +161,7 @@ is carried out with the aid of _Blueprint Interfaces_, **Interface_Auto_Arbalest
 
 The flamethrower does not shoot projectiles like the other turret types.  Instead, it burns all enemies in the flame area.  It is possible for the player to charge the flamethrower by clicking on it and holding; 
 depending on how long the mouse button is held, the flame released after charging can do up to three times more damage, as well as slowing down any enemy Pawns hit by the charged fire.  If the player charges up the flamethrower, 
-there is a small cooldown afterwards and then the regular fire attack continues.  
+there is a small cool down afterwards and then the regular fire attack continues.  
 
 ## Camera
 
@@ -211,8 +211,8 @@ This text is created using a simple `STextBlock` widget with delegates for Visib
 
 
 The main menu is contained in the **StrategyMenu** map, which loads the main menu specific HUD.  The menu is Slate-based, with `SStrategyMenuWidget` being responsible for the main menu animations,
-layout, and event handling.  The `SStrategyMenuItem` class inherits from the `SStrategyButtonWidget` used in the In-Game HUD and describes a single menu item.  Each menu item, and events attached to the items,
-are defined in the `StrategyMenuHUD`.  
+layout, and event handling.  The `SStrategyMenuItem` class inherits from the `SStrategyButtonWidget` used in the In-Game HUD and describes a single menu item.  Each menu item (and events attached to the items)
+is defined in the `StrategyMenuHUD`.  
 
 To go back to previous menus, an array of shared pointers to menus is stored in the `MenuHistory` variable. This variable acts like a stack to hold previously viewed menus, so that it is easy to go back while
 also removing the requirement to store the parent of a menu so that menus can be reused in multiple places.
@@ -251,7 +251,7 @@ property of the enemy brewery's `StrategyAIDirector`.
 
 There are five waves of enemies, each with different combinations of fast, normal, and heavy enemy Pawns.  At the beginning of a wave, the **Show Wave Title** node displays the wave number.  Then, the first enemy 
 spawn of the wave is called.  There are two types of delays after a spawn: a timer delay set by a **Delay** node, or a Pawn-based delay set by a **WaitForWaveMacro**. The **WaitForWaveMacro** macro continually checks to see how many enemy Pawns are alive
-and will not let execution leave the macro until either the delay time has expired or all the enemy Pawns are dead. After all the spawns for a wave are complete, and all the enemy Pawns for that wave are dead (or two minutes have passed), 
+and will not let execution leave the macro until either the delay time has expired or all of the enemy Pawns are dead. After all the spawns for a wave are complete and all of the enemy Pawns for that wave are dead (or two minutes have passed), 
 a _Remote Event_ node is used to call the _Custom Event_ for the next wave.
 
 
