#UE4doc #docs: changed include to includes

diff --git a/Engine/Documentation/Source/Programming/UnrealBuildSystem/ProjectFileGenerator/ProjectFileGenerator.INT.udn b/Engine/Documentation/Source/Programming/UnrealBuildSystem/ProjectFileGenerator/ProjectFileGenerator.INT.udn
--- a/Engine/Documentation/Source/Programming/UnrealBuildSystem/ProjectFileGenerator/ProjectFileGenerator.INT.udn
+++ b/Engine/Documentation/Source/Programming/UnrealBuildSystem/ProjectFileGenerator/ProjectFileGenerator.INT.udn
@@ -74,7 +74,7 @@ There are certainly both pros and cons, but here are some important reasons why
 
 * UE4 programming methodology incorporates many sub-modules, and it was very important that we made it as easy as possible for programmers to add new modules.
 
-* The project generator emits highly-accurate definitions and include paths which are used by Visual Studio IntelliSense while you are working in UE4 code.
+* The project generator emits highly-accurate definitions and includes paths which are used by Visual Studio IntelliSense while you are working in UE4 code.
 
 * Setting up a new project is much easier when project files are automatically generated.
 
