#UE4doc #docs: duplicate word: "support" sp: executable

diff --git a/Engine/Documentation/Source/Programming/UnrealBuildSystem/Configuration/Windows/BuildConfigurationWindows.INT.udn b/Engine/Documentation/Source/Programming/UnrealBuildSystem/Configuration/Windows/BuildConfigurationWindows.INT.udn
--- a/Engine/Documentation/Source/Programming/UnrealBuildSystem/Configuration/Windows/BuildConfigurationWindows.INT.udn
+++ b/Engine/Documentation/Source/Programming/UnrealBuildSystem/Configuration/Windows/BuildConfigurationWindows.INT.udn
@@ -19,7 +19,7 @@ For Windows specific options, you can alter the `UEBuildWindows.cs` configuratio
 $ WindowsPlatform::SupportWindowsXP : If _true_, a build of UE4 which is compatible with Windows XP will be created.
 
 
-To enable xp support, find where it is forced off in the `SetupEnvironment` and update this, with:
+To enable XP support, find where it is forced off in the `SetupEnvironment` and update this, with:
 	public override void SetUpEnvironment(UEBuildTarget InBuildTarget)
 	{
 		...
@@ -36,10 +36,10 @@ Enabling a Windows XP compatible build forces the compiler to use `Windows SDK 7
 	
 ### Getting the game running on Windows XP
 
-As Windows XP does not support support DirectX 10 or DirectX 11, OpenGL must be included as a target RHI for your game.
+As Windows XP does not support DirectX 10 or DirectX 11, OpenGL must be included as a target RHI for your game.
 
 You can find the option to include this in the Editor:
 _Edit->Project Settings->Platform->Windows_
 [![](EnableOpenGLForWindows.png)](Programming/UnrealBuildSystem/Configuration/Windows)
 
-Once you have packaged your game, you should then be able to run the game exectuable from the binaries directory on a Windows XP machine.
\ No newline at end of file
+Once you have packaged your game, you should then be able to run the game executable from the binaries directory on a Windows XP machine.
\ No newline at end of file
