#UE4doc #docs: sp: behaves

diff --git a/Engine/Documentation/Source/Resources/Showcases/BlueprintExamples/BlueprintExamples.INT.udn b/Engine/Documentation/Source/Resources/Showcases/BlueprintExamples/BlueprintExamples.INT.udn
--- a/Engine/Documentation/Source/Resources/Showcases/BlueprintExamples/BlueprintExamples.INT.udn
+++ b/Engine/Documentation/Source/Resources/Showcases/BlueprintExamples/BlueprintExamples.INT.udn
@@ -82,7 +82,7 @@ Throughout the **BlueprintOffice** level, the default level of this example proj
 
 ![](CommentaryBoxes.png)(w:700)
 
-It is recommended to play through the level in the Editor first, viewing all of the commentary boxes and exploring how each of the elements in the level behave. Then, once you have an idea of what each Blueprint does, you can examine the **Details** panel and Blueprint Editor to have a look at the scripting within them.
+It is recommended to play through the level in the Editor first, viewing all of the commentary boxes and exploring how each of the elements in the level behaves. Then, once you have an idea of what each Blueprint does, you can examine the **Details** panel and Blueprint Editor to have a look at the scripting within them.
 
 
 ## Blueprint Basics
