#UE4doc #docs: grammatical changes: ","

diff --git a/Engine/Documentation/Source/Shared/Types/AAmbientSound/AAmbientSound.INT.udn b/Engine/Documentation/Source/Shared/Types/AAmbientSound/AAmbientSound.INT.udn
--- a/Engine/Documentation/Source/Shared/Types/AAmbientSound/AAmbientSound.INT.udn
+++ b/Engine/Documentation/Source/Shared/Types/AAmbientSound/AAmbientSound.INT.udn
@@ -5,7 +5,7 @@ Description:Overview of the AmbientSound class
 
 [EXCERPT:AAmbientSound]
 
-Ambient Sound Actors conform to the real world where the closer you are to a sound, the louder it will appear. By comparison, a sound that is normally loud, may appear soft if it is further away. 
+Ambient Sound Actors conform to the real world where the closer you are to a sound, the louder it will appear. By comparison, a sound that is normally loud may appear soft if it is further away. 
 
 [VAR:ToolTipFullLink]
 Engine/Audio/SoundActors
