#UE4 DOCS: updated physics mat reference

diff --git a/Engine/Documentation/Source/Engine/Physics/PhysicalMaterials/PhysMatReference/PhysMatReference.INT.udn b/Engine/Documentation/Source/Engine/Physics/PhysicalMaterials/PhysMatReference/PhysMatReference.INT.udn
--- a/Engine/Documentation/Source/Engine/Physics/PhysicalMaterials/PhysMatReference/PhysMatReference.INT.udn
+++ b/Engine/Documentation/Source/Engine/Physics/PhysicalMaterials/PhysMatReference/PhysMatReference.INT.udn
@@ -25,7 +25,6 @@ Description:Assets applied to physically simulated primitives directly or via ma
 [VAR:TopicCompact]
 [OBJECT:TopicCompact]
 	[PARAM:image]
-		![%Engine/Physics/PhysicalMaterials/PhysMatReference:title%](Engine/Physics/physics_topic.png)
 	[/PARAM]
 	[PARAM:icon]
 		![](%ROOT%/physics_icon.png)(convert:false)
@@ -42,31 +41,157 @@ Description:Assets applied to physically simulated primitives directly or via ma
 [/OBJECT]
 [/VAR]
 
-[TOC(start:2)]
+[INCLUDE:Shared/Glossary/P/#Glossary_Physical_Material]
 
-## Overview
+## Properties
 
-[INCLUDE:Shared/Glossary/P/#Glossary_Physical_Material]
+Below are the properties for Physical Materials separated by major category.
 
-![PhysicalMaterialProperties.png](PhysicalMaterialProperties.png)
+[OBJECT:modulesectionexpanded]
+	[PARAMLITERAL:id]
+		physmat
+	[/PARAMLITERAL]
+	[PARAM:heading]
+		Physical Material
+	[/PARAM]
+	[PARAM:content]
+		%p1%
+	[/PARAM]
+[/OBJECT]
 
-## Properties
+[VAR:p1]
+
+This category contains the core properties of the Physics Material: Friction, Restitution, and Density.
+
+![](p1.png)
+
+[REGION:raw]
+| Property | %spacer% | Description |
+| -------- | --- | :---------- |
+| **Friction** || This is the friction value of the surface, which controls how easily things can slide on this surface. |
+| **Friction Combine Mode** || [INCLUDE:#FcombineMode]  |
+| **Override Friction Combine Mode** || By default, the Friction Combine Mode is set to **Average**, by enabling this property you can change how Frictions are combined between touching Physical Materials.|
+| **Restitution** || This is how "bouncy" the surface is, or how much energy it retains when it collides with another surface. |
+| **Restitution Combine Mode** || [INCLUDE:#RcombineMode] |
+| **Override Restitution Combine Mode** || By default, the Restitution Combine Mode is set to **Average**, by enabling this property you can change how Restitutions are combined between touching Physical Materials. |
+| **Density** || Used with the shape of the object to calculate its mass properties. The higher the number, the heavier the object. Measured as **g** per cubic **cm**. |
+[/REGION]
+
+[/VAR]
+
+
+[OBJECT:modulesection]
+	[PARAMLITERAL:id]
+		adv
+	[/PARAMLITERAL]
+	[PARAM:heading]
+		Advanced
+	[/PARAM]
+	[PARAM:content]
+		%p2%
+	[/PARAM]
+[/OBJECT]
+
+[VAR:p2]
+
+This category contains a property to change how scale affects the mass of the Actor with the Physical Material applied to it.
+
+![](p2.png)
+
+[REGION:raw]
+| Property | %spacer% | Description |
+| -------- | --- | :---------- |
+| **Raise Mass To Power** || Used to adjust the way that mass increases as objects get larger. This is applied to the mass as calculated based on a "solid" object. In actuality, larger objects do not tend to be solid, and become more like "shells" (e.g. a car is not a solid piece of metal). Values are clamped to 1 or less. |
+[/REGION]
+
+[/VAR]
+
+
+
+[OBJECT:modulesection]
+	[PARAMLITERAL:id]
+		dest
+	[/PARAMLITERAL]
+	[PARAM:heading]
+		Destruction
+	[/PARAM]
+	[PARAM:content]
+		%p3%
+	[/PARAM]
+[/OBJECT]
+
+[VAR:p3]
+
+Properties that are specific to the Destruction system in Unreal Engine 4.
 
-| Property | Description |
-| ---- | ----- |
-|[REGION:tablesection]Physical Material[/REGION]||
-| **Friction** | %doxygen:UPhysicalMaterial::Friction% |
-| **Restitution** | %doxygen:UPhysicalMaterial::Restitution% |
-| **Density** | %doxygen:UPhysicalMaterial::Density% |
-|[REGION:tablesection]Advanced[/REGION]||
-| **RaiseMassToPower** | %doxygen:UPhysicalMaterial::RaiseMassToPower% |
-|[REGION:tablesection]Destruction[/REGION]||
-| **DestructibleDamageThresholdScale** | %doxygen:UPhysicalMaterial::DestructibleDamageThresholdScale% |
-|[REGION:tablesection]Physical Properties[/REGION]||
-| **SurfaceType** | %doxygen:UPhysicalMaterial::SurfaceType% [INCLUDE:#one]|
-|[REGION:tablesection]Vehicles[/REGION]||
-| **TireFrictionScale** | %doxygen:UPhysicalMaterial::TireFrictionScale% |
-| **TireFrictionScales** | %doxygen:UPhysicalMaterial::TireFrictionScales% |
+![](p3.png)
+
+[REGION:raw]
+| Property | %spacer% | Description |
+| -------- | :---: | :---------- |
+| **Destructible Damage Threshold Scale** || How much to scale the damage threshold by on any destructible this Physics Material is applied to. |
+[/REGION]
+
+[/VAR]
+
+
+
+[OBJECT:modulesection]
+	[PARAMLITERAL:id]
+		physpro
+	[/PARAMLITERAL]
+	[PARAM:heading]
+		Physical Properties
+	[/PARAM]
+	[PARAM:content]
+		%p4%
+	[/PARAM]
+[/OBJECT]
+
+[VAR:p4]
+
+Gameplay related properties for the Physical Material.
+
+![](p4.png)
+
+[REGION:raw]
+| Property | %spacer% | Description |
+| -------- | --- | :---------- |
+| **Surface Type** || [INCLUDE:#one]|
+[/REGION]
+
+[/VAR]
+
+
+
+[OBJECT:modulesection]
+	[PARAMLITERAL:id]
+		vehicles
+	[/PARAMLITERAL]
+	[PARAM:heading]
+		Vehicles
+	[/PARAM]
+	[PARAM:content]
+		%p5%
+	[/PARAM]
+[/OBJECT]
+
+[VAR:p5]
+
+These properties are specific to Vehicles in Unreal Engine 4. While they may say they deal with Tires, they only do so when applied to a Vehicle (as Tire Data Types, and Wheel Blueprints, don't directly reference a Physical Material).
+
+![](p5.png)
+
+[REGION:raw]
+| Property | %spacer% | Description |
+| -------- | :---: | :---------- |
+| **Tire Friction Scale** || Overall tire friction scalar when this Physical Material is applied to a Vehicle. This value is multiplied against the Vehicle's Wheels' specific Friction Scale values. |
+| **Tire Friction Scales** || Tire friction scales for specific Wheels when this Physical Material is applied to a Vehicle. These values are multiplied against the Vehicle's Wheels' specific Friction Scale values. |
+[/REGION]
+
+[/VAR]
+
+[VAR:spacer]![](spacer.png)(w:150 convert:false)[/VAR]
 
 <!--
 [EXCERPT:one]
@@ -76,9 +201,29 @@ You can use ProjectSetting/Physics/Physical Surfaces
 
 This data can be pulled in code or Blueprint:
 
-![](SurfaceType.png)
+![](SurfaceType.png)(w:600)
 
 By default, without editing the source code, you are limited to 30 Surface Types, labeled SurfaceType1 through SurfaceType30.
 
 [/EXCERPT:one]
+
+[EXCERPT:FcombineMode]
+This property enables you to adjust how Physical Material frictions are combined. This is set to Average by default, but can be overriden using the **Override Friction Combine Mode** property. 
+
+| --- | --- |
+| **Average** | Uses the average value of the touching materials: (a+b)/2 |
+| **Min** | Uses the smaller value of the touching materials: min(a,b) |
+| **Multiply** | Multiplies the values of the touching materials: a*b |
+| **Max** | Uses the larger value of the touching materials: max(a,b) |
+[/EXCERPT]
+
+[EXCERPT:RcombineMode]
+This property enables you to adjust how Physical Material restitutions are combined. This is set to Average by default, but can be overriden using the **Override Restitution Combine Mode** property. 
+
+| --- | --- |
+| **Average** | Uses the average value of the touching materials: (a+b)/2 |
+| **Min** | Uses the smaller value of the touching materials: min(a,b) |
+| **Multiply** | Multiplies the values of the touching materials: a*b |
+| **Max** | Uses the larger value of the touching materials: max(a,b) |
+[/EXCERPT]
 -->
\ No newline at end of file
