#UE4 DOCS: submitting Binary vs Source swapper images and the globals to support them

diff --git a/Engine/Documentation/Source/Globals/GlobalIncludes.INT.udn b/Engine/Documentation/Source/Globals/GlobalIncludes.INT.udn
--- a/Engine/Documentation/Source/Globals/GlobalIncludes.INT.udn
+++ b/Engine/Documentation/Source/Globals/GlobalIncludes.INT.udn
@@ -133,6 +133,43 @@ Description:
 [/REGION]
 [/VAR]
 
+
+[VAR:BlueprintOrCPPSelectVar]
+[REGION:section]
+	[REGION:titlesection]
+	Choose your implementation method:
+	[/REGION]
+	[OBJECT:ToggleButton]
+		[PARAMLITERAL:category]
+		BlueprintOrCPPSelect
+		[/PARAMLITERAL]
+		[PARAMLITERAL:id]
+		blueprintimplementation
+		[/PARAMLITERAL]
+		[PARAMLITERAL:active]
+		active_button
+		[/PARAMLITERAL]
+		[PARAM:label]
+		Blueprints
+		[/PARAM]
+	[/OBJECT]
+	[OBJECT:ToggleButton]
+		[PARAMLITERAL:category]
+		BlueprintOrCPPSelect
+		[/PARAMLITERAL]
+		[PARAMLITERAL:id]
+		cppimplementation
+		[/PARAMLITERAL]
+		[PARAMLITERAL:active]
+		[/PARAMLITERAL]
+		[PARAM:label]
+		C++
+		[/PARAM]
+	[/OBJECT]
+[/REGION]
+[/VAR]
+
+
 [EXCERPT:bpModesRemoved]
 [REGION:note]
 The Blueprint Editor Modes were removed in 4.7. Please see [%Engine/Blueprints/Editor:title%](Engine/Blueprints/Editor "%Engine/Blueprints/Editor:description%") for information about the **Blueprint Editor Interface**.
