#UE4 DOCS: EQS Quickstart and Updates to the Behavior Tree Quick Start

diff --git a/Engine/Documentation/Source/Engine/AI/BehaviorTrees/BehaviorTrees.INT.udn b/Engine/Documentation/Source/Engine/AI/BehaviorTrees/BehaviorTrees.INT.udn
--- a/Engine/Documentation/Source/Engine/AI/BehaviorTrees/BehaviorTrees.INT.udn
+++ b/Engine/Documentation/Source/Engine/AI/BehaviorTrees/BehaviorTrees.INT.udn
@@ -41,6 +41,18 @@ Description:Landing Page for all the documentation relating to Behavior Trees
 	[/PARAM]
 [/OBJECT]
 [/VAR]
+
+From a lowly roach that hides from the light, to the Bot in Shooter Game that knows when to shoot, and when to seek health or ammo; **Behavior Trees** are a powerful tool for creating Artificial Intelligence in Unreal Engine 4. 
+
+They're a combination of two asset types: The **Blackboard** and the **Behavior Tree**.
+
+| Asset Type | Description |
+| --- | --- |
+| **Blackboard** | The **Blackboard** is the AI's memory. It stores key values for the Behavior Tree to use. |
+| **Behavior Tree** | The **Behavior Tree** is the AI's processor. It makes decisions, and then acts on those decisions. |
+
+Below are links to the sections in the documentation covering the use of Behavior Trees. The [Quick Start Guide](Engine/AI/BehaviorTrees/QuickStart) takes you through making an AI chase the Player. Then there's the [User Guide](Engine/AI/BehaviorTrees/UserGuide), which contains instructions on the most common tasks when using Behavior Trees. Next, the [Node Reference](Engine/AI/BehaviorTrees/NodeReference) contains a listing of all then nodes available in the Behavior Tree editor, what they do, and what their properties are. Finally, [How Unreal Engine 4 Behavior Trees Differ](Engine/AI/BehaviorTrees/HowUE4BehaviorTreesDiffer) is a paper written for those who are already familiar with how other Behavior Trees work.
+
 [OBJECT:TopicButtonList]
 	[PARAM:icon]
 		![](%ROOT%/tech_icon.png)(convert:false)
diff --git a/Engine/Documentation/Source/Engine/AI/BehaviorTrees/NodeReference/BehaviorTreeNodeReference.INT.udn b/Engine/Documentation/Source/Engine/AI/BehaviorTrees/NodeReference/BehaviorTreeNodeReference.INT.udn
--- a/Engine/Documentation/Source/Engine/AI/BehaviorTrees/NodeReference/BehaviorTreeNodeReference.INT.udn
+++ b/Engine/Documentation/Source/Engine/AI/BehaviorTrees/NodeReference/BehaviorTreeNodeReference.INT.udn
@@ -3,9 +3,6 @@ Title: Behavior Trees Nodes Reference
 Crumbs:%ROOT%, Engine, Gameplay, Engine/AI/BehaviorTrees
 Description:Reference for the Behavior Tree nodes.
 
-[TOC(start:2 end:2)]
-
-## Overview
 
 This is a reference for the default nodes available in the Behavior Tree Editor. You may see more depending on the project you are working with (such as Shooter Game). There are five types of Behavior Tree Nodes:
 
@@ -36,7 +33,6 @@ See the links below to see full listings of the default Behavior Tree Nodes.
 	[/PARAM]
 [/OBJECT]
 
-
 <!--
 [EXCERPT:rootNode]
 The Root node is unique in the Behavior Tree and is the starting point for the Behavior Tree. It can only have one connection, and you cannot attach Decorators or Services to it. The Root Node has no properties of its own, but selecting it will show the Behavior Tree properties in the Details Panel, where you can set the Blackboard Asset of the Behavior Tree.
diff --git a/Engine/Documentation/Source/Engine/AI/BehaviorTrees/NodeReference/Composites/Composites.INT.udn b/Engine/Documentation/Source/Engine/AI/BehaviorTrees/NodeReference/Composites/Composites.INT.udn
--- a/Engine/Documentation/Source/Engine/AI/BehaviorTrees/NodeReference/Composites/Composites.INT.udn
+++ b/Engine/Documentation/Source/Engine/AI/BehaviorTrees/NodeReference/Composites/Composites.INT.udn
@@ -5,8 +5,6 @@ Availability:Public
 
 [TOC(start:2 end:2)]
 
-## Overview
-
 **Composite Nodes** define the root of a branch and the base rules for how that branch is executed. They can have [**Decorators**](Engine/AI/BehaviorTrees/NodeReference/Decorators) applied to them to modify entry into their branch or even cancel out mid execution. Also, they can have [**Services**](Engine/AI/BehaviorTrees/NodeReference/Services) attached to them that will only be active if the children of the Composite are being executed.
 
 ## Selector
diff --git a/Engine/Documentation/Source/Engine/AI/BehaviorTrees/NodeReference/Decorators/Decorators.INT.udn b/Engine/Documentation/Source/Engine/AI/BehaviorTrees/NodeReference/Decorators/Decorators.INT.udn
--- a/Engine/Documentation/Source/Engine/AI/BehaviorTrees/NodeReference/Decorators/Decorators.INT.udn
+++ b/Engine/Documentation/Source/Engine/AI/BehaviorTrees/NodeReference/Decorators/Decorators.INT.udn
@@ -5,7 +5,6 @@ Description:Reference for the Behavior Tree Decorator nodes.
 
 [TOC(start:2 end:2)]
 
-## Overview
 
 **Decorator**, also known as conditionals in other Behavior Tree systems, are attached to either a [**Composite**](Engine/AI/BehaviorTrees/NodeReference/Composites) or a [**Task**](Engine/AI/BehaviorTrees/NodeReference/Tasks) node and define whether or not a branch in the tree, or even a single node, can be executed.
 
diff --git a/Engine/Documentation/Source/Engine/AI/BehaviorTrees/NodeReference/Services/Services.INT.udn b/Engine/Documentation/Source/Engine/AI/BehaviorTrees/NodeReference/Services/Services.INT.udn
--- a/Engine/Documentation/Source/Engine/AI/BehaviorTrees/NodeReference/Services/Services.INT.udn
+++ b/Engine/Documentation/Source/Engine/AI/BehaviorTrees/NodeReference/Services/Services.INT.udn
@@ -5,8 +5,6 @@ Description:Reference for the Behavior Tree Service nodes.
 
 [TOC(start:2 end:2)]
 
-## Overview
-
 **Services** attach to [**Composite**](Engine/AI/BehaviorTrees/NodeReference/Composites) nodes, and will execute at their defined frequency as long as their branch is being executed. These are often used to make checks and to update the Blackboard. These take the place of traditional Parallel nodes in other Behavior Tree systems
 
 
diff --git a/Engine/Documentation/Source/Engine/AI/BehaviorTrees/NodeReference/Tasks/Tasks.INT.udn b/Engine/Documentation/Source/Engine/AI/BehaviorTrees/NodeReference/Tasks/Tasks.INT.udn
--- a/Engine/Documentation/Source/Engine/AI/BehaviorTrees/NodeReference/Tasks/Tasks.INT.udn
+++ b/Engine/Documentation/Source/Engine/AI/BehaviorTrees/NodeReference/Tasks/Tasks.INT.udn
@@ -5,8 +5,6 @@ Description:Reference for the Behavior Tree Task nodes.
 
 [TOC(start:2 end:2)]
 
-## Overview
-
 This document is a reference page for the **Task Nodes** available in the Behavior Tree Editor. Tasks are nodes that "do" things, like move an AI, or adjust Blackboard values. They can have [**Decorators**](Engine/AI/BehaviorTrees/NodeReference/Decorators) attached to them.
 
 ## Make Noise
diff --git a/Engine/Documentation/Source/Engine/AI/BehaviorTrees/QuickStart/1/1.INT.udn b/Engine/Documentation/Source/Engine/AI/BehaviorTrees/QuickStart/1/1.INT.udn
--- a/Engine/Documentation/Source/Engine/AI/BehaviorTrees/QuickStart/1/1.INT.udn
+++ b/Engine/Documentation/Source/Engine/AI/BehaviorTrees/QuickStart/1/1.INT.udn
@@ -9,6 +9,9 @@ Description:Step 1. Essential Setup
 	[/PARAM]
 	[PARAM:current]
 	[/PARAM]
+	[PARAM:coloroverride]
+		#000000
+	[/PARAM]
 	[PARAM:home]
 		[Quick Start Home](Engine/AI/BehaviorTrees/QuickStart)
 	[/PARAM]
@@ -19,16 +22,21 @@ Description:Step 1. Essential Setup
 [/VAR]
 
 %Steps%
+
+This Quick Start makes some assumptions about what is available in the project, as such, if you aren't comfortable making substations or problem solving on the fly, then you should setup the project as follows.
+
+## Steps
+
 [EXCERPT:AIProjectStart]
 
-1. Open Unreal Editor 4.
+1. Open Unreal Editor
 
-1. Generate a new project based on the **Blueprint Top Down** template.
+1. Generate a new project based on the **Blueprint Top Down** template
 
 	The **Starter Assets** are optional for this quick start guide.
 
 
-If you need a more in depth explanation of starting a new project, see the [Project Browser](Engine/Basics/Projects/Browser) User Guide.
+If you need further explanation of starting a new project, see the [Project Browser](Engine/Basics/Projects/Browser) User Guide.
 
 [/EXCERPT:AIProjectStart]
 %Steps%
diff --git a/Engine/Documentation/Source/Engine/AI/BehaviorTrees/QuickStart/10/10.INT.udn b/Engine/Documentation/Source/Engine/AI/BehaviorTrees/QuickStart/10/10.INT.udn
--- a/Engine/Documentation/Source/Engine/AI/BehaviorTrees/QuickStart/10/10.INT.udn
+++ b/Engine/Documentation/Source/Engine/AI/BehaviorTrees/QuickStart/10/10.INT.udn
@@ -21,9 +21,6 @@ Description:Step 10. Behavior Tree Basics
 
 %Steps%
 
-
-This step is mostly exposition on how Behavior Trees execute and commonly used nodes in Behavior Trees.
-
 The **Root** node starts the tree, and all other nodes connect below it.
 
 ![](root.png)
@@ -35,12 +32,12 @@ It can only have one output connection. The darker bar on the bottom of the node
 However, unlike Blueprint node pins which allow you to drag connections from either output to input, or input to output, Behavior Tree nodes are directional: you must drag from an "output bar" to an "input bar."
 
 
-## Create and Connect Some Nodes
+## Steps
 
 ![](selAndSeq.png)
 
 1. Open your Behavior Tree asset, if it is not open already, by **Double-clicking** it in the **Content Browser**.
-1. In Behavior Tree mode, **Right-click** in the graph area and create a **Composites** -> **Selector** node.
+1. In Behavior Tree mode, **Right-click** in the graph area and create two **Composites** -> **Selector** nodes.
 
 	![](selector.png)
 
diff --git a/Engine/Documentation/Source/Engine/AI/BehaviorTrees/QuickStart/11/11.INT.udn b/Engine/Documentation/Source/Engine/AI/BehaviorTrees/QuickStart/11/11.INT.udn
--- a/Engine/Documentation/Source/Engine/AI/BehaviorTrees/QuickStart/11/11.INT.udn
+++ b/Engine/Documentation/Source/Engine/AI/BehaviorTrees/QuickStart/11/11.INT.udn
@@ -21,6 +21,13 @@ Description:Step 11. Make a Service
 
 %Steps%
 
+**Services** attach to [**Composite**](Engine/AI/BehaviorTrees/NodeReference/Composites) nodes, and will execute at their defined frequency as long as their branch is being executed. These are often used to make checks and to update the Blackboard. These take the place of traditional Parallel nodes in other Behavior Tree systems
+
+In this case, the service is going to look for any **Pawn** that is not an AI Character, then set the TargetToFollow and TargetLocation Blackboard keys.
+
+
+## Steps
+
 1. Open the FollowerBT Behavior Tree by **Double-clicking** it in the **Content Browser**.
 1. Click New Service ![](newService_button.png) button.
 
diff --git a/Engine/Documentation/Source/Engine/AI/BehaviorTrees/QuickStart/12/12.INT.udn b/Engine/Documentation/Source/Engine/AI/BehaviorTrees/QuickStart/12/12.INT.udn
--- a/Engine/Documentation/Source/Engine/AI/BehaviorTrees/QuickStart/12/12.INT.udn
+++ b/Engine/Documentation/Source/Engine/AI/BehaviorTrees/QuickStart/12/12.INT.udn
@@ -21,6 +21,13 @@ Description:Step 12. Make a Task
 
 %Steps%
 
+Tasks are nodes that "do" things, like move an AI, or adjust Blackboard values. They can have [**Decorators**](Engine/AI/BehaviorTrees/NodeReference/Decorators) attached to them.
+
+This **Task** is basically a re-creation of the **Move To** node, with one major difference: this Task will have the AI always run towards the Player's Pawn, instead of its location. Also, it helps to break things down and see that if you don't ever execute **Finish Execute** the Task Node will continue to process, and the **AI MoveTo** node shows this perfectly.
+
+
+## Steps
+
 1. Open the FollowerBT Behavior Tree by **Double-clicking** it in the **Content Browser**.
 1. Click New Task ![](newTask_button.png) button.
 
diff --git a/Engine/Documentation/Source/Engine/AI/BehaviorTrees/QuickStart/13/13.INT.udn b/Engine/Documentation/Source/Engine/AI/BehaviorTrees/QuickStart/13/13.INT.udn
--- a/Engine/Documentation/Source/Engine/AI/BehaviorTrees/QuickStart/13/13.INT.udn
+++ b/Engine/Documentation/Source/Engine/AI/BehaviorTrees/QuickStart/13/13.INT.udn
@@ -21,6 +21,12 @@ Description:13. Make a Decorator
 
 %Steps%
 
+**Decorators**, also known as conditionals in other Behavior Tree systems, are attached to either a [**Composite**](Engine/AI/BehaviorTrees/NodeReference/Composites) or a [**Task**](Engine/AI/BehaviorTrees/NodeReference/Tasks) node and define whether or not a branch in the tree, or a single node, can be executed.
+
+This **Decorator** is going to stop execution of the Rapid Move To branch of the Behavior Tree once the AI's Pawn is close enough to its target. While this could be handled in the **Rapid Move To** Task or by just using a **Move To** Task, this decorator can take in any Actor and can block execution of any branch. This makes it reusable in a number of situations.
+
+## Steps
+
 1. Open the FollowerBT Behavior Tree by **Double-clicking** it in the **Content Browser**.
 1. Click New Decorator ![](newDecorator_button.png) button.
 
diff --git a/Engine/Documentation/Source/Engine/AI/BehaviorTrees/QuickStart/14/14.INT.udn b/Engine/Documentation/Source/Engine/AI/BehaviorTrees/QuickStart/14/14.INT.udn
--- a/Engine/Documentation/Source/Engine/AI/BehaviorTrees/QuickStart/14/14.INT.udn
+++ b/Engine/Documentation/Source/Engine/AI/BehaviorTrees/QuickStart/14/14.INT.udn
@@ -26,9 +26,19 @@ You should now have all the pieces needed to build this Behavior Tree:
 
 ![](fullTree.png)
 
+Let's break down what's going on.
+
+* The AI will just stand around if it hasn't detected a Pawn that isn't an AI Character.
+* The AgroCheck Service is constantly looking for a suitable target.
+* Once it finds one, the left most branch is open, and the AI runs towards the targeted Pawn.
+* Once it gets in range of that Pawn the branch is prevented from executing, thusly returning it to the top of the graph
+* If the desired Pawn hides, then the right most branch will be able to run, which causes the AI to move to the last know location of the desired Pawn.
+* Once there, it waits, then returns home.
+* At any time, if a suitable target is found, the right most branch will be forced to stop executing, and the left most branch will be opened.
+
 Let us step through the creation of this tree.
 
-## Walkthrough
+## Steps
 
 ### Basic Setup
 
diff --git a/Engine/Documentation/Source/Engine/AI/BehaviorTrees/QuickStart/2/2.INT.udn b/Engine/Documentation/Source/Engine/AI/BehaviorTrees/QuickStart/2/2.INT.udn
--- a/Engine/Documentation/Source/Engine/AI/BehaviorTrees/QuickStart/2/2.INT.udn
+++ b/Engine/Documentation/Source/Engine/AI/BehaviorTrees/QuickStart/2/2.INT.udn
@@ -8,6 +8,9 @@ Description:Step 2. Create a Nav Mesh Bounds Volume
 	[PARAM:previous]
 		[Previous Step](Engine/AI/BehaviorTrees/QuickStart/1 "%Engine/AI/BehaviorTrees/QuickStart/1:title%")
 	[/PARAM]
+	[PARAM:coloroverride]
+		#000000
+	[/PARAM]
 	[PARAM:current]
 	[/PARAM]
 	[PARAM:home]
@@ -25,7 +28,10 @@ Description:Step 2. Create a Nav Mesh Bounds Volume
 
 The NavMesh in Unreal has a number of functions that work well with Behavior Trees, and while it may not be necessary depending on the game you are making, it will help in a number of game types. Also, the Top Down template already includes a Nav Mesh Bounds Volume, but you may choose to start in a completely blank project at a later date.
 
-1. **Modes Panel** -> **Create** -> **Volumes**.
+## Steps
+
+1. [REGION:raw]**Modes Panel** %arrow% **Create** %arrow% **Volumes**.[/REGION]
+
 1. Click and drag a **Nav Mesh Bounds Volume** into the map.
 
 	![](createNavVol.png)
@@ -44,3 +50,5 @@ The NavMesh in Unreal has a number of functions that work well with Behavior Tre
 
 [/EXCERPT:AICreateNav]
 %Steps%
+
+[VAR:arrow]![](Engine\AI\EnvironmentQuerySystem\QuickStart\arrow_button.png)[/VAR]
diff --git a/Engine/Documentation/Source/Engine/AI/BehaviorTrees/QuickStart/3/3.INT.udn b/Engine/Documentation/Source/Engine/AI/BehaviorTrees/QuickStart/3/3.INT.udn
--- a/Engine/Documentation/Source/Engine/AI/BehaviorTrees/QuickStart/3/3.INT.udn
+++ b/Engine/Documentation/Source/Engine/AI/BehaviorTrees/QuickStart/3/3.INT.udn
@@ -28,6 +28,7 @@ To start, you will need create 4 assets in the **Content Browser**: An AI Contro
 ![](Assets.png)
 [/REGION]
 
+## Steps
 
 1. **New** -> **Blueprint** -> **AI Controller**
 
diff --git a/Engine/Documentation/Source/Engine/AI/BehaviorTrees/QuickStart/4/4.INT.udn b/Engine/Documentation/Source/Engine/AI/BehaviorTrees/QuickStart/4/4.INT.udn
--- a/Engine/Documentation/Source/Engine/AI/BehaviorTrees/QuickStart/4/4.INT.udn
+++ b/Engine/Documentation/Source/Engine/AI/BehaviorTrees/QuickStart/4/4.INT.udn
@@ -21,7 +21,11 @@ Description:Step 4. Setup the Character Blueprint
 
 %Steps%
 
-As we have created a project with a character already in it, we have access to a humanoid Skeletal Mesh and its Animation Blueprint. We will use these in the tutorial, but feel free to use your own if you have them available.
+As we have created a project with a **Character** already in it, we have access to a humanoid **Skeletal Mesh** and its **Animation Blueprint**. We will use these in the tutorial, but feel free to use your own if you have them available.
+
+This Character will become our new AI Character. The reason we use a Character, over say a Pawn, is to gain access to the Character Movement Component's functionality.
+
+## Steps
 
 1. **Double-click** the AI_Character Blueprint in the **Content Browser** to open it for editing.
 
@@ -35,6 +39,8 @@ As we have created a project with a character already in it, we have access to a
 
 1. Move the **Mesh** Component to center it up inside in the capsule.
 
+1. Rotate the **Mesh** Component so it faces the direction the blue Arrow Component is pointing.
+
 1. In the **Details** panel, find the **Animation** category and assign the `HeroTPP_AnimBlueprint` to the **Animation Blueprint Generated Class** property.
 
 	![](assignAnimBP.png)
@@ -57,7 +63,7 @@ As we have created a project with a character already in it, we have access to a
 	
 	![](maxWalk.png)
 
-1. Click the **Class Defaults** ![](Engine\Blueprints\UserGuide\Variables\defaultsButton.png) button and assign the new AI Controller (`Follower_AI_CON`) to the **AIController Class** property in the **AI** category.
+1. Click the **Class Defaults** ![](Engine\Blueprints\UserGuide\Variables\defaultsButton.png) button and assign the new AI Controller (`Follower_AI_CON`) to the **AIController Class** property in the **Pawn** category.
 
 	![](AI_con.png)
 
diff --git a/Engine/Documentation/Source/Engine/AI/BehaviorTrees/QuickStart/5/5.INT.udn b/Engine/Documentation/Source/Engine/AI/BehaviorTrees/QuickStart/5/5.INT.udn
--- a/Engine/Documentation/Source/Engine/AI/BehaviorTrees/QuickStart/5/5.INT.udn
+++ b/Engine/Documentation/Source/Engine/AI/BehaviorTrees/QuickStart/5/5.INT.udn
@@ -22,7 +22,9 @@ Description:Step 5. Setup the Blackboard Asset
 %Steps%
 
 
-The **Blackboard** asset stores all of the data needed for a specific instance of an AI controller that will be referenced by the **Behavior Tree**.
+The **Blackboard** asset stores all of the data needed for a specific instance of an AI controller that will be referenced by the **Behavior Tree**. In this specific case, we'll be storing an Actor to Follow (TargetToFollow), the AI's starting location (HomeLocation), and the last know location of the Actor we're following (TargetLocation).
+
+## Steps
 
 ![](BBAsset.png)
 
diff --git a/Engine/Documentation/Source/Engine/AI/BehaviorTrees/QuickStart/6/6.INT.udn b/Engine/Documentation/Source/Engine/AI/BehaviorTrees/QuickStart/6/6.INT.udn
--- a/Engine/Documentation/Source/Engine/AI/BehaviorTrees/QuickStart/6/6.INT.udn
+++ b/Engine/Documentation/Source/Engine/AI/BehaviorTrees/QuickStart/6/6.INT.udn
@@ -27,7 +27,9 @@ Description:Step 6. Setup the AI Controller Blueprint
 
 Using the Top Down template gives us access to mouse controls to move our character around, but for the other character we are going to create, he will need his own controller. This AI Controller is going to set a few values on its Blackboard and is going to kick off the running of our Behavior Tree. For those familiar with Blueprints, here is what the final **Event Graph** will look like:
 
-Below are the steps for creating the first part of the **Event Graph**. Steps 7 and 8 also involve editing this graph.
+Below are the steps for creating the first part of the **Event Graph**. Step _7. AI Controller Blueprint: Initial Blackboard Values_ and Step _8. AI Controller Blueprint: Run the Behavior Tree_ also involve editing this graph.
+
+## Steps
 
 1. **Double-click** the AI Controller Blueprint in the **Content Browser** to begin editing it.
 
diff --git a/Engine/Documentation/Source/Engine/AI/BehaviorTrees/QuickStart/7/7.INT.udn b/Engine/Documentation/Source/Engine/AI/BehaviorTrees/QuickStart/7/7.INT.udn
--- a/Engine/Documentation/Source/Engine/AI/BehaviorTrees/QuickStart/7/7.INT.udn
+++ b/Engine/Documentation/Source/Engine/AI/BehaviorTrees/QuickStart/7/7.INT.udn
@@ -25,7 +25,9 @@ Description:Step 7. AI Controller Blueprint: Initial Blackboard Values
 [![](fullEG.png)(w:940 convert:false)](fullEG.png)
 [/REGION]
 
-Now we need to set the "HomeLocation" key on the **Blackboard**. This is a more involved step, but not terribly difficult.
+Now we need to set the "HomeLocation" key on the **Blackboard**. This is a more involved step.
+
+## Steps
 
 1. You need a **Get Blackboard** node and a **Get Controlled Pawn** node.
 
diff --git a/Engine/Documentation/Source/Engine/AI/BehaviorTrees/QuickStart/8/8.INT.udn b/Engine/Documentation/Source/Engine/AI/BehaviorTrees/QuickStart/8/8.INT.udn
--- a/Engine/Documentation/Source/Engine/AI/BehaviorTrees/QuickStart/8/8.INT.udn
+++ b/Engine/Documentation/Source/Engine/AI/BehaviorTrees/QuickStart/8/8.INT.udn
@@ -25,7 +25,9 @@ Description:Step 8. Controller Blueprint: Run the Behavior Tree
 [![](fullEG.png)(w:940 convert:false)](fullEG.png)
 [/REGION]
 
-Finally, we just have to run our **Behavior Tree** with a single node:
+Finally, we just have to run our **Behavior Tree** with a single node.
+
+## Steps
 
 1. Create a **Run Behavior Tree** node.
 
diff --git a/Engine/Documentation/Source/Engine/AI/BehaviorTrees/QuickStart/9/9.INT.udn b/Engine/Documentation/Source/Engine/AI/BehaviorTrees/QuickStart/9/9.INT.udn
--- a/Engine/Documentation/Source/Engine/AI/BehaviorTrees/QuickStart/9/9.INT.udn
+++ b/Engine/Documentation/Source/Engine/AI/BehaviorTrees/QuickStart/9/9.INT.udn
@@ -21,11 +21,17 @@ Description:Step 9. Place the Character
 
 %Steps%
 
-1. From the **Content Browser**, click and drag the **Character** you created back in step 3 into the level.
+[EXCERPT:s9]
+
+We need to place our AI_Character into the level so it can interact with the player or environment.
+
+## Steps
+
+* From the **Content Browser**, click and drag the `AI_Character`, that you created earlier, into the level.
 
 [REGION:lightbox]
 [![](placeCharacter.png)(w:940 convert:false)](placeCharacter.png)
 [/REGION]
-
+[/EXCERPT:s9]
 
 %Steps%
\ No newline at end of file
diff --git a/Engine/Documentation/Source/Engine/AI/BehaviorTrees/QuickStart/BehaviorTreeQuickStart.INT.udn b/Engine/Documentation/Source/Engine/AI/BehaviorTrees/QuickStart/BehaviorTreeQuickStart.INT.udn
--- a/Engine/Documentation/Source/Engine/AI/BehaviorTrees/QuickStart/BehaviorTreeQuickStart.INT.udn
+++ b/Engine/Documentation/Source/Engine/AI/BehaviorTrees/QuickStart/BehaviorTreeQuickStart.INT.udn
@@ -1,9 +1,7 @@
-﻿Availability:Public
+Availability:Public
 Title: Behavior Tree Quick Start Guide
 Crumbs:%ROOT%, Engine, Gameplay, Engine/AI/BehaviorTrees
 Description:Overviews and examples AI
-Prereq:Engine/Blueprints/BP_HowTo/BasicUsage
-Prereq:Engine/Blueprints
 
 [VAR:Topic]
 [OBJECT:Topic]
diff --git a/Engine/Documentation/Source/Engine/AI/BehaviorTrees/UserGuide/BehaviorTreeUserGuide.INT.udn b/Engine/Documentation/Source/Engine/AI/BehaviorTrees/UserGuide/BehaviorTreeUserGuide.INT.udn
--- a/Engine/Documentation/Source/Engine/AI/BehaviorTrees/UserGuide/BehaviorTreeUserGuide.INT.udn
+++ b/Engine/Documentation/Source/Engine/AI/BehaviorTrees/UserGuide/BehaviorTreeUserGuide.INT.udn
@@ -5,8 +5,6 @@ Description:Behavior Tree Editor User Guide
 
 [TOC(start:2 end:2)]
 
-## Overview
-
 This document will cover the most common workflows encountered while working with Behavior Trees.
 
 ## Creating Behavior Tree Assets
diff --git a/Engine/Documentation/Source/Engine/AI/EnvironmentQuerySystem/AskingSmartQuestions/AskingSmartQuestions.INT.udn b/Engine/Documentation/Source/Engine/AI/EnvironmentQuerySystem/AskingSmartQuestions/AskingSmartQuestions.INT.udn
--- a/Engine/Documentation/Source/Engine/AI/EnvironmentQuerySystem/AskingSmartQuestions/AskingSmartQuestions.INT.udn
+++ b/Engine/Documentation/Source/Engine/AI/EnvironmentQuerySystem/AskingSmartQuestions/AskingSmartQuestions.INT.udn
@@ -1,4 +1,4 @@
-Availability:Public
+Availability:Docs
 Title: Asking the Environment Smart Questions
 Crumbs:%ROOT%, Engine, Gameplay/AI, Engine/AI/EnvironmentQuerySystem
 Description:This is a white paper written by Sr. AI Programmer Mieszko Zielinski over the predecessor to the Environment Query System: the Environment Tactical Querying.
diff --git a/Engine/Documentation/Source/Engine/AI/EnvironmentQuerySystem/EQSPawn/EQSPawn.INT.udn b/Engine/Documentation/Source/Engine/AI/EnvironmentQuerySystem/EQSPawn/EQSPawn.INT.udn
--- a/Engine/Documentation/Source/Engine/AI/EnvironmentQuerySystem/EQSPawn/EQSPawn.INT.udn
+++ b/Engine/Documentation/Source/Engine/AI/EnvironmentQuerySystem/EQSPawn/EQSPawn.INT.udn
@@ -44,7 +44,7 @@ Description:The Environment Query System includes a specialized Pawn that will d
 
 ![](testingPawnBanner.png)(convert:false)
 
-The **Environment Query System Testing Pawn** (EQSTestingPawn) is a specialized Pawn class that enables you to check out what your Environment Queries are actually doing. The exact makeup of your Environment Query will define the size and shape of what is created, but it will always be represented as colored spheres. Spheres that are on the color scale from green to red, indicate some level of match for the various tests your Environment Query ran. Blue spheres indicate a failure or a bool type test that returned false, such as the **Trace** test.
+The **Environment Query System Testing Pawn** (EQSTestingPawn) is a specialized Pawn class that enables you to check out what your Environment Queries are actually doing. The exact makeup of your Environment Query will define the size and shape of what is created, but it will always be represented as colored spheres. Spheres that are on the color scale from green to red, indicate some level of match for the various tests your Environment Query ran. Blue spheres indicate a failure or a boolean type test that returned false, such as the **Trace** test.
 
 ## Creating an EQS Testing Pawn
 
diff --git a/Engine/Documentation/Source/Engine/AI/EnvironmentQuerySystem/EnvironmentQuerySystem.INT.udn b/Engine/Documentation/Source/Engine/AI/EnvironmentQuerySystem/EnvironmentQuerySystem.INT.udn
--- a/Engine/Documentation/Source/Engine/AI/EnvironmentQuerySystem/EnvironmentQuerySystem.INT.udn
+++ b/Engine/Documentation/Source/Engine/AI/EnvironmentQuerySystem/EnvironmentQuerySystem.INT.udn
@@ -42,6 +42,11 @@ Description:Landing Page for all the documentation relating to the Environment Q
 [/OBJECT]
 [/VAR]
 
+The **Environment Query System** is a feature of the **Artificial Intelligence** system in Unreal Engine 4 for collecting data on the
+environment, asking questions of the data through **Tests**, then returning then one **Item** that best fits the questions asked.
+
+You can use it to find the closest power up, figure out which enemy is the highest threat, or to find cover.
+
 [OBJECT:TopicButtonList]
 	[PARAM:icon]
 		![](%ROOT%/tech_icon.png)(convert:false)
@@ -52,13 +57,16 @@ Description:Landing Page for all the documentation relating to the Environment Q
 	[PARAM:description]
 	[/PARAM]
 	[PARAM:links]
+	* [](Engine/AI/EnvironmentQuerySystem/QuickStart "%Engine/AI/EnvironmentQuerySystem/QuickStart:description%")
 	* [](Engine/AI/EnvironmentQuerySystem/UserGuide "%Engine/AI/EnvironmentQuerySystem/UserGuide:description%")
 	* [](Engine/AI/EnvironmentQuerySystem/EQSPawn "%Engine/AI/EnvironmentQuerySystem/EQSPawn:description%")
 	* [](Engine/AI/EnvironmentQuerySystem/NodeReference "%Engine/AI/EnvironmentQuerySystem/NodeReference:description%")
-	* [](Engine/AI/EnvironmentQuerySystem/AskingSmartQuestions "%Engine/AI/EnvironmentQuerySystem/AskingSmartQuestions:description%")
 	[/PARAM]
 [/OBJECT]
 
 <!---
 * [](Engine/AI/EnvironmentQuerySystem/QuickStart "%Engine/AI/EnvironmentQuerySystem/QuickStart:description%")
+
+May not have the rights to post:
+* [](Engine/AI/EnvironmentQuerySystem/AskingSmartQuestions "%Engine/AI/EnvironmentQuerySystem/AskingSmartQuestions:description%")
 --->
\ No newline at end of file
diff --git a/Engine/Documentation/Source/Engine/AI/EnvironmentQuerySystem/NodeReference/NodeReference.INT.udn b/Engine/Documentation/Source/Engine/AI/EnvironmentQuerySystem/NodeReference/NodeReference.INT.udn
--- a/Engine/Documentation/Source/Engine/AI/EnvironmentQuerySystem/NodeReference/NodeReference.INT.udn
+++ b/Engine/Documentation/Source/Engine/AI/EnvironmentQuerySystem/NodeReference/NodeReference.INT.udn
@@ -44,7 +44,6 @@ Description:Break downs and descriptions of the various nodes available in the E
 
 [toc]
 
-## Overview
 
 This is a quick reference for the Nodes and Tests available in the Environment Query Editor.
 
diff --git a/Engine/Documentation/Source/Engine/AI/EnvironmentQuerySystem/QuickStart/1/1.INT.udn b/Engine/Documentation/Source/Engine/AI/EnvironmentQuerySystem/QuickStart/1/1.INT.udn
new file mode 100644
--- /dev/null
+++ b/Engine/Documentation/Source/Engine/AI/EnvironmentQuerySystem/QuickStart/1/1.INT.udn
@@ -0,0 +1,33 @@
+Availability:Public
+Title: 1. Essential Setup
+Crumbs:%ROOT%, Engine, Gameplay/AI, Engine/AI/EnvironmentQuerySystem, Engine/AI/EnvironmentQuerySystem/QuickStart
+Description: Create a new project using the Top Down Template.
+
+[VAR:Steps]
+[OBJECT:Navigation]
+	[PARAM:previous]
+	[/PARAM]
+	[PARAM:current]
+	[/PARAM]
+	[PARAM:home]
+		[Quick Start Home](Engine/AI/EnvironmentQuerySystem/QuickStart)
+	[/PARAM]
+	[PARAM:next]
+		[Next Step](Engine/AI/EnvironmentQuerySystem/QuickStart/2 "%Engine/AI/EnvironmentQuerySystem/QuickStart/2:title%")
+	[/PARAM]
+[/OBJECT]
+[/VAR]
+
+%Steps%
+
+This Quick Start makes some assumptions about what is available in the project, as such, if you aren't comfortable making substations or problem solving on the fly, then you should setup the project as follows.
+
+[REGION:note]
+If you have already completed the [**Behavior Tree Quick Start**](Engine/AI/BehaviorTrees/QuickStart), feel free to use that project instead of starting a new one.
+[/REGION]
+
+## Steps
+
+[INCLUDE:Engine/AI/BehaviorTrees/QuickStart/1#aiprojectstart]
+
+%Steps%
\ No newline at end of file
diff --git a/Engine/Documentation/Source/Engine/AI/EnvironmentQuerySystem/QuickStart/10/10.INT.udn b/Engine/Documentation/Source/Engine/AI/EnvironmentQuerySystem/QuickStart/10/10.INT.udn
new file mode 100644
--- /dev/null
+++ b/Engine/Documentation/Source/Engine/AI/EnvironmentQuerySystem/QuickStart/10/10.INT.udn
@@ -0,0 +1,46 @@
+Availability:Public
+Title: 10. Creating the Player Context
+Crumbs:%ROOT%, Engine, Gameplay/AI, Engine/AI/EnvironmentQuerySystem, Engine/AI/EnvironmentQuerySystem/QuickStart
+Description: Create a new Context to be used with in the Environment Query System.
+
+[VAR:Steps]
+[OBJECT:Navigation]
+	[PARAM:previous]
+		[Previous Step](Engine/AI/EnvironmentQuerySystem/QuickStart/9 "%Engine/AI/EnvironmentQuerySystem/QuickStart/9:title%")
+	[/PARAM]
+	[PARAM:current]
+	[/PARAM]
+	[PARAM:home]
+		[Quick Start Home](Engine/AI/EnvironmentQuerySystem/QuickStart)
+	[/PARAM]
+	[PARAM:next]
+		[Next Step](Engine/AI/EnvironmentQuerySystem/QuickStart/11 "%Engine/AI/EnvironmentQuerySystem/QuickStart/11:title%")
+	[/PARAM]
+[/OBJECT]
+[/VAR]
+
+%Steps%
+
+While an Environment Query can generate a number of Items based on getting an Actor type, it doesn't do that for Contexts. So, we need a way to have the Player's Pawn as a Context, which is where the **EnvQueryContext_BlueprintBase** comes into play.
+
+
+## Steps
+
+1. Double click the **PlayerContext** Blueprint in the Content Browser to open it for editing.
+1. Override the **ProvideActorsSet** Function:
+
+	[REGION:raw]![](override1.png)(convert:false)[/REGION]
+
+	The **Override** button is hidden until you mouse over the Functions portion of the My Blueprint Panel.
+
+	This will override the default behavior of the function with the Blueprint Graph that opens:
+
+	![](newlyOverriden.png)
+
+1. Add a **Get All Actors Of Class** and connect it between the **Provide Actor Set** Node and the **ReturnNode** like this:
+
+	![](finalGraph.png)
+
+1. Finally set the **Actor Class** of the **Get All Actors Of Class** Node to **Top Down Character**, **Compile**, and **Save**.
+
+%Steps%
diff --git a/Engine/Documentation/Source/Engine/AI/EnvironmentQuerySystem/QuickStart/11/11.INT.udn b/Engine/Documentation/Source/Engine/AI/EnvironmentQuerySystem/QuickStart/11/11.INT.udn
new file mode 100644
--- /dev/null
+++ b/Engine/Documentation/Source/Engine/AI/EnvironmentQuerySystem/QuickStart/11/11.INT.udn
@@ -0,0 +1,83 @@
+Availability:Public
+Title: 11. Creating the FindPlayer EQS
+Crumbs:%ROOT%, Engine, Gameplay/AI, Engine/AI/EnvironmentQuerySystem, Engine/AI/EnvironmentQuerySystem/QuickStart
+Description: Create an Environment Query that finds all instances of the player's TopDownCharacter in the level.
+
+[VAR:Steps]
+[OBJECT:Navigation]
+	[PARAM:previous]
+		[Previous Step](Engine/AI/EnvironmentQuerySystem/QuickStart/10 "%Engine/AI/EnvironmentQuerySystem/QuickStart/10:title%")
+	[/PARAM]
+	[PARAM:current]
+	[/PARAM]
+	[PARAM:home]
+		[Quick Start Home](Engine/AI/EnvironmentQuerySystem/QuickStart)
+	[/PARAM]
+	[PARAM:next]
+		[Next Step](Engine/AI/EnvironmentQuerySystem/QuickStart/12 "%Engine/AI/EnvironmentQuerySystem/QuickStart/12:title%")
+	[/PARAM]
+[/OBJECT]
+[/VAR]
+
+%Steps%
+
+This is the simpler of the two **Environment Queries**, as all it will do is return a list of all the TopDownCharacters within 1500 units of the AI then test if the AI can see the TopDownCharacters. If it can, then we'll get our **ActorToFleeFrom** for our **Behavior Tree**.
+
+There are a number of ways this process can be done, each with their own advantages and disadvantages. Using EQS for this means that we can add additional tests at a later time to choose the _best_ enemy (TopDownCharacter) to run from.
+
+Some quick glossary terms you're going to run into:
+
+* A **Generator** creates **Items** in relation to a **Context**
+* **Items** are used in **Tests** often in relation to a **Context** (how far is Item X from Context Y)
+* **Items** are culled or scored based on the results of **Tests**
+
+## Steps
+
+1. To start, double click the `FindPlayer` **Environment Query** asset in the **Content Browser**.
+
+	![](findPlayerAsset.png)
+
+	Which will open the **Environment Query** editor which will look like this:
+
+	![](startingEQ.png)
+
+	The Environment Query editor works almost exactly like the Behavior Tree editor.
+
+1. Right-Click into the **Update Graph** area, and select **ActorsOfClass** from the context menu:
+
+	![](createGetActors.png)
+
+	Which will create the ActorsOfClass Generator like so:
+
+	![](cga_made.png)
+
+1. To connect nodes in the Environment Query editor, you just need to drag from the bottom of the Root node to the Top of the **ActorsOfClass** node, exactly how it is done in the **Behavior Tree** editor.
+
+	![](connected.png)
+
+	There can only be one Generator per **Environment Query**
+
+1. Next, setup the **ActorsOfClass** node by selecting it, and changing its **Search Radius** and **Searched Actor Class** properties:
+
+	![](aoc_details.png)
+
+	* **Search Radius** = 1500
+	* **Searched Actor Class** = TopDownCharacter.
+
+	This **Generator** finds all the Actors of the given Class, within the radius, and creates **Items** from them.
+
+1. Add a **Trace** test to the **ActorsOfClass** node
+
+	![](addTraceTest.png)
+
+	Trace tests are useful to check if a Context has line of sight to an Item, in this case, our AI to the Items (TopDownCharacters/Players in this case)
+
+1. Now select the Trace Test on the **ActorsOfClass** node by left-clicking it, and in the Details Panel, change the **Bool Match** Property to `false` (unchecked)
+
+	![](aoc_TraceFinal.png)
+
+
+And that's it!
+
+
+%Steps%
diff --git a/Engine/Documentation/Source/Engine/AI/EnvironmentQuerySystem/QuickStart/12/12.INT.udn b/Engine/Documentation/Source/Engine/AI/EnvironmentQuerySystem/QuickStart/12/12.INT.udn
new file mode 100644
--- /dev/null
+++ b/Engine/Documentation/Source/Engine/AI/EnvironmentQuerySystem/QuickStart/12/12.INT.udn
@@ -0,0 +1,135 @@
+Availability:Public
+Title: 12. Creating the FindHidingSpot EQS
+Crumbs:%ROOT%, Engine, Gameplay/AI, Engine/AI/EnvironmentQuerySystem, Engine/AI/EnvironmentQuerySystem/QuickStart
+Description: Create an EQS that finds the best point on a grid to hide from the Player's TopDownCharacter.
+
+[VAR:Steps]
+[OBJECT:Navigation]
+	[PARAM:previous]
+		[Previous Step](Engine/AI/EnvironmentQuerySystem/QuickStart/11 "%Engine/AI/EnvironmentQuerySystem/QuickStart/11:title%")
+	[/PARAM]
+	[PARAM:current]
+	[/PARAM]
+	[PARAM:home]
+		[Quick Start Home](Engine/AI/EnvironmentQuerySystem/QuickStart)
+	[/PARAM]
+	[PARAM:next]
+		[Next Step](Engine/AI/EnvironmentQuerySystem/QuickStart/13 "%Engine/AI/EnvironmentQuerySystem/QuickStart/13:title%")
+	[/PARAM]
+[/OBJECT]
+[/VAR]
+
+%Steps%
+
+For this step you'll be creating a much more involved Environment Query. It's going to create a grid of points, projected onto the NavMesh, then run a series of tests to figure out the "best" place to hide. Some of the tests will simply cull out completely unsatisfactory locations, such as those that are visible to the Player's Character (Trace), others will weight the potential hiding spots based on if they are reachable on the NavMesh (PathFinding) or how much a point is in front of the Player's Character (Dot Product).
+
+Just a quick reminder:
+
+* A **Generator** creates **Items** in relation to a **Context**
+* **Items** are used in **Tests** often in relation to a **Context** (how far is Item X from Context Y)
+* **Items** are culled or scored based on the results of **Tests**
+
+## Steps
+
+1. To start, double click the `FindHidingSpot` **Environment Query** asset in the **Content Browser**.
+
+	![](findHide.png)
+
+1. Create **Simple Grid** Generator and connect it to the **Root** node
+
+	![](simpleGrid.png)
+
+1. Select the Simple Grid Generator and in the **Details** Panel set the following Properties:
+	
+	 ![](gridProp.png)
+
+	[INCLUDE:#eqsqs121]
+
+	These settings for the Grid Generator will create a sizable, but manageable, number of Items for the Query to test against. The **Project Down**, **Project Up,** and **Post Projection Vertical Offset** properties are particularly useful to change if you have an AI that needs to find something above or below them.
+
+1. [REGION:raw]Add 2 **Distance** Tests by Right-Clicking the **Simple Grid** node and selecting **Tests** %arrow% **Distance**[/REGION]
+
+	![](addDistance.png)
+
+1. Select the first one and in the **Details** Panel adjust the following Properties:
+
+	![](distanceProps.png)
+
+	[INCLUDE:#eqsqs122]
+
+	This Distance test is meant to give **Items** _further_ from the Player a higher Score.
+
+1. Select the second one and in the **Details** Panel adjust the following Properties:
+
+	![](distanceProps2.png)
+
+	[INCLUDE:#eqsqs123]
+
+	This Distance test is meant to give **Items** _closer_ to the AI a higher Score, and this test carries more weight than the previous **Distance Test** meaning it will favor closer to the AI but still select away from the Player.
+
+1. [REGION:raw]Add 2 **Trace** Tests by Right-Clicking the **Simple Grid** node and selecting **Tests** %arrow% **Trace**[/REGION]
+
+	![](addTrace.png)
+
+1. Select the one of the Trace tests and in the **Details** panel change **Context** to `PlayerContext` (the Context created earlier in this Quick Start):
+
+	![](traceProps.png)	
+
+	Here we remove any Item that the Player Pawn or AI pawn has line of sight to. Also, you may be concerned that this test comes after a bunch of **Distance** tests, but it's ok, as the Environment Query runs **Trace** tests first regardless of order here. Also, the other Trace requires no additional setup.
+
+1. [REGION:raw]Add a **Dot** test by Right-Clicking the **Simple Grid** node and selecting **Tests** %arrow% **Dot**[/REGION]
+
+	![](addDot.png)
+
+1. Select the Dot test and in the **Details** panel adjust the following properties: 
+	
+	![](dotProps.png)
+
+	[INCLUDE:#eqsqs125]
+
+	Here things get tricky, at the high level, this is saying "Favor Items that are to the sides of the Player Pawn," but how it works is a bit odd. A Dot Product function will return how "alike" two vectors are on a scale of 1.0 (perfectly the same) to -1.0 (completely opposite). So if we get the **Absolute Value** of the Dot Product between the rotation vector of the Player and the vector from the Player to the Item we'll favor points in front of and behind the Player. But if we set the **Scoring Factor** to -1.0, then we favor to the sides of the Player.
+
+1. [REGION:raw]Add a **PathFinding** test by Right-Clicking the **Simple Grid** node and selecting **Tests** %arrow% **PathFinding**[/REGION]
+
+	![](addPath.png)
+
+1. Select the PathFinding test and in the **Details** panel adjust **Test Purpose** to `Score Only`:
+
+	![](pathProps.png)
+
+	And that's it for the editing! The final Environment Query will look similar to this:
+
+	![](final.png)
+
+%Steps%
+
+[VAR:arrow]![](Engine\AI\EnvironmentQuerySystem\QuickStart\arrow_button.png)[/VAR]
+
+<!--
+[EXCERPT:eqsqs125]
+* **Line A - Rotation** = `PlayerContext`
+* **Line B - Line From** = `PlayerContext`
+* **Absolute Value** = `True`
+* **Test Purpose** = `Score Only`
+* **Scoring Factor** = `-1.0`
+[/EXCERPT]
+
+[EXCERPT:eqsqs121]
+* Grid Size = `2000.0`
+* Space Between = `150.0`
+* Navigation Filer = `RecastFilter_UseDefaultArea`
+* Project Down = `2048.0`
+* Post Projection Vertical Offset = `120.0`
+[/EXCERPT]
+
+[EXCERPT:eqsqs122]
+* **Distance To** = `Player Context`
+* **Test Purpose** = `Score Only`
+[/EXCERPT]
+
+[EXCERPT:eqsqs123]
+* **Test Purpose** = `Score Only`
+* **Scoring Equation** = `Inverse Linear`
+* **Scoring Factor** = `2.0`
+[/EXCERPT]
+-->
\ No newline at end of file
diff --git a/Engine/Documentation/Source/Engine/AI/EnvironmentQuerySystem/QuickStart/13/13.INT.udn b/Engine/Documentation/Source/Engine/AI/EnvironmentQuerySystem/QuickStart/13/13.INT.udn
new file mode 100644
--- /dev/null
+++ b/Engine/Documentation/Source/Engine/AI/EnvironmentQuerySystem/QuickStart/13/13.INT.udn
@@ -0,0 +1,36 @@
+Availability:Public
+Title: 13. Play in Editor and Testing
+Crumbs:%ROOT%, Engine, Gameplay/AI, Engine/AI/EnvironmentQuerySystem, Engine/AI/EnvironmentQuerySystem/QuickStart
+Description: Finally, test your Runner AI by using Play in Editor.
+
+[VAR:Steps]
+[OBJECT:Navigation]
+	[PARAM:previous]
+		[Previous Step](Engine/AI/EnvironmentQuerySystem/QuickStart/12 "%Engine/AI/EnvironmentQuerySystem/QuickStart/12:title%")
+	[/PARAM]
+	[PARAM:current]
+	[/PARAM]
+	[PARAM:home]
+		[Quick Start Home](Engine/AI/EnvironmentQuerySystem/QuickStart)
+	[/PARAM]
+	[PARAM:next]
+		[Next Step](Engine/AI/EnvironmentQuerySystem/QuickStart/14 "%Engine/AI/EnvironmentQuerySystem/QuickStart/14:title%")
+	[/PARAM]
+[/OBJECT]
+[/VAR]
+
+%Steps%
+
+Now since you setup the **Behavior Tree** [earlier](Engine/AI/EnvironmentQuerySystem/QuickStart/9), and told the [AI Controller to run](Engine/AI/EnvironmentQuerySystem/QuickStart/7) it, you are ready to test! 
+
+
+## Steps
+
+* Hit the **Play in Editor** button and watch your AI flee from you.
+
+![](play.png)
+
+![](running.gif)(convert:false)
+
+
+%Steps%
diff --git a/Engine/Documentation/Source/Engine/AI/EnvironmentQuerySystem/QuickStart/14/14.INT.udn b/Engine/Documentation/Source/Engine/AI/EnvironmentQuerySystem/QuickStart/14/14.INT.udn
new file mode 100644
--- /dev/null
+++ b/Engine/Documentation/Source/Engine/AI/EnvironmentQuerySystem/QuickStart/14/14.INT.udn
@@ -0,0 +1,54 @@
+Availability:Public
+Title: 14. Next Steps
+Crumbs:%ROOT%, Engine, Gameplay/AI, Engine/AI/EnvironmentQuerySystem, Engine/AI/EnvironmentQuerySystem/QuickStart
+Description: Suggestions on where to go after completing the EQS Quick Start.
+
+[VAR:Steps]
+[OBJECT:Navigation]
+	[PARAM:previous]
+		[Previous Step](Engine/AI/EnvironmentQuerySystem/QuickStart/13 "%Engine/AI/EnvironmentQuerySystem/QuickStart/13:title%")
+	[/PARAM]
+	[PARAM:current]
+	[/PARAM]
+	[PARAM:home]
+		[Quick Start Home](Engine/AI/EnvironmentQuerySystem/QuickStart)
+	[/PARAM]
+	[PARAM:next]
+	[/PARAM]
+[/OBJECT]
+[/VAR]
+
+%Steps%
+
+This page covers a number of ideas for improvements you can make now that you're familiar with the basics of the Environment Query System.
+
+## Ideas
+
+* **TopDownCharacter** is far to specific, update everything to a more flexible system.
+* The **Behavior Tree** shows one way to set a system like this up:
+	* Try using a **Wait** node instead of the **Cooldown** Decorator. This will require a **Bool** Blackboard entry and some extra work in the AI Controller.
+	* Add a secondary Behavior that, after a time, the AI will begin to wander around.
+	* **Bonus** - Store the last know location(s) of the Player(s) and use those in an Environment Query to find a better hiding spot.
+	* **Double Bonus** - Have the AI attempt to approach a Player from behind without entering the forward view arc of the Player(s).
+* Did you know you can create your own **Generators** in Blueprint? Maybe there's a better way to find the possible points for the AI to hide.
+* [REGION:raw]The ' key (single quote) turns on a host of Gameplay Debugging tools. Check out all of the options available in the **Editor Preferences** %arrow% **Gameplay Debugger**[/REGION]
+
+[OBJECT:TopicButtonList]
+	[PARAM:icon]
+		![](%ROOT%/docs_icon.png)(convert:false)
+	[/PARAM]
+	[PARAM:title]
+		AI
+	[/PARAM]
+	[PARAM:links]
+	* [](Engine/AI/EnvironmentQuerySystem/EQSPawn "%Engine/AI/EnvironmentQuerySystem/EQSPawn:description%")
+	* [](Engine/AI/EnvironmentQuerySystem/NodeReference "%Engine/AI/EnvironmentQuerySystem/NodeReference:description%")
+	* [](Engine/AI/EnvironmentQuerySystem/UserGuide "%Engine/AI/EnvironmentQuerySystem/UserGuide:description%")
+	* [Behavior Trees](Engine/AI/BehaviorTrees "%Engine/AI/BehaviorTrees:description%")
+	[/PARAM]
+[/OBJECT]
+
+
+%Steps%
+
+[VAR:arrow]![](Engine\AI\EnvironmentQuerySystem\QuickStart\arrow_button.png)[/VAR]
\ No newline at end of file
diff --git a/Engine/Documentation/Source/Engine/AI/EnvironmentQuerySystem/QuickStart/2/2.INT.udn b/Engine/Documentation/Source/Engine/AI/EnvironmentQuerySystem/QuickStart/2/2.INT.udn
new file mode 100644
--- /dev/null
+++ b/Engine/Documentation/Source/Engine/AI/EnvironmentQuerySystem/QuickStart/2/2.INT.udn
@@ -0,0 +1,37 @@
+Availability:Public
+Title: 2. Enable the Environment Query System
+Crumbs:%ROOT%, Engine, Gameplay/AI, Engine/AI/EnvironmentQuerySystem, Engine/AI/EnvironmentQuerySystem/QuickStart
+Description: Enable the Environment Query System as it is currently experimental.
+
+[VAR:Steps]
+[OBJECT:Navigation]
+	[PARAM:previous]
+		[Previous Step](Engine/AI/EnvironmentQuerySystem/QuickStart/1 "%Engine/AI/EnvironmentQuerySystem/QuickStart/1:title%")
+	[/PARAM]
+	[PARAM:current]
+	[/PARAM]
+	[PARAM:home]
+		[Quick Start Home](Engine/AI/EnvironmentQuerySystem/QuickStart)
+	[/PARAM]
+	[PARAM:next]
+		[Next Step](Engine/AI/EnvironmentQuerySystem/QuickStart/3 "%Engine/AI/EnvironmentQuerySystem/QuickStart/3:title%")
+	[/PARAM]
+[/OBJECT]
+[/VAR]
+
+%Steps%
+
+The **Environment Query System** is considered **Experimental**, meaning it can change drastically at any time, so you must "opt in" to it.
+
+## Steps
+
+* [REGION:raw] As of 4.7 the Environment Querying System is enabled by going to **Edit Menu** %arrow% **Editor Preferences** %arrow% **Experimental**, and checking the check box next to **Environment Querying System** under the AI category. [/REGION]
+
+![](enableEQS.png)
+
+This will eventually not be necessary, so if you don't see this check box, then you don't need to do this step.
+
+%Steps%
+
+
+[VAR:arrow]![](Engine\AI\EnvironmentQuerySystem\QuickStart\arrow_button.png)[/VAR]
\ No newline at end of file
diff --git a/Engine/Documentation/Source/Engine/AI/EnvironmentQuerySystem/QuickStart/3/3.INT.udn b/Engine/Documentation/Source/Engine/AI/EnvironmentQuerySystem/QuickStart/3/3.INT.udn
new file mode 100644
--- /dev/null
+++ b/Engine/Documentation/Source/Engine/AI/EnvironmentQuerySystem/QuickStart/3/3.INT.udn
@@ -0,0 +1,27 @@
+Availability:Public
+Title: 3. Create Nav Mesh Volume
+Crumbs:%ROOT%, Engine, Gameplay/AI, Engine/AI/EnvironmentQuerySystem, Engine/AI/EnvironmentQuerySystem/QuickStart
+Description: Create a Nav Mesh Volume so our AI can walk around.
+
+[VAR:Steps]
+[OBJECT:Navigation]
+	[PARAM:previous]
+		[Previous Step](Engine/AI/EnvironmentQuerySystem/QuickStart/2 "%Engine/AI/EnvironmentQuerySystem/QuickStart/2:title%")
+	[/PARAM]
+	[PARAM:current]
+	[/PARAM]
+	[PARAM:home]
+		[Quick Start Home](Engine/AI/EnvironmentQuerySystem/QuickStart)
+	[/PARAM]
+	[PARAM:next]
+		[Next Step](Engine/AI/EnvironmentQuerySystem/QuickStart/4 "%Engine/AI/EnvironmentQuerySystem/QuickStart/4:title%")
+	[/PARAM]
+[/OBJECT]
+[/VAR]
+
+%Steps%
+
+[INCLUDE:Engine/AI/BehaviorTrees/QuickStart/2#aicreatenav]
+
+%Steps%
+
diff --git a/Engine/Documentation/Source/Engine/AI/EnvironmentQuerySystem/QuickStart/4/4.INT.udn b/Engine/Documentation/Source/Engine/AI/EnvironmentQuerySystem/QuickStart/4/4.INT.udn
new file mode 100644
--- /dev/null
+++ b/Engine/Documentation/Source/Engine/AI/EnvironmentQuerySystem/QuickStart/4/4.INT.udn
@@ -0,0 +1,81 @@
+Availability:Public
+Title: 4. Create the Basic Assets
+Crumbs:%ROOT%, Engine, Gameplay/AI, Engine/AI/EnvironmentQuerySystem, Engine/AI/EnvironmentQuerySystem/QuickStart
+Description: Create the 7 needed assets for the EQS Quick Start.
+
+[VAR:Steps]
+[OBJECT:Navigation]
+	[PARAM:previous]
+		[Previous Step](Engine/AI/EnvironmentQuerySystem/QuickStart/3 "%Engine/AI/EnvironmentQuerySystem/QuickStart/3:title%")
+	[/PARAM]
+	[PARAM:current]
+	[/PARAM]
+	[PARAM:home]
+		[Quick Start Home](Engine/AI/EnvironmentQuerySystem/QuickStart)
+	[/PARAM]
+	[PARAM:next]
+		[Next Step](Engine/AI/EnvironmentQuerySystem/QuickStart/5 "%Engine/AI/EnvironmentQuerySystem/QuickStart/5:title%")
+	[/PARAM]
+[/OBJECT]
+[/VAR]
+
+%Steps%
+
+To start, you will need create 7 assets in the **Content Browser**: An **AI Controller** Blueprint, A **Character** Blueprint, and **Behavior Tree** Asset, a **Blackboard** Asset, an **EnvQueryContext_BlueprintBase** (EQS Context) Blueprint, and 2 **Environment Query** templates. These are all the assets you will need to complete this Quick Start.
+
+[REGION:raw]
+![](Assets.png)(convert:false)
+[/REGION]
+
+
+##Steps
+
+1. [REGION:raw]**New** %arrow% **Blueprint** %arrow% **AI Controller**.[/REGION] 
+
+	You will need to search for AI Controller, as they are different from Player Controller.
+
+	![](Engine/AI/BehaviorTrees/QuickStart/3/AIC.png)
+
+	Name this **Blueprint** `Runner_AI_CON`.
+
+1. [REGION:raw]**New** %arrow% **Blueprint** %arrow% **Character**.[/REGION]
+
+	![](Engine/AI/BehaviorTrees/QuickStart/3/charBP.png)
+
+	Name this **Blueprint** `AI_Character`.
+
+1. [REGION:raw]**New** %arrow% **Miscellaneous** %arrow% **Behavior Tree**.[/REGION]
+
+	[REGION:raw]
+	![](Engine/AI/BehaviorTrees/QuickStart/3/newBT.png)
+	[/REGION]
+
+	Name the **Behavior Tree** `RunnerBT`.
+	
+1. [REGION:raw]**New** %arrow% **Miscellaneous** %arrow%  **Blackboard**.[/REGION]
+	
+	[REGION:raw]
+	![](Engine/AI/BehaviorTrees/QuickStart/3/newBB.png)
+	[/REGION]	
+
+	Name the **Blackboard** `RunnerBlackboard`.
+
+1. Create an **Environment Query**: [REGION:raw]**New** %arrow% **Miscellaneous** %arrow%  **Environment Query**.[/REGION]
+	
+	[REGION:raw]
+	![](newEQ.PNG)
+	[/REGION]	
+
+	Name one **Environment Query** `FindHidingSpot`.
+
+1. Name the other **Environment Query** `FindPlayer`.
+
+1. [REGION:raw]**New** %arrow% **Blueprint** %arrow%  **EnvQueryContext_BlueprintBase**.[/REGION]
+	
+	![](newContext.png)
+
+	Name the **Environment Query** `PlayerContext`.
+
+%Steps%
+
+[VAR:arrow]![](Engine\AI\EnvironmentQuerySystem\QuickStart\arrow_button.png)[/VAR]
\ No newline at end of file
diff --git a/Engine/Documentation/Source/Engine/AI/EnvironmentQuerySystem/QuickStart/5/5.INT.udn b/Engine/Documentation/Source/Engine/AI/EnvironmentQuerySystem/QuickStart/5/5.INT.udn
new file mode 100644
--- /dev/null
+++ b/Engine/Documentation/Source/Engine/AI/EnvironmentQuerySystem/QuickStart/5/5.INT.udn
@@ -0,0 +1,73 @@
+Availability:Public
+Title: 5. Character Blueprint
+Crumbs:%ROOT%, Engine, Gameplay/AI, Engine/AI/EnvironmentQuerySystem, Engine/AI/EnvironmentQuerySystem/QuickStart
+Description: Setup the AI's Character Blueprint.
+
+[VAR:Steps]
+[OBJECT:Navigation]
+	[PARAM:previous]
+		[Previous Step](Engine/AI/EnvironmentQuerySystem/QuickStart/4 "%Engine/AI/EnvironmentQuerySystem/QuickStart/4:title%")
+	[/PARAM]
+	[PARAM:current]
+	[/PARAM]
+	[PARAM:home]
+		[Quick Start Home](Engine/AI/EnvironmentQuerySystem/QuickStart)
+	[/PARAM]
+	[PARAM:next]
+		[Next Step](Engine/AI/EnvironmentQuerySystem/QuickStart/6 "%Engine/AI/EnvironmentQuerySystem/QuickStart/6:title%")
+	[/PARAM]
+[/OBJECT]
+[/VAR]
+
+%Steps%
+
+As we have created a project with a **Character** already in it, we have access to a humanoid **Skeletal Mesh** and its **Animation Blueprint**. We will use these in the tutorial, but feel free to use your own if you have them available.
+
+This Character will become our new AI Character. The reason we use a Character, over say a Pawn, is to gain access to the Character Movement Component's functionality.
+
+## Steps
+
+1. **Double-click** the `AI_Character` Blueprint in the **Content Browser** to open it for editing.
+
+1. Select the **Mesh** component in the **Components** panel.
+
+	![](Engine\AI\BehaviorTrees\QuickStart\4\meshComp.png)
+
+1. In the **Details** panel, find the **Mesh** category and assign the `HeroTPP` (or your own) **Skeletal Mesh**.
+
+	![](Engine\AI\BehaviorTrees\QuickStart\4\assignSkelMesh.png)
+
+1. Move the **Mesh** Component to center it up inside in the capsule.
+
+1. Rotate the **Mesh** Component so it faces the direction the blue **Arrow Component** is pointing.
+
+1. In the **Details** panel, find the **Animation** category and assign the `HeroTPP_AnimBlueprint` to the **Animation Blueprint Generated Class** property.
+
+	![](Engine\AI\BehaviorTrees\QuickStart\4\assignAnimBP.png)
+
+1. Select the **CapsuleComponent** from the **Components** panel, then in the **Details** panel, adjust the **Capsule Half Height** and **Capsule Radius** so they encompass the **Skeletal Mesh**.
+
+	![](Engine\AI\BehaviorTrees\QuickStart\4\shapeCat.png)
+
+1. Select the **Character Movement** Component.
+
+	![](Engine\AI\BehaviorTrees\QuickStart\4\charmve.png)
+
+1. In the **Details** panel, in the Movement Component category, adjust the Agent Radius and Agent Height values.
+
+	![](Engine\AI\BehaviorTrees\QuickStart\4\agentProps.png)
+
+	_Set the Agent Radius to at least the Capsule's radius, and the Agent Height to at least twice the Capsule's "half height."_
+
+1. With the Character Movement Component still selected, change the **Max Walk Speed** property, in the **Character Movement** category to `400`.
+	
+	![](Engine\AI\BehaviorTrees\QuickStart\4\maxWalk.png)
+
+1. Click the **Class Defaults** ![](Engine\Blueprints\UserGuide\Variables\defaultsButton.png) button and assign the new AI Controller (`Runner_AI_CON`) to the **AIController Class** property in the **Pawn** category.
+
+	![](AI_con.png)
+
+1. Compile ![](Engine\AI\BehaviorTrees\QuickStart\4\compile_button.png) and Save ![](Engine\AI\BehaviorTrees\QuickStart\4\save_button.png)
+
+
+%Steps%
\ No newline at end of file
diff --git a/Engine/Documentation/Source/Engine/AI/EnvironmentQuerySystem/QuickStart/6/6.INT.udn b/Engine/Documentation/Source/Engine/AI/EnvironmentQuerySystem/QuickStart/6/6.INT.udn
new file mode 100644
--- /dev/null
+++ b/Engine/Documentation/Source/Engine/AI/EnvironmentQuerySystem/QuickStart/6/6.INT.udn
@@ -0,0 +1,37 @@
+Availability:Public
+Title: 6. Blackboard Setup
+Crumbs:%ROOT%, Engine, Gameplay/AI, Engine/AI/EnvironmentQuerySystem, Engine/AI/EnvironmentQuerySystem/QuickStart
+Description: Adding the needed Key Values to the Runner Blackboard.
+
+[VAR:Steps]
+[OBJECT:Navigation]
+	[PARAM:previous]
+		[Previous Step](Engine/AI/EnvironmentQuerySystem/QuickStart/5 "%Engine/AI/EnvironmentQuerySystem/QuickStart/5:title%")
+	[/PARAM]
+	[PARAM:current]
+	[/PARAM]
+	[PARAM:home]
+		[Quick Start Home](Engine/AI/EnvironmentQuerySystem/QuickStart)
+	[/PARAM]
+	[PARAM:next]
+		[Next Step](Engine/AI/EnvironmentQuerySystem/QuickStart/7 "%Engine/AI/EnvironmentQuerySystem/QuickStart/7:title%")
+	[/PARAM]
+[/OBJECT]
+[/VAR]
+
+%Steps%
+
+The **Blackboard** is the "memory" of the **Behavior Tree**. Here we will store a reference to the Actor we should flee from and a location to flee to.
+
+## Steps
+
+1. **Double-click** the Blackboard asset in the **Content Browser**.
+1. Using the **New Key** button ![](Engine/AI/BehaviorTrees/QuickStart/5/newKey_button.png),  add a **Vector** key named **FleeToLocation**.
+1. Add an **Object** key named **ActorToFleeFrom**.
+
+	![](finalBB.png)
+
+1. Save your **Blackboard** asset.
+
+
+%Steps%
\ No newline at end of file
diff --git a/Engine/Documentation/Source/Engine/AI/EnvironmentQuerySystem/QuickStart/7/7.INT.udn b/Engine/Documentation/Source/Engine/AI/EnvironmentQuerySystem/QuickStart/7/7.INT.udn
new file mode 100644
--- /dev/null
+++ b/Engine/Documentation/Source/Engine/AI/EnvironmentQuerySystem/QuickStart/7/7.INT.udn
@@ -0,0 +1,32 @@
+Availability:Public
+Title: 7. AI Controller Setup
+Crumbs:%ROOT%, Engine, Gameplay/AI, Engine/AI/EnvironmentQuerySystem, Engine/AI/EnvironmentQuerySystem/QuickStart
+Description: Setup the AI Controller for the EQS Quick Start.
+
+[VAR:Steps]
+[OBJECT:Navigation]
+	[PARAM:previous]
+		[Previous Step](Engine/AI/EnvironmentQuerySystem/QuickStart/6 "%Engine/AI/EnvironmentQuerySystem/QuickStart/6:title%")
+	[/PARAM]
+	[PARAM:current]
+	[/PARAM]
+	[PARAM:home]
+		[Quick Start Home](Engine/AI/EnvironmentQuerySystem/QuickStart)
+	[/PARAM]
+	[PARAM:next]
+		[Next Step](Engine/AI/EnvironmentQuerySystem/QuickStart/8 "%Engine/AI/EnvironmentQuerySystem/QuickStart/8:title%")
+	[/PARAM]
+[/OBJECT]
+[/VAR]
+
+%Steps%
+
+The setup for this AI Controller is a single **Run Behavior Tree** node running our `RunnerBT` **Behavior Tree** when ever **Event BeginPlay** is executed.
+
+## Steps
+
+1. In the Event Graph, have **On Begin Play** execute a **Run Beahvior Tree** node using the `RunnerBT` as the **BTAsset**:
+
+	![](setup.png)
+
+%Steps%
diff --git a/Engine/Documentation/Source/Engine/AI/EnvironmentQuerySystem/QuickStart/8/8.INT.udn b/Engine/Documentation/Source/Engine/AI/EnvironmentQuerySystem/QuickStart/8/8.INT.udn
new file mode 100644
--- /dev/null
+++ b/Engine/Documentation/Source/Engine/AI/EnvironmentQuerySystem/QuickStart/8/8.INT.udn
@@ -0,0 +1,26 @@
+Availability:Public
+Title: 8. Place the AI Character
+Crumbs:%ROOT%, Engine, Gameplay/AI, Engine/AI/EnvironmentQuerySystem, Engine/AI/EnvironmentQuerySystem/QuickStart
+Description: Place the Runner AI Character into the level so it can interact with the Player and environment.
+
+[VAR:Steps]
+[OBJECT:Navigation]
+	[PARAM:previous]
+		[Previous Step](Engine/AI/EnvironmentQuerySystem/QuickStart/7 "%Engine/AI/EnvironmentQuerySystem/QuickStart/7:title%")
+	[/PARAM]
+	[PARAM:current]
+	[/PARAM]
+	[PARAM:home]
+		[Quick Start Home](Engine/AI/EnvironmentQuerySystem/QuickStart)
+	[/PARAM]
+	[PARAM:next]
+		[Next Step](Engine/AI/EnvironmentQuerySystem/QuickStart/9 "%Engine/AI/EnvironmentQuerySystem/QuickStart/9:title%")
+	[/PARAM]
+[/OBJECT]
+[/VAR]
+
+%Steps%
+
+[INCLUDE:Engine/AI/BehaviorTrees/QuickStart/9#s9]
+
+%Steps%
\ No newline at end of file
diff --git a/Engine/Documentation/Source/Engine/AI/EnvironmentQuerySystem/QuickStart/9/9.INT.udn b/Engine/Documentation/Source/Engine/AI/EnvironmentQuerySystem/QuickStart/9/9.INT.udn
new file mode 100644
--- /dev/null
+++ b/Engine/Documentation/Source/Engine/AI/EnvironmentQuerySystem/QuickStart/9/9.INT.udn
@@ -0,0 +1,98 @@
+Availability:Public
+Title: 9. Behavior Tree Setup
+Crumbs:%ROOT%, Engine, Gameplay/AI, Engine/AI/EnvironmentQuerySystem, Engine/AI/EnvironmentQuerySystem/QuickStart
+Description: Setup the EQS Quick Start Behavior Tree.
+
+[VAR:Steps]
+[OBJECT:Navigation]
+	[PARAM:previous]
+		[Previous Step](Engine/AI/EnvironmentQuerySystem/QuickStart/8 "%Engine/AI/EnvironmentQuerySystem/QuickStart/8:title%")
+	[/PARAM]
+	[PARAM:current]
+	[/PARAM]
+	[PARAM:home]
+		[Quick Start Home](Engine/AI/EnvironmentQuerySystem/QuickStart)
+	[/PARAM]
+	[PARAM:next]
+		[Next Step](Engine/AI/EnvironmentQuerySystem/QuickStart/10 "%Engine/AI/EnvironmentQuerySystem/QuickStart/10:title%")
+	[/PARAM]
+[/OBJECT]
+[/VAR]
+
+%Steps%
+
+Since you already have all the assets necessary to fill out the tree, you're going to be building this Behavior Tree:
+
+![](finalBT_Setup.png)
+
+Before we get into the steps to create this [**Behavior Tree**](Engine/AI/BehaviorTrees/UserGuide), let's talk about what it does (or rather what it will do once we create the necessary EQS bits). 
+
+* We enter the graph at the top Sequence Node, and then start executing it's children from left to right, and will keep executing as long as its children are successful.
+* First is the EQS Query node that runs the _FindPlayer_ template. We only allow this node to execute once every 0.5 seconds, and if it fails, the tree is locked out of executing the other nodes by the Cooldown Node returning **failed**, thusly failing its parent sequence node.
+* But, if the EQS Query is successful (it finds the player), then the second EQS Query is allowed to execute, finding the AI a place to hide.
+* Finally we move to the place that the second EQS Query found.
+
+## Steps
+
+1. [REGION:raw]Add a **Sequence Node** (right-click in the Graph Area %arrow% **Composites** %arrow% **Sequence**)[/REGION]
+
+	![](addSeq.png)
+
+1. Connect the **Root Node** to this new **Sequence Node**
+
+	![](connectSeq.png)
+
+1. [REGION:raw]Create a **Run EQS Query Node** (right-click in the Graph Area %arrow% **Tasks** %arrow% **Run EQSQuery**)[/REGION]
+
+	![](addEqs1.png)
+
+1. Connect it to the **Sequence Node**
+
+	![](connectEqs1.png)
+
+1. Select this **Run EQS Query Node**, and in the details panel, set the **Query Template** to `FindPlayer`
+1. While still in the Details panel set the **Blackboard Key** to `ActorToFleeFrom`
+
+	![](setEqs1props.png)
+
+1. [REGION:raw]Add a **Cooldown Decorator** to the **Run EQS Query Node** (right-click the **Run EQS Query Node** %arrow% **Decorators** %arrow% **Cooldown**)[/REGION]
+
+	![](addCooldown.png)
+
+1. Select the **Cooldown Decorator** and set its **Cool Down Time** to `0.5`
+
+	![](setCooldown.png)
+
+1. [REGION:raw]Create another **Run EQS Query Node** (right-click in the Graph Area %arrow% **Tasks** %arrow% **Run EQSQuery**)[/REGION] and connect it to the **Sequence Node**
+	
+	![](addEqs2.png)
+
+1. Select the second **Run EQS Query Node**, and in the details panel, set the **Query Template** to `FindHidingSpot`
+
+1. Again, while still in the Details panel, set the **Blackboard Key** to **FleeToLocation**
+
+	![](setEqs2props.png)
+
+1. Position it to the right of the second **Run EQS Query Node**
+1. [REGION:raw]Create a **Move To Node** (right-click in the Graph Area %arrow% **Tasks** %arrow% **Move To**)[/REGION]
+
+	![](addMoveTo.png)
+
+1. Connect it to the **Sequence Node**
+
+	![](connectMoveTo.png)
+
+1. Select the Move To node and set **Blackboard Key** to `FleeToLocation`
+1. While still in the Details Panel, set **Acceptable Radius** to `0.5`
+
+	![](setMoveTo.png)
+
+1. Position it to the right of the second **Run EQS Query Node**
+
+	![](finalBT_Setup.png)
+
+With everything connected the Behavior Tree will is ready and no further edits to it will be necessary.
+
+%Steps%
+
+[VAR:arrow]![](Engine\AI\EnvironmentQuerySystem\QuickStart\arrow_button.png)[/VAR]
\ No newline at end of file
diff --git a/Engine/Documentation/Source/Engine/AI/EnvironmentQuerySystem/QuickStart/QuickStart.INT.udn b/Engine/Documentation/Source/Engine/AI/EnvironmentQuerySystem/QuickStart/QuickStart.INT.udn
new file mode 100644
--- /dev/null
+++ b/Engine/Documentation/Source/Engine/AI/EnvironmentQuerySystem/QuickStart/QuickStart.INT.udn
@@ -0,0 +1,79 @@
+Availability:Public
+Title:Environment Query System Quick Start Guide
+Crumbs:%ROOT%, Engine, Gameplay/AI, Engine/AI/EnvironmentQuerySystem
+Description:This Quick Start Guide covers creating a Behavior Tree that uses multiple Run EQS Query nodes to find a place to hide from the player character.
+
+[VAR:Topic]
+[OBJECT:Topic]
+	[PARAM:image]
+		![%Engine/AI/EnvironmentQuerySystem/QuickStart:title%](AI_Topic.png)
+	[/PARAM]
+	[PARAM:icon]
+		![](%ROOT%/start_icon.png)(convert:false)
+	[/PARAM]
+	[PARAM:title]
+		%Engine/AI/EnvironmentQuerySystem/QuickStart:title%
+	[/PARAM]
+	[PARAM:description]
+		%Engine/AI/EnvironmentQuerySystem/QuickStart:description%
+	[/PARAM]
+	[PARAM:path]
+		[RELATIVE:Engine/AI/EnvironmentQuerySystem/QuickStart]
+	[/PARAM]
+[/OBJECT]
+[/VAR]
+[VAR:TopicCompact]
+[OBJECT:TopicCompact]
+	[PARAM:image]
+		![%Engine/AI/EnvironmentQuerySystem/QuickStart:title%](AI_Topic.png)
+	[/PARAM]
+	[PARAM:icon]
+		![](%ROOT%/start_icon.png)(convert:false)
+	[/PARAM]
+	[PARAM:title]
+		%Engine/AI/EnvironmentQuerySystem/QuickStart:title%
+	[/PARAM]
+	[PARAM:description]
+		%Engine/AI/EnvironmentQuerySystem/QuickStart:description%
+	[/PARAM]
+	[PARAM:path]
+		[RELATIVE:Engine/AI/EnvironmentQuerySystem/QuickStart]
+	[/PARAM]
+[/OBJECT]
+[/VAR]
+
+## Overview
+
+This Quick Start will teach you the basics of using the **Environment Query System**. In it, you will use Environment Queries to find Actors and Locations, then, using a series of tests, you'll be given the "best" option. Then, using a Behavior tree, you'll have the AI flee from the Player Pawn (or Pawns) to a hiding spot.
+
+This is a more advanced Quick Start and assumes you are familiar with [Blueprints](Engine/Blueprints/BP_HowTo/BasicUsage) and [Behavior Trees](Engine/AI/BehaviorTrees/QuickStart).
+
+[OBJECT:TopicButtonList]
+	[PARAM:icon]
+		![](%ROOT%/build_icon.png)(convert:false)
+	[/PARAM]
+	[PARAM:title]
+		Steps
+	[/PARAM]
+	[PARAM:description]
+	[/PARAM]
+	[PARAM:links]
+	* [](Engine/AI/EnvironmentQuerySystem/QuickStart/1 "%Engine/AI/EnvironmentQuerySystem/QuickStart/1:description%")
+	* [](Engine/AI/EnvironmentQuerySystem/QuickStart/2 "%Engine/AI/EnvironmentQuerySystem/QuickStart/2:description%")
+	* [](Engine/AI/EnvironmentQuerySystem/QuickStart/3 "%Engine/AI/EnvironmentQuerySystem/QuickStart/3:description%")
+	* [](Engine/AI/EnvironmentQuerySystem/QuickStart/4 "%Engine/AI/EnvironmentQuerySystem/QuickStart/4:description%")
+	* [](Engine/AI/EnvironmentQuerySystem/QuickStart/5 "%Engine/AI/EnvironmentQuerySystem/QuickStart/5:description%")
+	* [](Engine/AI/EnvironmentQuerySystem/QuickStart/6 "%Engine/AI/EnvironmentQuerySystem/QuickStart/6:description%")
+	* [](Engine/AI/EnvironmentQuerySystem/QuickStart/7 "%Engine/AI/EnvironmentQuerySystem/QuickStart/7:description%")
+	* [](Engine/AI/EnvironmentQuerySystem/QuickStart/8 "%Engine/AI/EnvironmentQuerySystem/QuickStart/8:description%")
+	* [](Engine/AI/EnvironmentQuerySystem/QuickStart/9 "%Engine/AI/EnvironmentQuerySystem/QuickStart/9:description%")
+	* [](Engine/AI/EnvironmentQuerySystem/QuickStart/10 "%Engine/AI/EnvironmentQuerySystem/QuickStart/10:description%")
+	* [](Engine/AI/EnvironmentQuerySystem/QuickStart/11 "%Engine/AI/EnvironmentQuerySystem/QuickStart/11:description%")
+	* [](Engine/AI/EnvironmentQuerySystem/QuickStart/12 "%Engine/AI/EnvironmentQuerySystem/QuickStart/12:description%")
+	* [](Engine/AI/EnvironmentQuerySystem/QuickStart/13 "%Engine/AI/EnvironmentQuerySystem/QuickStart/13:description%")
+	* [](Engine/AI/EnvironmentQuerySystem/QuickStart/14 "%Engine/AI/EnvironmentQuerySystem/QuickStart/14:description%")
+	[/PARAM]
+[/OBJECT]
+
+
+
