#UE4doc #docs: Reworked sentence fragment.

diff --git a/Engine/Documentation/Source/Shared/Types/ABrush/ABrush.INT.udn b/Engine/Documentation/Source/Shared/Types/ABrush/ABrush.INT.udn
--- a/Engine/Documentation/Source/Shared/Types/ABrush/ABrush.INT.udn
+++ b/Engine/Documentation/Source/Shared/Types/ABrush/ABrush.INT.udn
@@ -5,7 +5,7 @@ Description:Overview of the Brush class
 
 [EXCERPT:ABrush]
 
-A Brush is a 3D volume used to define geometry and gameplay volumes. Also, a user-interface apparatus that you can use to paint values (such as color) onto a surface or a scene.
+A Brush is a 3D volume used to define geometry and gameplay volumes. It is also a user-interface apparatus that you can use to paint values (such as color) onto a surface or a scene.
 
 [VAR:ToolTipFullLink]
 Engine/Actors/Brushes
