#UE4doc #docs: sp: Components, changed "its" to "it is"

diff --git a/Engine/Documentation/Source/Resources/Showcases/BlueprintExamples/FogSheet/FogSheet.INT.udn b/Engine/Documentation/Source/Resources/Showcases/BlueprintExamples/FogSheet/FogSheet.INT.udn
--- a/Engine/Documentation/Source/Resources/Showcases/BlueprintExamples/FogSheet/FogSheet.INT.udn
+++ b/Engine/Documentation/Source/Resources/Showcases/BlueprintExamples/FogSheet/FogSheet.INT.udn
@@ -43,7 +43,7 @@ Each of the 3 aspects of the Blueprint (i.e. the _Components List_, the _Constru
 
 |Blueprint Aspect |Purpose |
 |---|---|
-| **Components List** | The only Component created within the _Compnonents List_ is a single Billboard Component, which provides only a selectable icon. This is used for placement in the level. |
+| **Components List** | The only Component created within the _Components List_ is a single Billboard Component, which provides only a selectable icon. This is used for placement in the level. |
 | **Construction Script** | [INCLUDE:#FogSheetCS] |
 | **Event Graph** | This Blueprint has no _Event Graph_, meaning that once its properties are set, it remains the same throughout gameplay. |
 
@@ -80,7 +80,7 @@ The _Size_ variable is particularly interesting in that it has had its 3D Widget
 ![](3DWidget.png)
 
 
-Since a 3D Widget is simply a mechanism for controlling a Blueprint variable, the widget is only visible when its associated Blueprint is selected. The widget can then be selected and moved using the standard Move tool. This provides a designer with a quick way of changing a vector value.
+Since a 3D Widget is simply a mechanism for controlling a Blueprint variable, the widget is only visible when it is associated Blueprint is selected. The widget can then be selected and moved using the standard Move tool. This provides a designer with a quick way of changing a vector value.
 
 In this case, the 3D Widget is employed to allow for quick resizing of the fog sheet geometry plane. The upper right corner of the mesh meets the location of the 3D Widget, and the sheet as a whole scales from the center, where the Blueprint's sprite icon is located.
 
