removed unneeded comma

diff --git a/Engine/Documentation/Source/Resources/Showcases/BlueprintExamples/Foliage/Foliage.INT.udn b/Engine/Documentation/Source/Resources/Showcases/BlueprintExamples/Foliage/Foliage.INT.udn
--- a/Engine/Documentation/Source/Resources/Showcases/BlueprintExamples/Foliage/Foliage.INT.udn
+++ b/Engine/Documentation/Source/Resources/Showcases/BlueprintExamples/Foliage/Foliage.INT.udn
@@ -19,7 +19,7 @@ The Foliage Blueprint shows how to use a Blueprint to create randomly-generated
 [/EXCERPT:Overview]
 
 
-At a glance, the system takes a set radius, then selects randomly-generated points within that radius to use as potential spawn points. It then performs a downward trace check to see where the surfaces are beneath each point, and uses those as potential locations. The system also checks the individual radius of each mesh to keep meshes from overlapping too heavily, and will choose new random locations whenever there is no room for a mesh at its current one.
+At a glance, the system takes a set radius then selects randomly-generated points within that radius to use as potential spawn points. It then performs a downward trace check to see where the surfaces are beneath each point, and uses those as potential locations. The system also checks the individual radius of each mesh to keep meshes from overlapping too heavily, and will choose new random locations whenever there is no room for a mesh at its current one.
 
 ## Blueprint Setup
 
