added comma

diff --git a/Engine/Documentation/Source/Resources/Showcases/BlueprintExamples/SecurityCameraAndAlarm/SecurityCameraAndAlarm.INT.udn b/Engine/Documentation/Source/Resources/Showcases/BlueprintExamples/SecurityCameraAndAlarm/SecurityCameraAndAlarm.INT.udn
--- a/Engine/Documentation/Source/Resources/Showcases/BlueprintExamples/SecurityCameraAndAlarm/SecurityCameraAndAlarm.INT.udn
+++ b/Engine/Documentation/Source/Resources/Showcases/BlueprintExamples/SecurityCameraAndAlarm/SecurityCameraAndAlarm.INT.udn
@@ -10,7 +10,7 @@ Related: Engine/Blueprints/UserGuide/Components
 
 
 [EXCERPT:Overview]
-Tied into the existing camera Blueprints is a security control room. This room contains a security console surrounded by video screens. As the player presses a large button in the center of the console, the screens switch to show views from various security cameras placed throughout the level. These are the same security cameras located throughout the level, that also trigger the alarm system if the player walks into their view.
+Tied into the existing camera Blueprints is a security control room. This room contains a security console surrounded by video screens. As the player presses a large button in the center of the console, the screens switch to show views from various security cameras placed throughout the level. These are the same security cameras, located throughout the level, that also trigger the alarm system if the player walks into their view.
 [/EXCERPT:Overview]
 
 ## Screen Blueprint
