#UE4doc #docs: remove extra "the"

diff --git a/Engine/Documentation/Source/Resources/Showcases/Effects/FireExamples/FireExamples.INT.udn b/Engine/Documentation/Source/Resources/Showcases/Effects/FireExamples/FireExamples.INT.udn
--- a/Engine/Documentation/Source/Resources/Showcases/Effects/FireExamples/FireExamples.INT.udn
+++ b/Engine/Documentation/Source/Resources/Showcases/Effects/FireExamples/FireExamples.INT.udn
@@ -34,7 +34,7 @@ This [Particle System](Engine/Rendering/ParticleSystems) simulates a realistic f
 
 In the above image, you can see the radii of the particle lights, visible when the _Preview Light Radius_ option is checked within Cascade.
 
-The effect also has a series of spark particles that come off the of the flames and move with convincing turbulent motion. This motion is achieved using a _Local Vector Field_, a feature available to GPU Sprite particles. Vector Fields are 3-dimensional volumes of vectors created from fluid simulation data, which is then used to drive the motion of GPU particles. As the particles pass through the volume, the vectors at each location can influence their velocity. Vector Fields can be visualized in Cascade via the Viewport's View menu (**View > Vector Fields**).
+The effect also has a series of spark particles that come off of the flames and move with convincing turbulent motion. This motion is achieved using a _Local Vector Field_, a feature available to GPU Sprite particles. Vector Fields are 3-dimensional volumes of vectors created from fluid simulation data, which is then used to drive the motion of GPU particles. As the particles pass through the volume, the vectors at each location can influence their velocity. Vector Fields can be visualized in Cascade via the Viewport's View menu (**View > Vector Fields**).
 
 ![](VectorFieldVisible.png)
 
