#UE4doc #docs: duplicate word "of"

diff --git a/Engine/Documentation/Source/Resources/Showcases/Effects/SnowExamples/SnowExamples.INT.udn b/Engine/Documentation/Source/Resources/Showcases/Effects/SnowExamples/SnowExamples.INT.udn
--- a/Engine/Documentation/Source/Resources/Showcases/Effects/SnowExamples/SnowExamples.INT.udn
+++ b/Engine/Documentation/Source/Resources/Showcases/Effects/SnowExamples/SnowExamples.INT.udn
@@ -29,7 +29,7 @@ For more information on Cascade and Particle Systems, please see [the Particle S
 
 **Content Browser Location:** _Game/Effects/Particles/Snow/P_Blizzard_
 
-The Blizzard effect is a [GPU Sprite Particle System](Engine/Rendering/ParticleSystems\Reference\TypeData\GPUSprites) comprised of a massive amount of small sprites using a Lit Translucent Material. The Material used on these snowflakes makes use of of the **Enable Responsive AA** flag, available in the _Translucency_ category of the Material Editor's Properties panel.
+The Blizzard effect is a [GPU Sprite Particle System](Engine/Rendering/ParticleSystems\Reference\TypeData\GPUSprites) comprised of a massive amount of small sprites using a Lit Translucent Material. The Material used on these snowflakes makes use of the **Enable Responsive AA** flag, available in the _Translucency_ category of the Material Editor's Properties panel.
 
 ![](enableResponsiveAA.png)
 
