#UE4doc #docs: changed how reserved words are set up. Ex: Point Light -> PointLight

diff --git a/Engine/Documentation/Source/Shared/Types/APointLight/APointLight.INT.udn b/Engine/Documentation/Source/Shared/Types/APointLight/APointLight.INT.udn
--- a/Engine/Documentation/Source/Shared/Types/APointLight/APointLight.INT.udn
+++ b/Engine/Documentation/Source/Shared/Types/APointLight/APointLight.INT.udn
@@ -5,7 +5,7 @@ Description:Overview of the PointLight class
 
 [EXCERPT:APointLight]
 
-Point Lights work much like a real world light bulb, emitting light in all directions from the light bulb's tungsten filament. However, for the sake of performance, Point Lights are simplified down emitting light equally in all directions from just a single point in space.
+A PointLight works much like a real world light bulb, emitting light in all directions from the light bulb's tungsten filament. However, for the sake of performance, PointLights are simplified down emitting light equally in all directions from just a single point in space.
 
 [VAR:ToolTipFullLink]
 Engine/Rendering/LightingAndShadows/LightTypes/Point
