#UE4doc #docs: Rephrased run-on sentence: The PostProcessVolume is a special type of volume that can be added to a level and, because Unreal Engine 4 does not make use of post processing chains, these volumes are currently the only way to manipulate post processing parameters.

diff --git a/Engine/Documentation/Source/Shared/Types/APostProcessVolume/APostProcessVolume.INT.udn b/Engine/Documentation/Source/Shared/Types/APostProcessVolume/APostProcessVolume.INT.udn
--- a/Engine/Documentation/Source/Shared/Types/APostProcessVolume/APostProcessVolume.INT.udn
+++ b/Engine/Documentation/Source/Shared/Types/APostProcessVolume/APostProcessVolume.INT.udn
@@ -5,7 +5,7 @@ Description:Overview of the PostProcessVolume class
 
 [EXCERPT:APostProcessVolume]
 
-The PostProcessVolume is a special type of volume that can be added to a level and, because Unreal Engine 4 does not make use of post processing chains, these volumes are currently the only way to manipulate post processing parameters.
+The PostProcessVolume is a special type of volume that can be added to a level. Currently Unreal Engine 4 does not make use of post processing chains, making PostProcessVolumes the only way to manipulate post processing parameters.
 
 [VAR:ToolTipFullLink]
 Engine/Rendering/PostProcessEffects
