#UE4doc #docs: sp: reflection + sentence structure

diff --git a/Engine/Documentation/Source/Resources/Showcases/RealisticRendering/RealisticRendering.INT.udn b/Engine/Documentation/Source/Resources/Showcases/RealisticRendering/RealisticRendering.INT.udn
--- a/Engine/Documentation/Source/Resources/Showcases/RealisticRendering/RealisticRendering.INT.udn
+++ b/Engine/Documentation/Source/Resources/Showcases/RealisticRendering/RealisticRendering.INT.udn
@@ -181,11 +181,11 @@ This Material has been highly parameterized, which means it is intended to be in
 
 This demo makes extensive use of reflections to show off certain objects. Reflections are handled in various ways in Unreal Engine 4. 
 
-* **Screen Space Reflections** - Screen Space reflections are an engine feature in UE4 and are visible by default. These are visible primarily at near-edge-on angles. They work in cooperation with the Reflection Environment and Reflection Captures to create a whole sense of scene refelction.
+* **Screen Space Reflections** - Screen Space reflections are an engine feature in UE4 and are visible by default. These are visible primarily at near-edge-on angles. They work in cooperation with the Reflection Environment and Reflection Captures to create a whole sense of scene reflection.
 
 * **Reflection Environment and Reflection Captures** - Unreal Engine 4 makes use of reflection captures to create realistic reflections. These are created by way of strategically placed Actors in the level which create spherical or cube-shaped captures of the environment, which in turn are projected onto any reflective objects within their radius. In this image, you can see the Reflection Capture Actors placed to give accurate reflections to several of the objects.
 
-	![](RefelctionCaptures.png)(w:600)
+	![](ReflectionCaptures.png)(w:600)
 
 
 ## Carpet Object
@@ -211,7 +211,7 @@ The model for the rug object is one of the central aspects of creating the effec
 
 ### Material
 
-The Material for the carpet makes use of Bump Offset multiplied and added into 4 separate stages. These four stages, multiplied by the four stacked pieces of geometry to give a total of 16 "layers" to the effect. The Bump Offset actually runs in a negative direction, giving the fibers a sense of receding down into the surface. This Material is highly parameterized and customizable.
+The Material for the carpet makes use of Bump Offset multiplied and added into 4 separate stages. These four stages, multiplied by the four stacked pieces of geometry, give a total of 16 "layers" to the effect. The Bump Offset actually runs in a negative direction, giving the fibers a sense of receding down into the surface. This Material is highly parameterized and customizable.
 
 ![](CarpetMaterial.png)
 
