#UE4 Docs: Added the own your own section and making the final check in.

diff --git a/Engine/Documentation/Source/Engine/OpenWorldTools/Grass/QuickStart/4/GTQS_4.INT.udn b/Engine/Documentation/Source/Engine/OpenWorldTools/Grass/QuickStart/4/GTQS_4.INT.udn
--- a/Engine/Documentation/Source/Engine/OpenWorldTools/Grass/QuickStart/4/GTQS_4.INT.udn
+++ b/Engine/Documentation/Source/Engine/OpenWorldTools/Grass/QuickStart/4/GTQS_4.INT.udn
@@ -17,6 +17,7 @@ Version: 4.8
 		[How-To Home](Engine\OpenWorldTools\Grass\QuickStart "%Engine\OpenWorldTools\Grass\QuickStart:Title%")
 	[/PARAM]
 	[PARAM:next]
+	[Next Step](Engine\OpenWorldTools\Grass\QuickStart\5 "%Engine\OpenWorldTools\Grass\QuickStart\5:title%") 
 	[/PARAM]
 [/OBJECT]
 [/VAR]
diff --git a/Engine/Documentation/Source/Engine/OpenWorldTools/Grass/QuickStart/5/GTQS_5.INT.udn b/Engine/Documentation/Source/Engine/OpenWorldTools/Grass/QuickStart/5/GTQS_5.INT.udn
--- a/Engine/Documentation/Source/Engine/OpenWorldTools/Grass/QuickStart/5/GTQS_5.INT.udn
+++ b/Engine/Documentation/Source/Engine/OpenWorldTools/Grass/QuickStart/5/GTQS_5.INT.udn
@@ -1,7 +1,49 @@
 Availability:Docs
 Title: 5. On Your Own
 Crumbs:
-Description: What to do next with the Grass Tools.
+Description: What to do next with the Grass Tool.
 Navigation:topic
 SkillLevel:Intermediate
-Version: 4.8
\ No newline at end of file
+Version: 4.8
+
+[VAR:Steps]
+[OBJECT:Navigation]
+	[PARAM:previous]
+	[Previous Step](Engine\OpenWorldTools\Grass\QuickStart\4 "%Engine\OpenWorldTools\ProceduralFoliage\QuickStart\3:title%") 
+	[/PARAM]
+	[PARAM:current]
+	[/PARAM]
+	[PARAM:home]
+		[How-To Home](Engine\OpenWorldTools\Grass\QuickStart "%Engine\OpenWorldTools\Grass\QuickStart:Title%")
+	[/PARAM]
+	[PARAM:next]
+	[/PARAM]
+[/OBJECT]
+[/VAR]
+
+%Steps%
+
+Now that you have seen that power the Grass tool offers, try using the tools listed below in conjunction with what you just learned about the 
+Grass tool to try and make a level that looks like the following image.
+
+![](Engine\OpenWorldTools\ProceduralFoliage\QuickStart\6\T_On_Your_Own.png)
+
+* Use the **[Procedural Foliage Tool](Engine\OpenWorldTools\ProceduralFoliage\QuickStart)** to make the Landscape look like it is densely covered in grass, flowers and bushes.
+
+* Define the look and feel of the Landscape terrain using the **[Landscape Sculpt](Engine/Landscape/Editing/SculptMode)** tools add things like hills, mountains and lakes.
+
+* Give the surface of the Landscape more visual variety and detail by creating a **[Landscape Material](Engine/Landscape/Materials)** that has multiple Textures that can be painted on the Landscape terrain.
+
+* Adjust the **[Directional Light](Engine/Rendering/LightingAndShadows/LightTypes/Directional)** to make the level lighting resemble lighting that happens in the early morning or late afternoon.
+
+* Set up the level lighting to use a completely dynamic based lighting solution that will make use of dynamic shadows as well as **[Ray Traced Distance Field Soft Shadows](Engine/Rendering/LightingAndShadows/RayTracedDistanceFieldShadowing)**. 
+
+* Try using the **[Foliage System](Engine/Foliage)** tools to remove or tweak the placement, rotation and scale of the Foliage meshes that are placed by the Procedural Foliage tool so that you can get the exact look you are going for.
+
+* Show off your creation to others by using a **[Camera](Gameplay/HowTo/UsingCameras)** in conjunction with **[Matinee](Engine/Matinee)** to render out a video of your level for you to share with the world.
+
+
+
+
+
+%Steps%
\ No newline at end of file
diff --git a/Engine/Documentation/Source/Engine/OpenWorldTools/Overview/OpenWorldToolsOverview.INT.udn b/Engine/Documentation/Source/Engine/OpenWorldTools/Overview/OpenWorldToolsOverview.INT.udn
--- a/Engine/Documentation/Source/Engine/OpenWorldTools/Overview/OpenWorldToolsOverview.INT.udn
+++ b/Engine/Documentation/Source/Engine/OpenWorldTools/Overview/OpenWorldToolsOverview.INT.udn
@@ -26,8 +26,7 @@ Once the Open World Demo Collection is downloaded add it to the project that you
 
 ## Types of Tools 
 The Open World Tools in Unreal Engine 4 allow you to place Static Meshes inside your levels using a computer algorithm to determine placement instead having to do this work manually. 
-This enables you to quickly generate a natural looking forest comprised of many different species of trees.
-You can also control which plants are spawned by painting Textures on the Landscape Terrain.
+This enables you to quickly generate a natural looking forest comprised of many different species of trees, plants and bushes.
 
 When talking about the Open World tools we are actually referring to the two different tool sets inside of UE4. 
 The tool set that you use depends on what you want to accomplish. 
@@ -53,7 +52,7 @@ To learn more about how to use the Procedural Foliage Tool check out the Procedu
 		Procedural Foliage Tool Quick Start
 	[/PARAM]
 	[PARAM:description]
-		A comprehensive guide for getting the Procedural Foliage Tool setup and working.
+		A comprehensive guide showing you how to unleash the full power of the Procedural Foliage tool in your projects.
 	[/PARAM]
 	[PARAM:path]
 		[RELATIVE:Engine/OpenWorldTools/ProceduralFoliage/QuickStart]
@@ -89,7 +88,7 @@ To learn more about how to use the Grass Tool, check out the Grass Tool Quick St
 [REGION:topics third]
 [OBJECT:Topic]
 	[PARAM:image]
-		![](T_PFT_Grass_00.png)
+		![Engine/OpenWorldTools/Overview:title%](Engine\OpenWorldTools\Grass\QuickStart\T_Grass_Intro.png)
 	[/PARAM]
 	[PARAM:icon]
 		![](%ROOT%/editor_icon.png)(convert:false)
@@ -98,7 +97,7 @@ To learn more about how to use the Grass Tool, check out the Grass Tool Quick St
 		Grass Tool Quick Start
 	[/PARAM]
 	[PARAM:description]
-		A comprehensive guide for getting the Grass Tool setup and working.
+	An in-depth look at how to use the Grass tool in your projects.
 	[/PARAM]
 	[PARAM:path]
 		[RELATIVE:Engine/OpenWorldTools/Grass/QuickStart]
diff --git a/Engine/Documentation/Source/Engine/OpenWorldTools/ProceduralFoliage/QuickStart/5/PFQS_5.INT.udn b/Engine/Documentation/Source/Engine/OpenWorldTools/ProceduralFoliage/QuickStart/5/PFQS_5.INT.udn
--- a/Engine/Documentation/Source/Engine/OpenWorldTools/ProceduralFoliage/QuickStart/5/PFQS_5.INT.udn
+++ b/Engine/Documentation/Source/Engine/OpenWorldTools/ProceduralFoliage/QuickStart/5/PFQS_5.INT.udn
@@ -17,7 +17,7 @@ Version: 4.8
 		[How-To Home](Engine\OpenWorldTools\ProceduralFoliage\QuickStart "Engine\OpenWorldTools\ProceduralFoliage\QuickStart:Title%")
 	[/PARAM]
 	[PARAM:next]
-
+	[Next Step](Engine\OpenWorldTools\ProceduralFoliage\QuickStart\6 "%Engine\OpenWorldTools\ProceduralFoliage\QuickStart\6:title%") 
 	[/PARAM]
 [/OBJECT]
 [/VAR]
diff --git a/Engine/Documentation/Source/Engine/OpenWorldTools/ProceduralFoliage/QuickStart/6/PFQS_6.INT.udn b/Engine/Documentation/Source/Engine/OpenWorldTools/ProceduralFoliage/QuickStart/6/PFQS_6.INT.udn
--- a/Engine/Documentation/Source/Engine/OpenWorldTools/ProceduralFoliage/QuickStart/6/PFQS_6.INT.udn
+++ b/Engine/Documentation/Source/Engine/OpenWorldTools/ProceduralFoliage/QuickStart/6/PFQS_6.INT.udn
@@ -1,18 +1,55 @@
 Availability:Docs
-Title: 6. Own your Own! 
+Title: 6. Own your Own
 Crumbs:
-Description:What to do next
+Description: What to do next with the Procedural Foliage tool.
 Navigation:topic
 SkillLevel:Intermediate
 Version: 4.8
 
+[VAR:Steps]
+[OBJECT:Navigation]
+	[PARAM:previous]
+	[Previous Step](Engine\OpenWorldTools\ProceduralFoliage\QuickStart\5 "%Engine\OpenWorldTools\ProceduralFoliage\QuickStart\5:title%") 
+	[/PARAM]
+	[PARAM:current]
+	[/PARAM]
+	[PARAM:home]
+		[How-To Home](Engine\OpenWorldTools\ProceduralFoliage\QuickStart "Engine\OpenWorldTools\ProceduralFoliage\QuickStart:Title%")
+	[/PARAM]
+	[PARAM:next]
 
+	[/PARAM]
+[/OBJECT]
+[/VAR]
 
+%Steps%
 
+Now that you have seen that power the Procedural Foliage tool offers, try using the tools listed below in conjunction with what you just learned about the 
+Procedural Foliage tool to try and make a level that looks like the following image.
 
+![](T_On_Your_Own.png)
 
+* Use the **[Grass Tool](Engine\OpenWorldTools\Grass\QuickStart)** to make the Landscape look like it is densely covered in grass, flowers and bushes.
 
+* Define the look and feel of the Landscape terrain using the **[Landscape Sculpt](Engine/Landscape/Editing/SculptMode)** tools add things like hills, mountains and lakes.
 
+* Give the surface of the Landscape more visual variety and detail by creating a **[Landscape Material](Engine/Landscape/Materials)** that has multiple Textures that can be painted on the Landscape terrain.
 
+* Adjust the **[Directional Light](Engine/Rendering/LightingAndShadows/LightTypes/Directional)** to make the level lighting resemble lighting that happens in the early morning or late afternoon.
 
+* Set up the level lighting to use a completely dynamic based lighting solution that will make use of dynamic shadows as well as **[Ray Traced Distance Field Soft Shadows](Engine/Rendering/LightingAndShadows/RayTracedDistanceFieldShadowing)**. 
 
+* Try using the **[Foliage System](Engine/Foliage)** tools to remove or tweak the placement, rotation and scale of the Foliage meshes that are placed by the Procedural Foliage tool so that you can get the exact look you are going for.
+
+* Show off your creation to others by using a **[Camera](Gameplay/HowTo/UsingCameras)** in conjunction with **[Matinee](Engine/Matinee)** to render out a video of your level for you to share with the world.
+
+
+
+
+
+
+
+
+
+
+%Steps%
\ No newline at end of file
diff --git a/Engine/Documentation/Source/Engine/OpenWorldTools/ProceduralFoliage/QuickStart/ProceduralFoliageQuickStart.INT.udn b/Engine/Documentation/Source/Engine/OpenWorldTools/ProceduralFoliage/QuickStart/ProceduralFoliageQuickStart.INT.udn
--- a/Engine/Documentation/Source/Engine/OpenWorldTools/ProceduralFoliage/QuickStart/ProceduralFoliageQuickStart.INT.udn
+++ b/Engine/Documentation/Source/Engine/OpenWorldTools/ProceduralFoliage/QuickStart/ProceduralFoliageQuickStart.INT.udn
@@ -8,8 +8,8 @@ Version: 4.8
 
 
 In this Quick Start tutorial we will take a look at how the **Procedural Foliage Tool** that are now available with the 4.8 release of Unreal Engine 4 (UE4) work. 
-Over the course of this Quick Start tutorial you will learn how to create, setup, and spawn an entire forests worth of trees inside of UE4 using the Procedural Foliage tool. 
-You will also be introduced to key properties and settings that will help you shape your virtual forest to fit your needs. 
+Over the course of this Quick Start tutorial you will learn how to create, setup, and spawn an entire forests worth of trees inside of UE4 using only the Procedural Foliage tool. 
+You will also be introduced to key properties and settings that will help you shape your virtual forest to fit your projects needs. 
 
 ![](T_PF_QS_Final_Product.png)
 
