removed unnecessary commas

diff --git a/Engine/Documentation/Source/Resources/Templates/FirstPerson/FirstPersonTemplate.INT.udn b/Engine/Documentation/Source/Resources/Templates/FirstPerson/FirstPersonTemplate.INT.udn
--- a/Engine/Documentation/Source/Resources/Templates/FirstPerson/FirstPersonTemplate.INT.udn
+++ b/Engine/Documentation/Source/Resources/Templates/FirstPerson/FirstPersonTemplate.INT.udn
@@ -48,7 +48,7 @@ resources
 So you have started a new game using the **First Person Template**, awesome! If you are looking for guidance on how to do specific things or perhaps just looking at the blank canvas of endless possibilities, this doc contains ideas for things you can implement.
 
 * [Read up](Engine/Blueprints) and [watch some videos](Videos) on the **Blueprints Visual Scripting** system. Sometimes it is not about what to do, but how you can do it, and knowing more about Blueprints will enable you to rapidly prototype ideas, even if you will ultimately be working in C++.
-* Create a [particle effect](Engine/Rendering/ParticleSystems), using the [Cascade Particle Editor](Engine/Rendering/ParticleSystems/Cascade), that represents a muzzle flash! Then using a [Particle System Component](Engine/Components/Rendering#particlesystemcomponent), add it to the **FirstPersonCharacter** Blueprint, and Activate the component in the **InputAction Fire** Event.
+* Create a [particle effect](Engine/Rendering/ParticleSystems) using the [Cascade Particle Editor](Engine/Rendering/ParticleSystems/Cascade) that represents a muzzle flash! Then using a [Particle System Component](Engine/Components/Rendering#particlesystemcomponent), add it to the **FirstPersonCharacter** Blueprint and Activate the component in the **InputAction Fire** Event.
 	* **Bonus**: Add a Particle System Component to the **FirstPersonProjectile** to make it look like it trails smoke behind it. Look for the **Use Local Space** property [here](Engine/Rendering/ParticleSystems/Reference/Modules/Required).
 * [Create a **HUD**](Engine/UMG/QuickStart) for your First Person Game using [**Unreal Motion Graphics**](Engine/UMG).
 * Why not add some Destructible Meshes to the environment? Check out the [Destructible Actor Documentation](Engine/Physics/Destructibles) and the [](Engine/Blueprints/HowTo) for methods on how to do this!
