#UE4doc #docs: made spacing changes for reserved words and re-worked the statement: Much like the Scene Capture Cube, a Scene Capture 2D makes an image of the scene on every frame. The Scene Capture 2D captures the scene from its view frustum, stores that view as an image, which is then used within a Material.

diff --git a/Engine/Documentation/Source/Shared/Types/ASceneCapture2D/ASceneCapture2D.INT.udn b/Engine/Documentation/Source/Shared/Types/ASceneCapture2D/ASceneCapture2D.INT.udn
--- a/Engine/Documentation/Source/Shared/Types/ASceneCapture2D/ASceneCapture2D.INT.udn
+++ b/Engine/Documentation/Source/Shared/Types/ASceneCapture2D/ASceneCapture2D.INT.udn
@@ -5,7 +5,7 @@ Description:Overview of the SceneCapture2D class
 
 [EXCERPT:ASceneCapture2D]
 
-Much like the Scene Capture Cube, a Scene Capture 2D makes an image of the scene on every frame.  The Scene Capture 2D captures the scene from its view frustum, stores that view as an image, which is then used within a Material.
+Much like the SceneCaptureCube, a SceneCapture2D makes an image of the scene on every frame. The SceneCapture2D captures the scene from its view frustum. It then stores that view as an image which is used within a Material. 
 
 [VAR:ToolTipFullLink]
 Resources/ContentExamples/Reflections/1_7
