#UE4doc #docs: re-worked sentence from: Triggers are Actors that are used to cause an event to occur when they are interacted with by some other object in the level.

diff --git a/Engine/Documentation/Source/Shared/Types/ATriggerBox/ATriggerBox.INT.udn b/Engine/Documentation/Source/Shared/Types/ATriggerBox/ATriggerBox.INT.udn
--- a/Engine/Documentation/Source/Shared/Types/ATriggerBox/ATriggerBox.INT.udn
+++ b/Engine/Documentation/Source/Shared/Types/ATriggerBox/ATriggerBox.INT.udn
@@ -5,7 +5,7 @@ Description:Overview of the TriggerBox class
 
 [EXCERPT:ATriggerBox]
 
-Triggers are Actors that are used to cause an event to occur when they are interacted with by some other object in the level. 
+Triggers are Actors that cause an event to occur when they interact with another object in the level.
 
 [VAR:ToolTipFullLink]
 Engine/Actors/Triggers
