#UE4doc #docs: all edits made to match ATriggerBox. No distinctions made between different types of triggers.

diff --git a/Engine/Documentation/Source/Shared/Types/ATriggerCapsule/ATriggerCapsule.INT.udn b/Engine/Documentation/Source/Shared/Types/ATriggerCapsule/ATriggerCapsule.INT.udn
--- a/Engine/Documentation/Source/Shared/Types/ATriggerCapsule/ATriggerCapsule.INT.udn
+++ b/Engine/Documentation/Source/Shared/Types/ATriggerCapsule/ATriggerCapsule.INT.udn
@@ -5,7 +5,7 @@ Description:Overview of the TriggerCapsule class
 
 [EXCERPT:ATriggerCapsule]
 
-Triggers are Actors that are used to cause an event to occur when they are interacted with by some other object in the level. 
+Triggers are Actors that cause an event to occur when they interact with another object in the level.
 
 [VAR:ToolTipFullLink]
 Engine/Actors/Triggers
diff --git a/Engine/Documentation/Source/Shared/Types/ATriggerSphere/ATriggerSphere.INT.udn b/Engine/Documentation/Source/Shared/Types/ATriggerSphere/ATriggerSphere.INT.udn
--- a/Engine/Documentation/Source/Shared/Types/ATriggerSphere/ATriggerSphere.INT.udn
+++ b/Engine/Documentation/Source/Shared/Types/ATriggerSphere/ATriggerSphere.INT.udn
@@ -5,7 +5,7 @@ Description:Overview of the TriggerSphere class
 
 [EXCERPT:ATriggerSphere]
 
-Triggers are Actors that are used to cause an event to occur when they are interacted with by some other object in the level. 
+Triggers are Actors that cause an event to occur when they interact with another object in the level.
 
 [VAR:ToolTipFullLink]
 Engine/Actors/Triggers
diff --git a/Engine/Documentation/Source/Shared/Types/ATriggerVolume/ATriggerVolume.INT.udn b/Engine/Documentation/Source/Shared/Types/ATriggerVolume/ATriggerVolume.INT.udn
--- a/Engine/Documentation/Source/Shared/Types/ATriggerVolume/ATriggerVolume.INT.udn
+++ b/Engine/Documentation/Source/Shared/Types/ATriggerVolume/ATriggerVolume.INT.udn
@@ -5,7 +5,7 @@ Description:Overview of the TriggerVolume class
 
 [EXCERPT:ATriggerVolume]
 
-Trigger Volumes are used as a way to cause events when a Player or other object enters or exits them.
+Triggers are Actors that cause an event to occur when they interact with another object in the level.
 
 [VAR:ToolTipFullLink]
 Engine/Actors/Volumes#Volume Types
