#UE4doc #docs: sp: respond

diff --git a/Engine/Documentation/Source/Shared/Collision/CollisionTooltips.INT.udn b/Engine/Documentation/Source/Shared/Collision/CollisionTooltips.INT.udn
--- a/Engine/Documentation/Source/Shared/Collision/CollisionTooltips.INT.udn
+++ b/Engine/Documentation/Source/Shared/Collision/CollisionTooltips.INT.udn
@@ -45,7 +45,7 @@ Object Channel is the one that can be assigned to object as Object Type and can
 Trace Channel is when you have a specific purpose of trace, for example, Weapon Trace. Even though it is possible to create bullet object that has object type as "Weapon", often most of games do direct ray cast to the target. When you trace using "Weapon" Trace type, each object can decide how you would like to respond to it. 
 
 Another example is camera. You can decide to have camera object to contain their Object Type to be Camera, in which case Camera is Object Channel. 
-But also you can choose to use Camera as Trace Channel, and use it in your trace code to trace using that channel, and other objects can decide how they would like to repond. 
+But also you can choose to use Camera as Trace Channel, and use it in your trace code to trace using that channel, and other objects can decide how they would like to respond. 
 
 The distinction is shown in below when using same functionality. 
 
