#UE4doc #docs: sp: Split sentences into two

diff --git a/Engine/Documentation/Source/Resources/Templates/ThirdPerson/ThirdPersonTemplate.INT.udn b/Engine/Documentation/Source/Resources/Templates/ThirdPerson/ThirdPersonTemplate.INT.udn
--- a/Engine/Documentation/Source/Resources/Templates/ThirdPerson/ThirdPersonTemplate.INT.udn
+++ b/Engine/Documentation/Source/Resources/Templates/ThirdPerson/ThirdPersonTemplate.INT.udn
@@ -46,7 +46,7 @@ resources
 
 ![](templateMore.png)(convert:false)
 
-So you have started a new game using the **Third Person Template**, awesome! If you are looking for guidance on how to do specific things or perhaps just looking at the blank canvas of endless possibilities, this doc contains ideas for things you can implement.
+So you have started a new game using the **Third Person Template**? Awesome! If you are looking for guidance on how to do specific things or perhaps just looking at the blank canvas of endless possibilities, this doc contains ideas for things you can implement.
 
 * [Read up](Engine/Blueprints) and [watch some video](Videos) on the **Blueprints Visual Scripting** system. Sometimes it is not about what to do, but how you can do it, and knowing more about Blueprints will enable you to rapidly prototype ideas, even if you will ultimately be working in C++.
 * [Create a **HUD**](Engine/UMG/QuickStart) for your Third Person Game using [**Unreal Motion Graphics**](Engine/UMG).
