#UE4doc #docs: "messages that indicate" -removed "s" sp: vertex

diff --git a/Engine/Documentation/Source/Shared/Editor/FbxErrors/FbxErrors.INT.udn b/Engine/Documentation/Source/Shared/Editor/FbxErrors/FbxErrors.INT.udn
--- a/Engine/Documentation/Source/Shared/Editor/FbxErrors/FbxErrors.INT.udn
+++ b/Engine/Documentation/Source/Shared/Editor/FbxErrors/FbxErrors.INT.udn
@@ -16,7 +16,7 @@ This is a list of error or warning messages you may encounter during importing F
 (#Generic_ImportingNewObjectFailed)
 [EXCERPT:Generic_ImportingNewObjectFailed]
 
-Importing failed to create new asset. In this case, you are likely to have other warning messages that indicates the cause of the failures. 
+Importing failed to create new asset. In this case, you are likely to have other warning messages that indicate the cause of the failures. 
 
 [/EXCERPT:Generic_ImportingNewObjectFailed]
 
@@ -435,7 +435,7 @@ New mesh for LOD bone hierarchy should match with the original bone hierarchy. T
 
 (#SkeletalMesh_LOD_HasSoftVerts)
 [EXCERPT:SkeletalMesh_LOD_HasSoftVerts]
-The mesh LOD you are importing has some vertices with more than one influence. If you would like to have soft vertice deformation, please set CheckSingleInfluenceLOD to False in the Editor INI file. 
+The mesh LOD you are importing has some vertices with more than one influence. If you would like to have soft vertex deformation, please set CheckSingleInfluenceLOD to False in the Editor INI file. 
 
 [/EXCERPT:SkeletalMesh_LOD_HasSoftVerts]
 
