#UE4doc #docs: made many grammatical changes. Ex: it's -> it is

diff --git a/Engine/Documentation/Source/Shared/Types/FBodyInstance/FBodyInstance.INT.udn b/Engine/Documentation/Source/Shared/Types/FBodyInstance/FBodyInstance.INT.udn
--- a/Engine/Documentation/Source/Shared/Types/FBodyInstance/FBodyInstance.INT.udn
+++ b/Engine/Documentation/Source/Shared/Types/FBodyInstance/FBodyInstance.INT.udn
@@ -4,7 +4,7 @@ Crumbs:
 
 [EXCERPT:bStartAwake]
 
-Whether or not this Physics Body should start awake. An awake Physics Body is actively simulating (not asleep). If an Awake Physics Body comes to a stop for a certain amount of time, it will go back to Sleep.
+Decide whether or not this Physics Body should start awake. An awake Physics Body is actively simulating (not asleep). If an Awake Physics Body comes to a stop for a certain amount of time, it will go back to Sleep.
 [VAR:ToolTipFullLink]
 Engine/Physics/PhysicsBodies#physics
 [/VAR]
@@ -56,7 +56,7 @@ Engine/Physics/PhysicsBodies#physics
 
 Adjusts the Center of Mass of a Physics Body.
 
-The primary use of this is for vehicles or large Physics Bodies that encompass an irregular shape. You will often use one large Physics Body to define the majority of the Mass of a Vehicle. As a Center of Mass is generated for the body based on the Physics Body being used, it will almost always be dead center of the Physics Body. This can make a vehicle handle oddly, thusly you can adjust the Center of Mass to account for where the mass of the vehicle really is located.
+The primary use of this is for vehicles or large Physics Bodies that encompass an irregular shape. You will often use one large Physics Body to define the majority of the Mass of a Vehicle. As a Center of Mass is generated for the body based on the Physics Body being used, it will almost always be dead center of the Physics Body. This can make a vehicle handle oddly, thus you can adjust the Center of Mass to account for where the mass of the vehicle really is located.
 
 ![](PhAT_COM.png)
 
@@ -88,9 +88,9 @@ Engine/Physics/PhysicsBodies#physics
 
 This will scale the Mass of the selected Physics Body.
 
-The uses for this is to make it appear as if the Skeletal Mesh is made from different materials or has varying densities. This is necessary because the calculation that defines the Mass of the Physics Body is uniform across all of the Physics Bodies.
+The use for this is to make it appear as if the Skeletal Mesh is made from different materials or has varying densities. This is necessary because the calculation that defines the Mass of the Physics Body is uniform across all of the Physics Bodies.
 
-This can make things more or less affected by the forces in the level. A lowered mass could be used to give something a more light an airy feel. While an increased mass can root a Physics Asset in place.
+This can make things more or less affected by the forces in the level. A lowered mass could be used to give something a more light an airy feel where an increased mass can root a Physics Asset in place.
 [VAR:ToolTipFullLink]
 Engine/Physics/PhysicsBodies#physics
 [/VAR]
@@ -110,7 +110,7 @@ Engine/Physics/PhAT/Properties
 
 [EXCERPT:LinearDamping]
 
-This adds a drag like effect to the Physics Body, but more accurately it defines how much the Physics Body resists moving. For a falling Physics Body this will slow its decent, potentially making it appear as if it's falling through water or sludge.
+This adds a drag like effect to the Physics Body, but more accurately it defines how much the Physics Body resists moving. For a falling Physics Body this will slow the decent, potentially making it appear as if it is falling through water or sludge.
 [VAR:ToolTipFullLink]
 Engine/Physics/PhAT/Properties
 [/VAR]
@@ -145,7 +145,7 @@ Engine/Physics/PhysicsBodies#physics
 
 [EXCERPT:bNotifyRigidBodyCollision]
 
-Enables the selected Physics Body to report that it has had a collision to Blueprints or Code elements. If this is off, the Physics Body will collide with things, but it will not actually report that it has done so.
+It enables the selected Physics Body to report that it has had a collision to Blueprints or Code elements. If this is off, the Physics Body will collide with things, but it will not actually report that it has done so.
 Engine/Physics/PhysicsBodies/#collision
 [/EXCERPT:bNotifyRigidBodyCollision]
 
