#UE4doc #docs: "a" to "an"

diff --git a/Engine/Documentation/Source/Shared/Editor/MapErrors/MapErrors.INT.udn b/Engine/Documentation/Source/Shared/Editor/MapErrors/MapErrors.INT.udn
--- a/Engine/Documentation/Source/Shared/Editor/MapErrors/MapErrors.INT.udn
+++ b/Engine/Documentation/Source/Shared/Editor/MapErrors/MapErrors.INT.udn
@@ -64,7 +64,7 @@ This warning is caused when there is a brush in the level that has non-coplanar
 [EXCERPT:SameLocation]
 **[Actor] in same location as [Another Actor]**
 
-This warning is caused when there is a Actor that is in the same exact location as another Actor.  This warning is usually the result of a accidental duplication or paste operation.  It can be fixed by deleting one of the Actors, or disregarded if the placement was intentional.
+This warning is caused when there is an Actor that is in the same exact location as another Actor.  This warning is usually the result of an accidental duplication or paste operation.  It can be fixed by deleting one of the Actors, or disregarded if the placement was intentional.
 [/EXCERPT:SameLocation]
 
 (#InvalidDrawscale)
@@ -78,7 +78,7 @@ This warning is caused when either DrawScale, DrawScale3D X, DrawScale3D Y, or D
 [EXCERPT:ActorIsObselete]
 **[Actor] is obsolete and must be removed!**
 
-This warning is caused when there is an instance of a Actor in a level that has been marked deprecated.  This is usually because a Actor was marked deprecated after the level was created, but the map was never updated.  This can be fixed by simply deleting the Actor.
+This warning is caused when there is an instance of an Actor in a level that has been marked deprecated.  This is usually because an Actor was marked deprecated after the level was created, but the map was never updated.  This can be fixed by simply deleting the Actor.
 [/EXCERPT:ActorIsObselete]
 
 (#StaticPhysNone)
@@ -277,7 +277,7 @@ This warning is caused when lighting has been invalidated by moving or modifying
 [EXCERPT:ParticleSystemComponentNull]
 **[Actor] &#58; Emitter Actor has NULL ParticleSystemComponent property - please delete!**
 
-This warning is caused when there is a Emitter Actor with a "None" ParticleSystemComponent.  It is usually found in older maps where duplication was used to create the Actor and can be fixed by deleting the Actor causing the warning and creating a new one.
+This warning is caused when there is an Emitter Actor with a "None" ParticleSystemComponent.  It is usually found in older maps where duplication was used to create the Actor and can be fixed by deleting the Actor causing the warning and creating a new one.
 [/EXCERPT:ParticleSystemComponentNull]
 
 (#PSysCompErrorEmptyActorRef)
@@ -325,7 +325,7 @@ The specified SkeletalMeshActor has a NULL SkeletalMesh.
 [EXCERPT:AudioComponentNull]
 **[Actor] &#58; Ambient sound Actor has NULL AudioComponent property - please delete!**
 
-This warning is caused when there is a Ambient sound Actor with a "None" AudioComponent.  It is usually found in older maps where duplication was used to create a Actor and can be fixed by deleting the Actor causing the warning and creating a new one.
+This warning is caused when there is an Ambient sound Actor with a "None" AudioComponent.  It is usually found in older maps where duplication was used to create an Actor and can be fixed by deleting the Actor causing the warning and creating a new one.
 [/EXCERPT:AudioComponentNull]
 
 (#SoundCueNull)
