#UE4doc #docs: sp: available

diff --git a/Engine/Documentation/Source/Shared/Editors/BlueprintEditor/BlueprintEditorTooltips.INT.udn b/Engine/Documentation/Source/Shared/Editors/BlueprintEditor/BlueprintEditorTooltips.INT.udn
--- a/Engine/Documentation/Source/Shared/Editors/BlueprintEditor/BlueprintEditorTooltips.INT.udn
+++ b/Engine/Documentation/Source/Shared/Editors/BlueprintEditor/BlueprintEditorTooltips.INT.udn
@@ -14,7 +14,7 @@ If you opened the menu by context-clicking on the background, the context will b
 
 
 [EXCERPT:BlueprintActionMenuContextToggle]
-If checked, the list of avaiable actions and nodes will be filtered to only those that are possible in the current 'context' (actions possible given how you opened the menu).
+If checked, the list of available actions and nodes will be filtered to only those that are possible in the current 'context' (actions possible given how you opened the menu).
 
 If you opened the menu by dragging off of a wire, then it will be filtered to variables and functions that operate on that type.
 
