#UE4doc #docs: sp: separate

diff --git a/Engine/Documentation/Source/Shared/Editors/Common/EditorMenuItems/EditorMenuItems.INT.udn b/Engine/Documentation/Source/Shared/Editors/Common/EditorMenuItems/EditorMenuItems.INT.udn
--- a/Engine/Documentation/Source/Shared/Editors/Common/EditorMenuItems/EditorMenuItems.INT.udn
+++ b/Engine/Documentation/Source/Shared/Editors/Common/EditorMenuItems/EditorMenuItems.INT.udn
@@ -70,16 +70,16 @@ Crumbs:
 * **Developer Tools** - Brings up the selected **Developer Tool** in a separate window (Camera Debugger, Collision Analyzer, Debug Tools, Modules, Widget Reflector).
 	* **Blueprint Debugger** - Brings up the **Blueprint Debugger** in a separate window.
 	* **Collision Analyzer** - Brings up the **Collision Analyzer** in a separate window.
-	* **Debug Tool** - Brings up the **Debug Tool** in a seperate window. 
+	* **Debug Tool** - Brings up the **Debug Tool** in a separate window. 
 	* **Message Log** - Brings up the **Message Log** in a separate window.
 	* **Output Log** - Brings up the **Output Log** in a separate window.
-	* **Visual Logger** - Bring up the **Visual Logger** tool in a seperate window. 
+	* **Visual Logger** - Bring up the **Visual Logger** tool in a separate window. 
 	* **Class Viewer** - Brings up the **Class Viewer** in a separate window.
 	* **Device Manager** - Brings up **Device Manager** in a separate window.
 	* **Device Profiles** - Brings up the **Device Profiles** in a separate window.
-	* **Modules** - Brings up **Modules** in a seperate window. 
+	* **Modules** - Brings up **Modules** in a separate window. 
 	* **Session Frontend** - Brings up **Unreal Frontend** in a separate window.
-	* **Widget Reflector** - Brings up **Widget Reflector** in a seperate window. 
+	* **Widget Reflector** - Brings up **Widget Reflector** in a separate window. 
 * **Undo History** - Brings up the **Undo History** panel in a separate window.
 * **Plugins** - Brings up the **Plugins** window.
 * **Reset Layout** - Resets your layout to a default arrangement. Requires the editor to restart after saving your changes and creating a backup of your settings. 
