#UE4doc #docs: took out repeat words: "can can", restructured setences, and made grammatical changes: ","

diff --git a/Engine/Documentation/Source/Shared/Types/FConstraintInstance/FConstraintInstance.INT.udn b/Engine/Documentation/Source/Shared/Types/FConstraintInstance/FConstraintInstance.INT.udn
--- a/Engine/Documentation/Source/Shared/Types/FConstraintInstance/FConstraintInstance.INT.udn
+++ b/Engine/Documentation/Source/Shared/Types/FConstraintInstance/FConstraintInstance.INT.udn
@@ -4,7 +4,7 @@ Crumbs:
 
 [EXCERPT:JointName]
 
-While working in the Physics Asset Tool, Physics Constraints are associated with a skinned Joint, and this field displays the name of that Joint.
+While working in the Physics Asset Tool, PhysicsConstraints are associated with a skinned Joint, and this field displays the name of that Joint.
 [/EXCERPT:JointName]
 
 
@@ -118,7 +118,7 @@ Engine/Physics/Constraints/ConstraintsReference#sharedproperties
 
 [EXCERPT:LinearLimitDamping]
 
-This defines how much velocity is removed from the motion of the constraint while it is outside the Soft Limit. High values can can cause the Physics Constraint to stop outside of its limit because it lacks the force to pull it back in and the Physics Body will enter a sleep state.
+This defines how much velocity is removed from the motion of the constraint while it is outside the Soft Limit. High values can cause the Physics Constraint to stop outside of its limit because it lacks the force to pull it back in and the Physics Body will enter a sleep state.
 [VAR:ToolTipFullLink]
 Engine/Physics/Constraints/ConstraintsReference#sharedproperties
 [/VAR]
@@ -133,9 +133,9 @@ Whether or not a linear force can be applied to break the constraint. See Linear
 
 [EXCERPT:LinearBreakThreshold]
 
-This type of break is like a shearing or tensile break, where the constraint in question does not bend, but simply exceeds its structural limits by being being torn like paper or a rubber band snapping after being stretched too far.
+This type of break is like a shearing or tensile break. The constraint in question does not bend but simply exceeds its structural limits by being torn like paper, or a rubber band snapping after being stretched too far.
 
-To calculate what value should be used to simply have a Physics Constraint resist the tensile force applied by gravity, take the mass of your Physics Body, multiply by gravity times 100. So a Physics Constraint holding up a mass of 242kg in standard gravity (9.8m/s^2) would need at least a value of 237160 to resist breaking assuming no starting velocity.
+To calculate what value should be used to have a Physics Constraint resist the tensile force applied by gravity, take the mass of your Physics Body, multiply by gravity times 100. For example, a Physics Constraint holding up a mass of 242kg in standard gravity (9.8m/s^2) would need at least a value of 237160 to resist breaking assuming no starting velocity.
 
 Note: An axis needs to be set to LCM_Limited or LCM_Locked to have anything to break away from, LCM_Free axes cannot break.
 
@@ -267,7 +267,7 @@ Note: An axis needs to be set to LCM_Limited or LCM_Locked to have anything to b
 
 [EXCERPT:AngularBreakThreshold]
 
-The Angular force required to to break the Physics Constraint. To keep a Physics Constraint from breaking, this value will need to be very high. 500 million or higher values are not unheard of.
+The Angular force required to break the Physics Constraint. To keep a Physics Constraint from breaking, this value will need to be very high. 500 million or higher values are not unheard of.
 
 Note: An axis needs to be set to ACM_Limited or ACM_Locked to have anything to break away from, ACM_Free axes cannot break.
 [/EXCERPT:AngularBreakThreshold]
@@ -327,7 +327,7 @@ Engine/Physics/Constraints/ConstraintsReference#sharedproperties
 
 [EXCERPT:LinearDriveSpring]
 
-Sets the basic tension of the Linear Position Drive, or how quickly it should attempt to reach its goal. A high value along with a high Linear Drive Force Limit, will cause the Physics Bodies to quickly achieve their goals.
+Sets the basic tension of the Linear Position Drive, or how quickly it should attempt to reach its goal. A high value along with a high Linear Drive Force Limit will cause the Physics Bodies to quickly achieve their goals.
 
 Has no effect on the Linear Velocity Drive.
 [VAR:ToolTipFullLink]
@@ -338,7 +338,7 @@ Engine/Physics/Constraints/ConstraintsReference#sharedproperties
 
 [EXCERPT:LinearDriveDamping]
 
-This is the damping to apply to the for linear drive.
+This is the damping to apply for linear drive.
 [VAR:ToolTipFullLink]
 Engine/Physics/Constraints/ConstraintsReference#sharedproperties
 [/VAR]
@@ -347,7 +347,7 @@ Engine/Physics/Constraints/ConstraintsReference#sharedproperties
 
 [EXCERPT:LinearDriveForceLimit]
 
-A large spring value and a low damping value, can result in substantial forces being applied to the Physics Bodies in an attempt to reach the Linear Position or Linear Velocity Drive goals. This value will limit the amount of force that can be applied to the Physics Bodies so that acceleration can occur over time and not at incredible rates which could introduce errors into the simulation.
+A large spring value and a low damping value can result in substantial forces being applied to the Physics Bodies in an attempt to reach the Linear Position or Linear Velocity Drive goals. This value will limit the amount of force that can be applied to the Physics Bodies so that acceleration can occur over time and not at incredible rates which could introduce errors into the simulation.
 [VAR:ToolTipFullLink]
 Engine/Physics/Constraints/ConstraintsReference#sharedproperties
 [/VAR]
@@ -394,7 +394,7 @@ Engine/Physics/Constraints/ConstraintsReference#sharedproperties
 
 [EXCERPT:AngularDriveSpring]
 
-This is the spring value to apply to the for angular drive.
+This is the spring value to apply to the angular drive.
 [VAR:ToolTipFullLink]
 Engine/Physics/Constraints/ConstraintsReference#sharedproperties
 [/VAR]
@@ -403,7 +403,7 @@ Engine/Physics/Constraints/ConstraintsReference#sharedproperties
 
 [EXCERPT:AngularDriveDamping]
 
-This is the damping value to apply to the for angulardrive.
+This is the damping value to apply to the angular drive.
 [VAR:ToolTipFullLink]
 Engine/Physics/Constraints/ConstraintsReference#sharedproperties
 [/VAR]
