#UE4doc #docs: sp: executed, flexible, cannot, persistent

diff --git a/Engine/Documentation/Source/Shared/Editors/BlueprintEditor/GraphTypes/GraphTypes.INT.udn b/Engine/Documentation/Source/Shared/Editors/BlueprintEditor/GraphTypes/GraphTypes.INT.udn
--- a/Engine/Documentation/Source/Shared/Editors/BlueprintEditor/GraphTypes/GraphTypes.INT.udn
+++ b/Engine/Documentation/Source/Shared/Editors/BlueprintEditor/GraphTypes/GraphTypes.INT.udn
@@ -5,10 +5,10 @@ Description:Overview of the Blueprint Graph types.
 
 
 [EXCERPT:ConstructionScript]
-The ConstructionScript function is exectuted whenever an instance of this Blueprint is placed or modified. It allows you to do several powerful things:
+The ConstructionScript function is executed whenever an instance of this Blueprint is placed or modified. It allows you to do several powerful things:
 
 * Configure components inside your Blueprint based on the value of exposed variables.
-* Create components programatically, allowing for procedural content.
+* Create components through programming allowing for procedural content.
 * Modify the blueprint based on the scene (e.g. results of line traces).
 
 [VAR:ToolTipFullLink]
@@ -18,7 +18,7 @@ Engine/Blueprints/UserGuide/UserConstructionScript
 [/EXCERPT:ConstructionScript]
 
 [EXCERPT:EventGraph]
-An Event Graph allows you to place many types of nodes to create different behaviors. Event graphs are very flexbile and can have multiple Event nodes to start execution, nodes that take time to execute (e.g. Delay), as well as Timeline nodes.
+An Event Graph allows you to place many types of nodes to create different behaviors. Event graphs are very flexible and can have multiple Event nodes to start execution, nodes that take time to execute (e.g. Delay), as well as Timeline nodes.
 
 [VAR:ToolTipFullLink]
 Engine/Blueprints/UserGuide/EventGraph
@@ -27,7 +27,7 @@ Engine/Blueprints/UserGuide/EventGraph
 [/EXCERPT:EventGraph]
 
 [EXCERPT:FunctionGraph]
-A Function Graph allows you to encapsulate functionality and re-use it. Function graphs have a single input and output node with a single execution pin, but you can add as many input and output data pins as you like. Functions can not contain any nodes that take time (e.g. Delay or Timelines) or Event nodes. You can add temporary 'local variable' nodes to a Function, but they will be wiped when the function is completed.
+A Function Graph allows you to encapsulate functionality and re-use it. Function graphs have a single input and output node with a single execution pin, but you can add as many input and output data pins as you like. Functions cannot contain any nodes that take time (e.g. Delay or Timelines) or Event nodes. You can add temporary 'local variable' nodes to a Function, but they will be wiped when the function is completed.
 
 You can drag-and-drop a Function from the MyBlueprint tab into the graph area to create a use of that function.
 
@@ -79,5 +79,5 @@ An Event is a node that starts execution in an Event Graph in response to some a
 [/EXCERPT:Event]
 
 [EXCERPT:LocalVariable]
-Local variables are variables that exist only within the scope of the function that they are created in, and the data is not persistant between calls of the functions.
+Local variables are variables that exist only within the scope of the function that they are created in, and the data is not persistent between calls of the functions.
 [/EXCERPT:LocalVariable]
