#UE4doc #docs: adjusted variable mismatch between pages. Ex: const float TimeBetweenCameraChanges = 3.0f;  -> const float TimeBetweenCameraChanges = 2.0f;

diff --git a/Engine/Documentation/Source/Programming/Tutorials/AutoCamera/1/AutoCamera_1.INT.udn b/Engine/Documentation/Source/Programming/Tutorials/AutoCamera/1/AutoCamera_1.INT.udn
--- a/Engine/Documentation/Source/Programming/Tutorials/AutoCamera/1/AutoCamera_1.INT.udn
+++ b/Engine/Documentation/Source/Programming/Tutorials/AutoCamera/1/AutoCamera_1.INT.udn
@@ -22,7 +22,7 @@ Description:Setting up a world with multiple cameras
 
 [REGION:tip]If you are new to **Unreal Engine 4**, you might want to read our [%Programming\QuickStart:title% tutorial](Programming\QuickStart) first. For this tutorial, we will assume you are familiar with creating a project, adding C++ code to it, compiling your code, and adding **Components** to **Actors** in the **Unreal Engine** editor.[/REGION]
 
-1. We will begin by creating a new, Basic Code project, with starter content, named "HowTo_AutoCamera". The first thing we'll need to do is create two cameras in our world. Since there are multiple ways to set up cameras, we'll go with the two most common ways here. For our first camera, go to the **Modes** tab in the **Placement Browser** and click **Place**, or press Shift-1. In the **All Classes** section, you'll find a **Camera** actor. Drag this into the **Level Editor** and position it so that it has a good view of our scene.
+1. We will begin by creating a new Basic Code project, with starter content, named "HowTo_AutoCamera". The first thing we'll need to do is create two cameras in our world. Since there are multiple ways to set up cameras, we'll go with the two most common ways here. For our first camera, go to the **Modes** tab in the **Placement Browser** and click **Place**, or press Shift-1. In the **All Classes** section, you'll find a **Camera** actor. Drag this into the **Level Editor** and position it so that it has a good view of our scene.
 
 	[![](PlaceCameraActor.png)](PlaceCameraActor.png)
 
diff --git a/Engine/Documentation/Source/Programming/Tutorials/AutoCamera/2/AutoCamera_2.INT.udn b/Engine/Documentation/Source/Programming/Tutorials/AutoCamera/2/AutoCamera_2.INT.udn
--- a/Engine/Documentation/Source/Programming/Tutorials/AutoCamera/2/AutoCamera_2.INT.udn
+++ b/Engine/Documentation/Source/Programming/Tutorials/AutoCamera/2/AutoCamera_2.INT.udn
@@ -43,7 +43,7 @@ Description:Coding a class to change camera views
 
 	The GameplayStatics header file gives us access to some useful general-purpose functions, one of which we will need for this tutorial. When that is done, add the following code to the bottom of **ACameraDirector::Tick**:
 
-		const float TimeBetweenCameraChanges = 3.0f;
+		const float TimeBetweenCameraChanges = 2.0f;
 		const float SmoothBlendTime = 0.75f;
 		TimeToNextCameraChange -= DeltaTime;
 		if (TimeToNextCameraChange <= 0.0f)
