#UE4doc #docs: that [is]   sp: project's, duplicate "the"

diff --git a/Engine/Documentation/Source/Engine/OpenWorldTools/ProceduralFoliage/QuickStart/ProceduralFoliageQuickStart.INT.udn b/Engine/Documentation/Source/Engine/OpenWorldTools/ProceduralFoliage/QuickStart/ProceduralFoliageQuickStart.INT.udn
--- a/Engine/Documentation/Source/Engine/OpenWorldTools/ProceduralFoliage/QuickStart/ProceduralFoliageQuickStart.INT.udn
+++ b/Engine/Documentation/Source/Engine/OpenWorldTools/ProceduralFoliage/QuickStart/ProceduralFoliageQuickStart.INT.udn
@@ -7,9 +7,9 @@ SkillLevel:Intermediate
 Version: 4.8
 
 
-In this Quick Start tutorial we will take a look at how the **Procedural Foliage Tool** that are now available with the 4.8 release of Unreal Engine 4 (UE4) work. 
+In this Quick Start tutorial we will take a look at how the **Procedural Foliage Tool** that is now available with the 4.8 release of Unreal Engine 4 (UE4) work. 
 Over the course of this Quick Start tutorial you will learn how to create, setup, and spawn an entire forests worth of trees inside of UE4 using only the Procedural Foliage tool. 
-You will also be introduced to key properties and settings that will help you shape your virtual forest to fit your projects needs. 
+You will also be introduced to key properties and settings that will help you shape your virtual forest to fit your project's needs. 
 
 ![](T_PF_QS_Final_Product.png)
 
@@ -33,7 +33,7 @@ Once the Open World Demo Collection is downloaded add it to the project that you
 	
 	![](Engine\OpenWorldTools\T_OWT_OWDC.png)
 
-1. Go the the **Library** section of the launcher and then under **Vault** section locate the Open World Demo Collection . 
+1. Go the **Library** section of the launcher and then under **Vault** section locate the Open World Demo Collection . 
 	
 	![](Engine\OpenWorldTools\T_OWT_Add_Content_00.png)
 
