#UE4doc #docs: inconsistent code: CollidingPawnMovementComponent.h not added to CollidingPawn.cpp until page 4

diff --git a/Engine/Documentation/Source/Programming/Tutorials/Components/1/Components_1.INT.udn b/Engine/Documentation/Source/Programming/Tutorials/Components/1/Components_1.INT.udn
--- a/Engine/Documentation/Source/Programming/Tutorials/Components/1/Components_1.INT.udn
+++ b/Engine/Documentation/Source/Programming/Tutorials/Components/1/Components_1.INT.udn
@@ -133,7 +133,6 @@ Our new **Pawn** now has a collection of useful **Components** attached to it, a
 
 			#include "HowTo_Components.h"
 			#include "CollidingPawn.h"
-			#include "CollidingPawnMovementComponent.h"
 
 			// Sets default values
 			ACollidingPawn::ACollidingPawn()
