#UE4doc #docs: Variable name mismatch: MovementComponent -> OurMovementComponent

diff --git a/Engine/Documentation/Source/Programming/Tutorials/Components/4/Components_4.INT.udn b/Engine/Documentation/Source/Programming/Tutorials/Components/4/Components_4.INT.udn
--- a/Engine/Documentation/Source/Programming/Tutorials/Components/4/Components_4.INT.udn
+++ b/Engine/Documentation/Source/Programming/Tutorials/Components/4/Components_4.INT.udn
@@ -32,8 +32,8 @@ Description:Tying the functionality of the Pawn, the input configuration, and th
 	Creating a **Pawn Movement Component** and associating it with our **Pawn** is simple. At the bottom of **ACollidingPawn::ACollidingPawn**, we can add this code:
 
 		// Create an instance of our movement component, and tell it to update the root.
-		MovementComponent = CreateDefaultSubobject<UCollidingPawnMovementComponent>(TEXT("CustomMovementComponent"));
-		MovementComponent->UpdatedComponent = RootComponent;
+		OurMovementComponent = CreateDefaultSubobject<UCollidingPawnMovementComponent>(TEXT("CustomMovementComponent"));
+		OurMovementComponent->UpdatedComponent = RootComponent;
 
 	[REGION:note]Unlike the other **Components** we've seen so far, we don't need to attach this **Component** to our own **Component** hierarchy. This is because our other **Components** are all types of **Scene Components**, which inherently require physical locations. **Movement Controllers**, however, are not **Scene Components** and do not represent physical objects, so the concept of existing at a physical location or being physically attached to another **Component** does not apply to them.[/REGION]
 
