#UE4doc #docs: tense: replaced

diff --git a/Engine/Documentation/Source/Engine/OpenWorldTools/ProceduralFoliage/QuickStart/3/PFQS_3.INT.udn b/Engine/Documentation/Source/Engine/OpenWorldTools/ProceduralFoliage/QuickStart/3/PFQS_3.INT.udn
--- a/Engine/Documentation/Source/Engine/OpenWorldTools/ProceduralFoliage/QuickStart/3/PFQS_3.INT.udn
+++ b/Engine/Documentation/Source/Engine/OpenWorldTools/ProceduralFoliage/QuickStart/3/PFQS_3.INT.udn
@@ -41,7 +41,7 @@ We will continue to work with the **PFT_00** level that was used in the last ste
 	![](T_Shade_Radius.png)
 
 	[region:note]
-	The Collision section determines which Foliage types should be removed when two Foliage Types are competing for the same spawn location or relative space. When a Virtual Seed's collision radius overlaps an existing collision or shade radius from another Foliage Type's Virtual Seed, the Virtual Seed will be replace or killed based which Foliage Type has the higher propriety.
+	The Collision section determines which Foliage types should be removed when two Foliage Types are competing for the same spawn location or relative space. When a Virtual Seed's collision radius overlaps an existing collision or shade radius from another Foliage Type's Virtual Seed, the Virtual Seed will be replaced or killed based which Foliage Type has the higher propriety.
 	[/region]
 
 1. Select the **Procedural Foliage Spawner** that was placed in the level and under **Procedural Foliage** section click on the **Resimulate** button.
