#UE4doc #docs: code missing sections for OurCameraSpringArm

diff --git a/Engine/Documentation/Source/Programming/Tutorials/PlayerCamera/2/PlayerCamera_2.INT.udn b/Engine/Documentation/Source/Programming/Tutorials/PlayerCamera/2/PlayerCamera_2.INT.udn
--- a/Engine/Documentation/Source/Programming/Tutorials/PlayerCamera/2/PlayerCamera_2.INT.udn
+++ b/Engine/Documentation/Source/Programming/Tutorials/PlayerCamera/2/PlayerCamera_2.INT.udn
@@ -104,6 +104,9 @@ Now that we have defined our input, we need to write some code to react to it. W
 				OurCameraSpringArm = CreateDefaultSubobject<USpringArmComponent>(TEXT("CameraSpringArm"));
 				OurCameraSpringArm->AttachTo(RootComponent);
 				OurCameraSpringArm->SetRelativeLocationAndRotation(FVector(0.0f, 0.0f, 50.0f), FRotator(-60.0f, 0.0f, 0.0f));
+				OurCameraSpringArm->TargetArmLength = 400.f
+				OurCameraSpringArm->bEnableCameraLag = true;
+				OurCameraSprintArm->CameraLagSpeed = 3.0f
 				OurCamera = CreateDefaultSubobject<UCameraComponent>(TEXT("GameCamera"));
 				OurCamera->AttachTo(OurCameraSpringArm, USpringArmComponent::SocketName);
 
