#UE4doc #docs: sp: Notifies changed "all list" to "a list" sp: asset

diff --git a/Engine/Documentation/Source/Shared/Editors/Persona/PersonaEditorModes.INT.udn b/Engine/Documentation/Source/Shared/Editors/Persona/PersonaEditorModes.INT.udn
--- a/Engine/Documentation/Source/Shared/Editors/Persona/PersonaEditorModes.INT.udn
+++ b/Engine/Documentation/Source/Shared/Editors/Persona/PersonaEditorModes.INT.udn
@@ -7,7 +7,7 @@ Description:Overview of the Persona Editor modes.
 [EXCERPT:SkeletonMode]
 
 The skeleton mode lets you view properties of a skeleton asset.
-* Edit Animation Notifes that belong to this skeleton
+* Edit Animation Notifies that belong to this skeleton
 
 ### Overview
 
@@ -110,7 +110,7 @@ Engine/Animation/Skeleton
 
 ###Overview
 
-This shows all list of bones for skeleton or for mesh. By default, it displays all the bones that belong to mesh and skeleton. 
+This shows a list of bones for skeleton or for mesh. By default, it displays all the bones that belong to mesh and skeleton. 
 
 * Bone Display
 	* Bold/White	: Current preview mesh has the bone and the bone is weighted.
@@ -245,7 +245,7 @@ This window displays the property of current asset that's being displayed.
 
 [EXCERPT:AnimationEdit_Window]
 
-In this panel, you can edit each assets. You can add AnimSequence, BlendSpace, AnimMontage, or AnimComposite. 
+In this panel, you can edit each asset. You can add AnimSequence, BlendSpace, AnimMontage, or AnimComposite. 
 
 [/EXCERPT:AnimationEdit_Window]
 
@@ -330,10 +330,10 @@ If you turn this off, the animation data will not be converted but copied direct
 
 [EXCERPT:AnimRemapSkeleton_ShowCompatbielSkeletons]
 
-This option will display only compatible skeletons. To make this to be useful, you need to configure rig option in Retarget Manager. 
+This option will display only compatible skeletons. To make this useful, you need to configure rig option in Retarget Manager. 
 
 Make sure you have same retarget base pose in both skeletons. If not please use Retarget Manager - Manage Retarget Base Poses to modify. 
 
-We recommend to set up to T-pose. 
+We recommend setting up your character in a T-pose. 
 
 [/EXCERPT:AnimRemapSkeleton_ShowCompatbielSkeletons]
