#UE4doc #docs: needed to add: UPROPERTY(EditAnywhere) throughout the tutorial so edits could be made in the details panel. This is the last step to be implemented into the project (page 3).

diff --git a/Engine/Documentation/Source/Programming/Tutorials/PlayerCamera/1/PlayerCamera_1.INT.udn b/Engine/Documentation/Source/Programming/Tutorials/PlayerCamera/1/PlayerCamera_1.INT.udn
--- a/Engine/Documentation/Source/Programming/Tutorials/PlayerCamera/1/PlayerCamera_1.INT.udn
+++ b/Engine/Documentation/Source/Programming/Tutorials/PlayerCamera/1/PlayerCamera_1.INT.udn
@@ -29,6 +29,7 @@ Description:Creating a spring arm and a camera.
 1. Next, in **Visual Studio**, we should open `PawnWithCamera.h` and add the following code to the bottom of our class definition:
 
 		protected:
+			UPROPERTY(EditAnywhere)
 			USpringArmComponent* OurCameraSpringArm;
 			UCameraComponent* OurCamera;
 
@@ -94,6 +95,7 @@ We now have a simple **Pawn** that will allow us to control our camera comfortab
 				virtual void SetupPlayerInputComponent(class UInputComponent* InputComponent) override;
 
 			protected:
+				UPROPERTY(EditAnywhere)
 				USpringArmComponent* OurCameraSpringArm;
 				UCameraComponent* OurCamera;
 			};
diff --git a/Engine/Documentation/Source/Programming/Tutorials/PlayerCamera/2/PlayerCamera_2.INT.udn b/Engine/Documentation/Source/Programming/Tutorials/PlayerCamera/2/PlayerCamera_2.INT.udn
--- a/Engine/Documentation/Source/Programming/Tutorials/PlayerCamera/2/PlayerCamera_2.INT.udn
+++ b/Engine/Documentation/Source/Programming/Tutorials/PlayerCamera/2/PlayerCamera_2.INT.udn
@@ -82,7 +82,9 @@ Now that we have defined our input, we need to write some code to react to it. W
 
 				// Called to bind functionality to input
 				virtual void SetupPlayerInputComponent(class UInputComponent* InputComponent) override;
-
+			
+			protected:
+				UPROPERTY(EditAnywhere)
 				USpringArmComponent* OurCameraSpringArm;
 				UCameraComponent* OurCamera;
 			};
diff --git a/Engine/Documentation/Source/Programming/Tutorials/PlayerCamera/3/PlayerCamera_3.INT.udn b/Engine/Documentation/Source/Programming/Tutorials/PlayerCamera/3/PlayerCamera_3.INT.udn
--- a/Engine/Documentation/Source/Programming/Tutorials/PlayerCamera/3/PlayerCamera_3.INT.udn
+++ b/Engine/Documentation/Source/Programming/Tutorials/PlayerCamera/3/PlayerCamera_3.INT.udn
@@ -180,6 +180,7 @@ Description:Implementing code to control a pawn and a camera with input.
 				virtual void SetupPlayerInputComponent(class UInputComponent* InputComponent) override;
 
 			protected:
+				UPROPERTY(EditAnywhere)
 				USpringArmComponent* OurCameraSpringArm;
 				UCameraComponent* OurCamera;
 
