#UE4doc #docs: sentence phrasing: "...is a 3D model asset that can be animated"

diff --git a/Engine/Documentation/Source/Shared/Glossary/S/Glossary_S.INT.udn b/Engine/Documentation/Source/Shared/Glossary/S/Glossary_S.INT.udn
--- a/Engine/Documentation/Source/Shared/Glossary/S/Glossary_S.INT.udn
+++ b/Engine/Documentation/Source/Shared/Glossary/S/Glossary_S.INT.udn
@@ -54,7 +54,7 @@ Skeletal Mesh Actor
 	[PARAM:Definition]
 
 [EXCERPT:Glossary_Skeletal_Mesh_Actor]
-A **Skeletal Mesh Actor** is an asset that was created as an animatable 3D model in an external 3D modeling program and then imported into the Unreal Editor. A Skeletal Mesh Actor consists of two parts: the mesh itself and a hierarchical skeleton whose bones can be used to create animation by deforming the mesh according to rules defined in the 3D modeling program. You can import Skeletal Mesh Actors as assets in the **Content Browser**.
+A **Skeletal Mesh Actor** is a 3D model asset that can be animated in an external 3D modeling program and then imported into the Unreal Editor. A Skeletal Mesh Actor consists of two parts: the mesh itself and a hierarchical skeleton whose bones can be used to create animation by deforming the mesh according to rules defined in the 3D modeling program. You can import Skeletal Mesh Actors as assets in the **Content Browser**.
 [/EXCERPT:Glossary_Skeletal_Mesh_Actor]
 
 	[/PARAM]
