#UE4doc #docs: removed "a" from "the a timeline"

diff --git a/Engine/Documentation/Source/Shared/Glossary/T/Glossary_T.INT.udn b/Engine/Documentation/Source/Shared/Glossary/T/Glossary_T.INT.udn
--- a/Engine/Documentation/Source/Shared/Glossary/T/Glossary_T.INT.udn
+++ b/Engine/Documentation/Source/Shared/Glossary/T/Glossary_T.INT.udn
@@ -42,7 +42,7 @@ time cursor
 	[PARAM:Definition]
 
 [EXCERPT:Glossary_time_cursor]
-The **time cursor** is the black vertical bar at the bottom of the a timeline that indicates the position in time of the animation's playback. You can drag the time cursor back and forth to preview the animation in the 3D viewport.
+The **time cursor** is the black vertical bar at the bottom of the timeline that indicates the position in time of the animation's playback. You can drag the time cursor back and forth to preview the animation in the 3D viewport.
 [/EXCERPT:Glossary_time_cursor]
 
 	[/PARAM]
