#UE4doc #docs:  changed UPROPERTY to UPROPERTY(EditAnywhere). This way, the user can make the needed changes in the details panel. Changes had to be implemented across all pages to match the updated code section. Tutorial up to date with 4.8

diff --git a/Engine/Documentation/Source/Programming/Tutorials/PlayerInput/1/PlayerInput_1.INT.udn b/Engine/Documentation/Source/Programming/Tutorials/PlayerInput/1/PlayerInput_1.INT.udn
--- a/Engine/Documentation/Source/Programming/Tutorials/PlayerInput/1/PlayerInput_1.INT.udn
+++ b/Engine/Documentation/Source/Programming/Tutorials/PlayerInput/1/PlayerInput_1.INT.udn
@@ -37,7 +37,7 @@ Description:Introduction to Pawns
 
 1. Next, we'll build a few basic **Components**. If you want to learn more about adding and managing **Components** in code, as well as some common types of **Components** you will encounter, try reading our [%Programming\Tutorials\Components:title% tutorial](Programming\Tutorials\Components). In order to keep track of the **Component** we will create, we should add the following code to `MyPawn.h`, at the bottom of our class definition:
 
-		UPROPERTY()
+		UPROPERTY(EditAnywhere)
 		USceneComponent* OurVisibleComponent;
 
 	[REGION:note]This variable is tagged as a **UPROPERTY** so that it will be visible to **Unreal Engine**. This is important because it prevents the variable from being reset when the game is launched, or when the project or level is closed and reloaded.[/REGION]
@@ -95,7 +95,7 @@ Now that we have a customized **Pawn** to react to our game's input, we'll need
 				// Called to bind functionality to input
 				virtual void SetupPlayerInputComponent(class UInputComponent* InputComponent) override;
 
-				UPROPERTY()
+				UPROPERTY(EditAnywhere)
 				USceneComponent* OurVisibleComponent;
 			};
 
diff --git a/Engine/Documentation/Source/Programming/Tutorials/PlayerInput/2/PlayerInput_2.INT.udn b/Engine/Documentation/Source/Programming/Tutorials/PlayerInput/2/PlayerInput_2.INT.udn
--- a/Engine/Documentation/Source/Programming/Tutorials/PlayerInput/2/PlayerInput_2.INT.udn
+++ b/Engine/Documentation/Source/Programming/Tutorials/PlayerInput/2/PlayerInput_2.INT.udn
@@ -98,7 +98,7 @@ We're now ready to finish coding our MyPawn class in **Visual Studio**.
 				// Called to bind functionality to input
 				virtual void SetupPlayerInputComponent(class UInputComponent* InputComponent) override;
 
-				UPROPERTY()
+				UPROPERTY(EditAnywhere)
 				USceneComponent* OurVisibleComponent;
 			};
 
diff --git a/Engine/Documentation/Source/Programming/Tutorials/PlayerInput/3/PlayerInput_3.INT.udn b/Engine/Documentation/Source/Programming/Tutorials/PlayerInput/3/PlayerInput_3.INT.udn
--- a/Engine/Documentation/Source/Programming/Tutorials/PlayerInput/3/PlayerInput_3.INT.udn
+++ b/Engine/Documentation/Source/Programming/Tutorials/PlayerInput/3/PlayerInput_3.INT.udn
@@ -141,7 +141,7 @@ Description:Binding input to C++ code
 				// Called to bind functionality to input
 				virtual void SetupPlayerInputComponent(class UInputComponent* InputComponent) override;
 
-				UPROPERTY()
+				UPROPERTY(EditAnywhere)
 				USceneComponent* OurVisibleComponent;
 
 				// Input functions
