#UE4doc #docs: sp: parameter, access

diff --git a/Engine/Documentation/Source/Shared/GraphNodes/Animation/BlendNodes.INT.udn b/Engine/Documentation/Source/Shared/GraphNodes/Animation/BlendNodes.INT.udn
--- a/Engine/Documentation/Source/Shared/GraphNodes/Animation/BlendNodes.INT.udn
+++ b/Engine/Documentation/Source/Shared/GraphNodes/Animation/BlendNodes.INT.udn
@@ -67,7 +67,7 @@ Engine/Animation/NodeReference/Blend#blend
 
 [EXCERPT:UAnimGraphNode_ApplyAdditive]
 
-Applies the supplied additive pose blended by the Alpha paramater.
+Applies the supplied additive pose blended by the Alpha parameter.
 
 [/EXCERPT:UAnimGraphNode_ApplyAdditive]
 
@@ -109,7 +109,7 @@ Rotate Root Bone allows the arbitrary rotation of a bone outside of animation da
 
 [EXCERPT:UAnimGraphNode_SaveCachedPose]
 
-Saves the supplied pose so that it can be accesed and used at a later point.
+Saves the supplied pose so that it can be accessed and used at a later point.
 
 [/EXCERPT:UAnimGraphNode_SaveCachedPose]
 
