#UE4 Docs: Adding in the hooks so that the OWT pages are not orphan pages.

diff --git a/Engine/Documentation/Source/Engine/Engine.INT.udn b/Engine/Documentation/Source/Engine/Engine.INT.udn
--- a/Engine/Documentation/Source/Engine/Engine.INT.udn
+++ b/Engine/Documentation/Source/Engine/Engine.INT.udn
@@ -86,6 +86,7 @@ engine
 %Engine/AI/BehaviorTrees:topiccompact%
 %Engine/Landscape:topiccompact%
 %Engine/Foliage:topiccompact%
+%Engine/OpenWorldTools:topiccompact%
 %Engine/LevelStreaming:topiccompact%
 %Engine/Audio:topiccompact%
 %Engine/Physics:topiccompact%
diff --git a/Engine/Documentation/Source/Engine/OpenWorldTools/OpenWorldTools.INT.udn b/Engine/Documentation/Source/Engine/OpenWorldTools/OpenWorldTools.INT.udn
--- a/Engine/Documentation/Source/Engine/OpenWorldTools/OpenWorldTools.INT.udn
+++ b/Engine/Documentation/Source/Engine/OpenWorldTools/OpenWorldTools.INT.udn
@@ -15,6 +15,26 @@ With the release of Unreal Engine 4.8 the tools used to create this demo, referr
 Please note that the Open World Tools are still in an experimental state which means that you could run into bugs or stability issues while using the tools.
 [/region]
 
+[VAR:TopicCompact]
+[OBJECT:TopicCompact]
+	[PARAM:image]
+	![](Engine\OpenWorldTools\Overview\T_Kite_Demo_00.png)
+	[/PARAM]
+	[PARAM:icon]
+		![](%ROOT%/landscape_icon.png)(convert:false)
+	[/PARAM]
+	[PARAM:title]
+		%Engine/OpenWorldTools:title%
+	[/PARAM]
+	[PARAM:description]
+		%Engine/OpenWorldTools:description%
+	[/PARAM]
+	[PARAM:path]
+		[RELATIVE:Engine/OpenWorldTools]
+	[/PARAM]
+[/OBJECT]
+[/VAR]
+
 [REGION:topics third]
 [OBJECT:Topic]
 	[PARAM:image]
@@ -70,6 +90,23 @@ Please note that the Open World Tools are still in an experimental state which m
 	[/PARAM]
 [/OBJECT]
 
+[OBJECT:Topic]
+	[PARAM:image]
+		![](Engine\OpenWorldTools\ProceduralFoliage\QuickStart\6\T_On_Your_Own.png)
+	[/PARAM]
+	[PARAM:icon]
+		![](%ROOT%/blueprint_icon.png)(convert:false)
+	[/PARAM]
+	[PARAM:title]
+		%Engine/OpenWorldTools/TipsAndTricks:title%
+	[/PARAM]
+	[PARAM:description]
+		%Engine/OpenWorldTools/TipsAndTricks:description%
+	[/PARAM]
+	[PARAM:path]
+		[RELATIVE:Engine/OpenWorldTools/TipsAndTricks]
+	[/PARAM]
+[/OBJECT]
 
 [/REGION]
 
diff --git a/Engine/Documentation/Source/Engine/OpenWorldTools/ProceduralFoliage/QuickStart/6/PFQS_6.INT.udn b/Engine/Documentation/Source/Engine/OpenWorldTools/ProceduralFoliage/QuickStart/6/PFQS_6.INT.udn
--- a/Engine/Documentation/Source/Engine/OpenWorldTools/ProceduralFoliage/QuickStart/6/PFQS_6.INT.udn
+++ b/Engine/Documentation/Source/Engine/OpenWorldTools/ProceduralFoliage/QuickStart/6/PFQS_6.INT.udn
@@ -1,5 +1,5 @@
 Availability:Docs
-Title: 6. Own your Own
+Title: 6. On your Own
 Crumbs:
 Description: What to do next with the Procedural Foliage tool.
 Navigation:topic
