#UE4doc #docs: sp: referred

diff --git a/Engine/Documentation/Source/Shared/GraphNodes/Blueprint/BlueprintNodes.INT.udn b/Engine/Documentation/Source/Shared/GraphNodes/Blueprint/BlueprintNodes.INT.udn
--- a/Engine/Documentation/Source/Shared/GraphNodes/Blueprint/BlueprintNodes.INT.udn
+++ b/Engine/Documentation/Source/Shared/GraphNodes/Blueprint/BlueprintNodes.INT.udn
@@ -49,7 +49,7 @@ pulse down the appropriate output.
 
 **Input Pins**
 
-Selection - Takes in a Enum value used to indicate which output pin will be triggered. 
+Selection - Takes in an Enum value used to indicate which output pin will be triggered. 
 
 ---
 
