#UE4doc #docs: duplicate word "is" sp: emitter

diff --git a/Engine/Documentation/Source/Shared/GraphNodes/Blueprint/UK2Node_AddComponent/UK2Node_AddComponent.INT.udn b/Engine/Documentation/Source/Shared/GraphNodes/Blueprint/UK2Node_AddComponent/UK2Node_AddComponent.INT.udn
--- a/Engine/Documentation/Source/Shared/GraphNodes/Blueprint/UK2Node_AddComponent/UK2Node_AddComponent.INT.udn
+++ b/Engine/Documentation/Source/Shared/GraphNodes/Blueprint/UK2Node_AddComponent/UK2Node_AddComponent.INT.udn
@@ -137,7 +137,7 @@ The 'Relative Transform' input defines where to position and orient the componen
 
 [EXCERPT:ParticleSystemComponent]
 
-A ParticleSystemComponent acts as a particle emmitter. When executed, this node creates a new ParticleSystemComponent for the specified 'Target' Actor.  
+A ParticleSystemComponent acts as a particle emitter. When executed, this node creates a new ParticleSystemComponent for the specified 'Target' Actor.  
 
 If the 'Manual Attachment' input is left unset (False), then the new component will automatically be attached under the Actor's root (or made the root if one did not previously exist). If the 'Manual Attachment' input is set (True), then the new component is left unattached.
 
@@ -174,7 +174,7 @@ The 'Relative Transform' input defines where to position and orient the componen
 
 [EXCERPT:PhysicsConstraintComponent]
 
-PhysicsConstraintComponent is is effectively a joint that allows you to connect two rigid bodies together. You can create different types of joints using the various parameters of this component.
+PhysicsConstraintComponent is effectively a joint that allows you to connect two rigid bodies together. You can create different types of joints using the various parameters of this component.
 
 When executed, this node creates a new PhysicsConstraintComponent for the specified 'Target' Actor.  
 
