#UE4doc #docs: sp: receives

diff --git a/Engine/Documentation/Source/Shared/GraphNodes/Blueprint/UK2Node_MacroInstance/UK2Node_MacroInstance.INT.udn b/Engine/Documentation/Source/Shared/GraphNodes/Blueprint/UK2Node_MacroInstance/UK2Node_MacroInstance.INT.udn
--- a/Engine/Documentation/Source/Shared/GraphNodes/Blueprint/UK2Node_MacroInstance/UK2Node_MacroInstance.INT.udn
+++ b/Engine/Documentation/Source/Shared/GraphNodes/Blueprint/UK2Node_MacroInstance/UK2Node_MacroInstance.INT.udn
@@ -12,7 +12,7 @@ The 'Loop Body' output sends an execution pulse for each entry in the 'Array' in
 
 The ForEachLoop node is a macro node. You can double-click macro nodes to see how they are constructed.
 
-The 'Completed' execution output recieves a pulse when the node has gone through every element in the 'Array'.
+The 'Completed' execution output receives a pulse when the node has gone through every element in the 'Array'.
 
 [VAR:ToolTipFullLink]
 Engine/Blueprints/UserGuide/FlowControl#forloop
