#UE4doc #docs: sp: explicitly, an

diff --git a/Engine/Documentation/Source/Shared/GraphNodes/Material/Material.INT.udn b/Engine/Documentation/Source/Shared/GraphNodes/Material/Material.INT.udn
--- a/Engine/Documentation/Source/Shared/GraphNodes/Material/Material.INT.udn
+++ b/Engine/Documentation/Source/Shared/GraphNodes/Material/Material.INT.udn
@@ -313,7 +313,7 @@ Engine/Rendering/Materials/ExpressionReference/Custom#customtexture
 [/EXCERPT:UMaterialExpressionCustomTexture]
 
 [EXCERPT:UMaterialExpressionDecalMipmapLevel]
-The DecalMipmapLevel is used to explitily compute a decal's texture's mipmap level assuming there is no custom uv, to avoid the 2x2 pixel mipmap artefacts on the edges where the decal is project to.
+The DecalMipmapLevel is used to explicitly compute a decal's texture's mipmap level assuming there is no custom uv, to avoid the 2x2 pixel mipmap artefacts on the edges where the decal is project to.
 
 [VAR:ToolTipFullLink]
 Engine/Actors/DecalActor/DecalUserGuide
@@ -884,7 +884,7 @@ Engine/Rendering/Materials/ExpressionReference/Utility#rotateaboutaxis
 [/EXCERPT:UMaterialExpressionRotateAboutAxis]
 
 [EXCERPT:UMaterialExpressionAntialiasedTextureMask]
-The AntialiasedTextureMask expression allows you to create a material using a soft (anti-aliased) transition mask. The mask can be used to blend between two complex material properties or to fade out a alpha blended material (works well with SoftMasked). 
+The AntialiasedTextureMask expression allows you to create a material using a soft (anti-aliased) transition mask. The mask can be used to blend between two complex material properties or to fade out an alpha blended material (works well with SoftMasked). 
 
 Simple specify a texture with the mask specified in one channel (red, green, blue, or alpha), set the used channel in the expression and specify the comparison value.
 
