#UE4doc #docs: sp: accounting, Blueprint, inherits, executing

diff --git a/Engine/Documentation/Source/Shared/Kismet/FBlueprintEditorCommands.INT.udn b/Engine/Documentation/Source/Shared/Kismet/FBlueprintEditorCommands.INT.udn
--- a/Engine/Documentation/Source/Shared/Kismet/FBlueprintEditorCommands.INT.udn
+++ b/Engine/Documentation/Source/Shared/Kismet/FBlueprintEditorCommands.INT.udn
@@ -17,7 +17,7 @@ Performs the same function as the Compile button on the tool bar:
 [EXCERPT:RefreshAllNodes]
 
 
-If you are making edits to other Blueprints, and trying to interact or call functions on external Blueprints, the input and output pins for these can change in number and type. This will refresh all the nodes in the current Blueprint, acounting for external changes.
+If you are making edits to other Blueprints, and trying to interact or call functions on external Blueprints, the input and output pins for these can change in number and type. This will refresh all the nodes in the current Blueprint, accounting for external changes.
 
 You can do this to an individual node in the graph by right-clicking it and selecting "Refresh Nodes."
 
@@ -39,7 +39,7 @@ This process can be undone.
 
 This will open the "Find Results" pane and set it to search all Blueprints in the project for the given search term.
 
-You can also search all Blueprints by opening the "Find Results" tab (Window -> Find Results) and disabling "Find in Current Bluerpint Only."
+You can also search all Blueprints by opening the "Find Results" tab (Window -> Find Results) and disabling "Find in Current Blueprint Only."
 
 ![](k2_findAll.png)
 
@@ -69,7 +69,7 @@ You can also search all Blueprints by opening the "Find Results" tab (Window ->
 
 This will open the "Find Results" pane and set it to search the current Blueprint for the given search term.
 
-You can also search the current Blueprint by opening the "Find Results" tab (Window -> Find Results) and enabling "Find in Current Bluerpint Only."
+You can also search the current Blueprint by opening the "Find Results" tab (Window -> Find Results) and enabling "Find in Current Blueprint Only."
 
 ![](k2_findCurrent.png)
 
@@ -79,7 +79,7 @@ You can also search the current Blueprint by opening the "Find Results" tab (Win
 [EXCERPT:ReparentBlueprint]
 
 
-This will pop up a dialog enabling you to change the class that the Blueprint inheirits from. Really useful if you realize you needed a Character Movement Component, but you started your Blueprint from Actor (that does not contain that component). In this example, you can use this tool to change the parent class of your Blueprint from Actor to Character or Pawn.
+This will pop up a dialog enabling you to change the class that the Blueprint inherits from. Really useful if you realize you needed a Character Movement Component, but you started your Blueprint from Actor (that does not contain that component). In this example, you can use this tool to change the parent class of your Blueprint from Actor to Character or Pawn.
 
 ![](k2_reparent.png)
 
@@ -201,7 +201,7 @@ Individual breakpoints can be removed by right-clicking the node with a breakpoi
 
 This will delete all watches in the current Blueprint.
 
-"Watches" are set on non-excuting pins, such as the pins on an addition node:
+"Watches" are set on non-executing pins, such as the pins on an addition node:
 
 ![](k2_watch.png)
 
