#UE4doc #docs: made grammar and punctuation changes.

diff --git a/Engine/Documentation/Source/Programming/Tutorials/UMG/2/UMG_2.INT.udn b/Engine/Documentation/Source/Programming/Tutorials/UMG/2/UMG_2.INT.udn
--- a/Engine/Documentation/Source/Programming/Tutorials/UMG/2/UMG_2.INT.udn
+++ b/Engine/Documentation/Source/Programming/Tutorials/UMG/2/UMG_2.INT.udn
@@ -55,7 +55,7 @@ Description:Extending the project's Game Mode to add and remove menus from the s
 
 	[REGION:note]When overriding functions from a parent class (referenced by the word Super), as we do here with _BeginPlay_, it is often important to call the parent class' version of that function. Since our version of the function is only meant to add one step to the end of the existing procedure, we call _Super::BeginPlay_ in the first line of the function.[/REGION]
 
-1. Next, still in `HowTo_UMGGameMode.cpp`, we need to define how we change between menus. We will need to remove whatever **User Widget** we have active from the viewport, if any, first. After that, we can create and a new **User Widget** and add it to the viewport.
+1. Next, still in `HowTo_UMGGameMode.cpp`, we need to define how we change between menus. We will need to remove whatever **User Widget** we have active from the viewport, if any. Then we can create and a new **User Widget** and add it to the viewport.
 
 		void AHowTo_UMGGameMode::ChangeMenuWidget(TSubclassOf<UUserWidget> NewWidgetClass)
 		{
