#UE4doc #docs: sp: duplicate word: the

diff --git a/Engine/Documentation/Source/Shared/LandscapeEditor/LandscapeEditor.INT.udn b/Engine/Documentation/Source/Shared/LandscapeEditor/LandscapeEditor.INT.udn
--- a/Engine/Documentation/Source/Shared/LandscapeEditor/LandscapeEditor.INT.udn
+++ b/Engine/Documentation/Source/Shared/LandscapeEditor/LandscapeEditor.INT.udn
@@ -100,7 +100,7 @@ Engine/Landscape/Editing/SculptMode
 [/EXCERPT:FlattenTool]
 
 [EXCERPT:RampTool]
-The Ramp Tool enables you to make a ramp between two points on your Landscape. You can even specify a falloff for the sides of the Ramp that determines how sharply the sides of the Ramp will falloff. To use the Ramp tool, simply hold down the CRTL key and use the Left Mouse Button to place where you want the first point to be located on the Landscape. To place the 2nd Ramp point, keep the CTRL key held down and simply click, using the Left Mouse Button, anywhere else on the Landscape. Once your Ramp points are set you can adjust their location by selecting either icon and moving them using the Gizmo. Once you are happy with the location of your Ramp, you can add the Ramp to the Landscape by using the Add Ramp Button from the Tool Settings menu. If you want to undo the Ramp, simply press the Reset Button to remove the Ramp.
+The Ramp Tool enables you to make a ramp between two points on your Landscape. You can even specify a fall off for the sides of the Ramp that determines how sharply the sides of the Ramp will falloff. To use the Ramp tool, simply hold down the CRTL key and use the Left Mouse Button to place where you want the first point to be located on the Landscape. To place the 2nd Ramp point, keep the CTRL key held down and simply click, using the Left Mouse Button, anywhere else on the Landscape. Once your Ramp points are set you can adjust their location by selecting either icon and moving them using the Gizmo. Once you are happy with the location of your Ramp, you can add the Ramp to the Landscape by using the Add Ramp Button from the Tool Settings menu. If you want to undo the Ramp, simply press the Reset Button to remove the Ramp.
 ![](Ramp_Before_And_After_00.png)
 Left is before using the Ramp Tool. Right is after using the Ramp tool.
 [VAR:ToolTipFullLink]
@@ -174,7 +174,7 @@ Engine/Landscape/Editing/Brushes
 [/EXCERPT:CircleBrush]
 
 [EXCERPT:AlphaBrush]
-The Alpha brush is similar to the Pattern brush, but instead of tiling the texture across the Landscape as you paint, it orients the brush texture in the direction of your painting and drags the shape as you move the cursor. You can pack multiple brush patterns in the the various RGB channels of a texture. To select a different pattern simply select a different channel of the texture that is currently assigned using the Texture Channel option.  
+The Alpha brush is similar to the Pattern brush, but instead of tiling the texture across the Landscape as you paint, it orients the brush texture in the direction of your painting and drags the shape as you move the cursor. You can pack multiple brush patterns in the various RGB channels of a texture. To select a different pattern simply select a different channel of the texture that is currently assigned using the Texture Channel option.  
 ![](brush_alpha_drag.png)
 [VAR:ToolTipFullLink]
 Engine/Landscape/Editing/Brushes
@@ -182,7 +182,7 @@ Engine/Landscape/Editing/Brushes
 [/EXCERPT:AlphaBrush]
 
 [EXCERPT:AlphaBrush_Pattern]
-The Pattern brush allows you to use an arbitrary shape for the brush by sampling a single color channel from a texture to use as the alpha for the brush. The texture pattern is tiled as the brush is painted. You can pack multiple brush patterns in the the various RGB channels of a texture. To select a different pattern, simply select a different channel of the texture that is currently assigned using the Texture Channel option.  
+The Pattern brush allows you to use an arbitrary shape for the brush by sampling a single color channel from a texture to use as the alpha for the brush. The texture pattern is tiled as the brush is painted. You can pack multiple brush patterns in the various RGB channels of a texture. To select a different pattern, simply select a different channel of the texture that is currently assigned using the Texture Channel option.  
 ![](Patteren_Brush_00.png)
 [VAR:ToolTipFullLink]
 Engine/Landscape/Editing/Brushes
