#UE4doc #docs: sp: vertices

diff --git a/Engine/Documentation/Source/Shared/LevelEditorModes/LevelEditorModes.INT.udn b/Engine/Documentation/Source/Shared/LevelEditorModes/LevelEditorModes.INT.udn
--- a/Engine/Documentation/Source/Shared/LevelEditorModes/LevelEditorModes.INT.udn
+++ b/Engine/Documentation/Source/Shared/LevelEditorModes/LevelEditorModes.INT.udn
@@ -37,7 +37,7 @@ Engine/UI/Tools/Foliage
 
 [EXCERPT:EditorMode.EM_Geometry]
 ###Geometry Editing
-Allows you to edit verticies of placed geometry brushes.
+Allows you to edit vertices of placed geometry brushes.
 
 [VAR:ToolTipFullLink]
 Resources/ContentExamples/Brushes
