#UE4doc #docs: mismatch in Font Size from image. Adjusted to match image.

diff --git a/Engine/Documentation/Source/Programming/Tutorials/UMG/5/UMG_5.INT.udn b/Engine/Documentation/Source/Programming/Tutorials/UMG/5/UMG_5.INT.udn
--- a/Engine/Documentation/Source/Programming/Tutorials/UMG/5/UMG_5.INT.udn
+++ b/Engine/Documentation/Source/Programming/Tutorials/UMG/5/UMG_5.INT.udn
@@ -31,7 +31,7 @@ Description:Build a second menu with new functionality.
 1.  The **Text Box** should be configured with the following values:
 	* Position is (325, 200). This leaves room for a **Text Block** placed to the left of the **Text Box**.
 	* Size is 250x40.
-	* Font Size (under the "Style" heading) is 9.
+	* Font Size (under the "Style" heading) is 20.
 	
 	[![](TextBoxDetails.png)](TextBoxDetails.png)
 
