#UE4doc #docs: sp: than

diff --git a/Engine/Documentation/Source/Shared/PhAT/PhAT.INT.udn b/Engine/Documentation/Source/Shared/PhAT/PhAT.INT.udn
--- a/Engine/Documentation/Source/Shared/PhAT/PhAT.INT.udn
+++ b/Engine/Documentation/Source/Shared/PhAT/PhAT.INT.udn
@@ -210,7 +210,7 @@ This is only available during Simulation.
 
 ![](PhAT_playAnim.png)
 
-If "Physics Blend" in the Anim category of the Physics Asset Tools options (deselect any Physics Bodies or Constraints) is set to anything less that 1.0, some amount of animation will be played on the Skeletal Mesh in the viewport of the Physics Asset Tool.
+If "Physics Blend" in the Anim category of the Physics Asset Tools options (deselect any Physics Bodies or Constraints) is set to anything less than 1.0, some amount of animation will be played on the Skeletal Mesh in the viewport of the Physics Asset Tool.
 
 [/EXCERPT:PlayAnimation]
 
