#UE4doc #docs: Spelling and grammar edits

diff --git a/Engine/Documentation/Source/Shared/Types/UAnimSequence/UAnimSequence.INT.udn b/Engine/Documentation/Source/Shared/Types/UAnimSequence/UAnimSequence.INT.udn
--- a/Engine/Documentation/Source/Shared/Types/UAnimSequence/UAnimSequence.INT.udn
+++ b/Engine/Documentation/Source/Shared/Types/UAnimSequence/UAnimSequence.INT.udn
@@ -13,18 +13,18 @@ Animations can be used as absolute or additive at any time. The system performs
 The additive animation calculates the delta of base animation and current animation and you can apply that difference to any base animation. 
 This allows you do create different permutation of animations by applying this delta to different base animation. 
 
-This requires two information
+This requires two pieces of information:
 
 1. Current animation : This animation to be converted
 2. Base animation : when you calculate delta, you need where it is based on, which changes how the delta is extracted
 
 ### Additive animation type
 
-This decides what kind of additve you want this to use. 
+This decides what kind of additive you want this to use. 
 
-1. No additive : This animation is not additive
-2. Local Space : This animation is additive and delta is calculated by local spase. This is most popular option of additive. This calculates different in local space, and apply that delta to local space. 
-3. Mesh Space : This animation is additve and delta will be applied in component space. This is an option when you want the additive to be applied to component space. This is required for aim offset because aim offset should be applied in component space
+1. No additive: This animation is not additive
+2. Local Space: This animation is additive and delta is calculated by local space. This is most popular option of additive. This calculates different in local space, and apply that delta to local space. 
+3. Mesh Space: This animation is additive and delta will be applied in component space. This is an option when you want the additive to be applied to component space. This is required for aim offset because aim offset should be applied in component space
 
 [/EXCERPT:AdditiveAnimType]
 
@@ -34,10 +34,10 @@ This decides what kind of additve you want this to use.
 
 This decides how this additive delta to be calculated. 
 
-1. None : No base pose
-2. Reference Pose : The mesh's default reference pose. To see this, in the viewport, select show->reference pose. 
-3. Selected animation scaled : Choose another animation below and scale whole animation when extract the delta
-4. Selected animation frame : Choose another animation below and only based on the frame number
+1. None: No base pose
+2. Reference Pose: The mesh's default reference pose. To see this, in the viewport, select show->reference pose. 
+3. Selected animation scaled: Choose another animation below and scale whole animation when extract the delta
+4. Selected animation frame: Choose another animation below and only based on the frame number
 
 ### Base Pose Animation
 
