#UE4doc #docs: grammar changes

diff --git a/Engine/Documentation/Source/Shared/Types/UBodySetup/UBodySetup.INT.udn b/Engine/Documentation/Source/Shared/Types/UBodySetup/UBodySetup.INT.udn
--- a/Engine/Documentation/Source/Shared/Types/UBodySetup/UBodySetup.INT.udn
+++ b/Engine/Documentation/Source/Shared/Types/UBodySetup/UBodySetup.INT.udn
@@ -86,7 +86,7 @@ Engine/Physics/PhAT/Properties#physicsbodyproperties
 
 By default, when we add collision to physics world, we add two shapes of collision geometry. We add simple shapes to query simple collision - such as player walking physics - and we add complex shapes to query complex collision - such as hitting bullet when precision is important. 
 
-Each trace function has "TraceComplex" flag, that will enable whether you'd like to trace complex or simple. 
+Each trace function has a "TraceComplex" flag that will enable whether you'd like to trace complex or simple. 
 
 ![](SimpleVSComplex.png)
 
