#UE4doc #docs: grammar and word choice changes.

diff --git a/Engine/Documentation/Source/Shared/Types/ULandscapeEditorObject/ULandscapeEditorObject.INT.udn b/Engine/Documentation/Source/Shared/Types/ULandscapeEditorObject/ULandscapeEditorObject.INT.udn
--- a/Engine/Documentation/Source/Shared/Types/ULandscapeEditorObject/ULandscapeEditorObject.INT.udn
+++ b/Engine/Documentation/Source/Shared/Types/ULandscapeEditorObject/ULandscapeEditorObject.INT.udn
@@ -4,7 +4,7 @@ Crumbs:
 Description:
 
 [EXCERPT:ManageMode]
-Manage Mode contains tools that will allow you to add a new Landscape, import/export a Landscape, add/remove Landscape components and manage Landscape streaming. Mange mode will also allow you to select individual landscape components and modify their properties in the level editors details panel.
+Manage Mode contains tools that will allow you to add a new Landscape, import/export a Landscape, add/remove Landscape components and manage Landscape streaming. Mange mode will also allow you to select individual landscape components and modify their properties in the level editor's details panel.
 [VAR:ToolTipFullLink]
 Engine/Landscape/Editing/ManageMode
 [/VAR]
@@ -33,5 +33,5 @@ Engine/Landscape/Creation
 [/EXCERPT:NewLandscape_ComponentCount]
 
 [EXCERPT:bUseClayBrush]
-The Clay Brush works by looking at the average direction of all of the faces that are under the mouse and then moving the vertices in the averaged direction. Because of this the Clay Brush is very good at filling in gaps and creating wholes in the Landscape. 
+The Clay Brush works by looking at the average direction of all of the faces that are under the mouse and then moving the vertices in the averaged direction. Because of this the Clay Brush is very good at filling in gaps and creating holes in the Landscape. 
 [/EXCERPT:bUseClayBrush]
