#UE4doc #docs: added a needed comma.

diff --git a/Engine/Documentation/Source/Shared/Types/ULevel/ULevel.INT.udn b/Engine/Documentation/Source/Shared/Types/ULevel/ULevel.INT.udn
--- a/Engine/Documentation/Source/Shared/Types/ULevel/ULevel.INT.udn
+++ b/Engine/Documentation/Source/Shared/Types/ULevel/ULevel.INT.udn
@@ -4,7 +4,7 @@ Crumbs:
 Description:Overview of the Level class
 
 [EXCERPT:ULevel]
-The Unreal Engine organizes Actors into what are called Levels. Each Level can contain a number of Actors of different types (meshes, lights, effects, landscape, volumes etc). Each level can be loaded and unloaded independently, and usually exists as a separate file on disk. Levels can overlap spatially, so you can use them as a way of organizing your world, to allow multiple people to work on different aspects at once. There is also a World Composition system that will automatically break up your world into multiple levels and load and unload them as you move around.
+The Unreal Engine organizes Actors into what are called Levels. Each Level can contain a number of Actors of different types (meshes, lights, effects, landscape, volumes, etc). Each level can be loaded and unloaded independently, and usually exists as a separate file on disk. Levels can overlap spatially, so you can use them as a way of organizing your world, to allow multiple people to work on different aspects at once. There is also a World Composition system that will automatically break up your world into multiple levels and load and unload them as you move around.
 [VAR:ToolTipFullLink]
 Engine/Levels
 [/VAR]
