#UE4doc #docs: duplicate word: the

diff --git a/Engine/Documentation/Source/Shared/Tutorials/inEditorTutorialPlacement/inEditorTutorialPlacement.INT.udn b/Engine/Documentation/Source/Shared/Tutorials/inEditorTutorialPlacement/inEditorTutorialPlacement.INT.udn
--- a/Engine/Documentation/Source/Shared/Tutorials/inEditorTutorialPlacement/inEditorTutorialPlacement.INT.udn
+++ b/Engine/Documentation/Source/Shared/Tutorials/inEditorTutorialPlacement/inEditorTutorialPlacement.INT.udn
@@ -9,7 +9,7 @@ Crumbs:
 Placing Actors
 [/VAR]
 
-In this tutorial, we will be going over the the various ways to place Actors into the level.
+In this tutorial, we will be going over the various ways to place Actors into the level.
 
 ![](Stage1.png)
 
