#UE4doc #docs: fixed a typo!

diff --git a/Engine/Documentation/Source/Shared/Types/UMaterialInstanceConstant/UMaterialInstanceConstant.INT.udn b/Engine/Documentation/Source/Shared/Types/UMaterialInstanceConstant/UMaterialInstanceConstant.INT.udn
--- a/Engine/Documentation/Source/Shared/Types/UMaterialInstanceConstant/UMaterialInstanceConstant.INT.udn
+++ b/Engine/Documentation/Source/Shared/Types/UMaterialInstanceConstant/UMaterialInstanceConstant.INT.udn
@@ -5,7 +5,7 @@ Description:Overview of the MaterialInstanceConstant class
 
 [EXCERPT:UMaterialInstanceConstant]
 
-Material Instances may be used to change the apperance of a material without incurring an expensive recompilation of the material.  General modification of the material cannot be supported without recompilation, so the instances are limited to changing the values of predefined material parameters.  The parameters are statically defined in the compiled material by a unique name, type, and default value. 
+Material Instances may be used to change the appearance of a material without incurring an expensive recompilation of the material.  General modification of the material cannot be supported without recompilation, so the instances are limited to changing the values of predefined material parameters.  The parameters are statically defined in the compiled material by a unique name, type, and default value. 
 
 [VAR:ToolTipFullLink]
 Engine/Rendering/Materials/MaterialInstances
