#UE4doc #docs: fixed a comma error.

diff --git a/Engine/Documentation/Source/Shared/Types/UObjectLibrary/UObjectLibrary.INT.udn b/Engine/Documentation/Source/Shared/Types/UObjectLibrary/UObjectLibrary.INT.udn
--- a/Engine/Documentation/Source/Shared/Types/UObjectLibrary/UObjectLibrary.INT.udn
+++ b/Engine/Documentation/Source/Shared/Types/UObjectLibrary/UObjectLibrary.INT.udn
@@ -5,7 +5,7 @@ Description:Overview of the ObjectLibrary class
 
 [EXCERPT:UObjectLibrary]
 
-An ObjectLibrary is an object that contains a list of either loaded objects or FAssetData for unloaded objects, that inherit from a shared base class.
+An ObjectLibrary is an object that contains a list of either loaded objects or FAssetData for unloaded objects that inherit from a shared base class.
 
 [VAR:ToolTipFullLink]
 Programming/Assets/AsyncLoading#theassetregistryandobjectlibraries
