#UE4doc #docs: fixed a typo, saved the world all that good stuff.

diff --git a/Engine/Documentation/Source/Shared/Types/UPaperSprite/UPaperSprite.INT.udn b/Engine/Documentation/Source/Shared/Types/UPaperSprite/UPaperSprite.INT.udn
--- a/Engine/Documentation/Source/Shared/Types/UPaperSprite/UPaperSprite.INT.udn
+++ b/Engine/Documentation/Source/Shared/Types/UPaperSprite/UPaperSprite.INT.udn
@@ -5,7 +5,7 @@ Description:Overview of the PaperSprite class
 
 [EXCERPT:UPaperSprite]
 
-Stores the data necessary to rnder a single 2D sprite (from a region of a texture).  Can also contain collision shapes for the sprite.
+Stores the data necessary to render a single 2D sprite (from a region of a texture).  Can also contain collision shapes for the sprite.
 
 [VAR:ToolTipFullLink]
 Engine/Paper2D/SpriteEditor
