#UE4doc #docs: sp: react, this function's, differences

diff --git a/Engine/Documentation/Source/Shared/Tutorials/TemplateTutorials/RollingCode/RollingCode.INT.udn b/Engine/Documentation/Source/Shared/Tutorials/TemplateTutorials/RollingCode/RollingCode.INT.udn
--- a/Engine/Documentation/Source/Shared/Tutorials/TemplateTutorials/RollingCode/RollingCode.INT.udn
+++ b/Engine/Documentation/Source/Shared/Tutorials/TemplateTutorials/RollingCode/RollingCode.INT.udn
@@ -81,7 +81,7 @@ The Rolling Game Pawn - Summary
 
 This class is the implementation of the pawn. In this case the pawn is a ball represented by a sphere mesh.
 
-The class only contains 6 functions. Lets look at these now.
+The class only contains 6 functions. Let's look at these now.
 
 The MarbleBall class that has been created derives from the engines' Pawn class. For more details on this see 
 [Pawn Overview](DOCLINK:Gameplay/Framework/Pawn)
@@ -104,7 +104,7 @@ The first thing we do in the constructor is find the Mesh we wish to use for our
 [First thing we do is find our mesh,](CODELINK:Private/[PROJECT]Ball.cpp,13,13)
 The code creates a simple struct that will find pointer to the mesh we specified. In this case
 [BallMesh](ASSETLINK:SELECT,BallMesh)
-This is a helper that will attempt to find an Static Mesh within our application package. In the struct as you can see we pass the name of the mesh we wish to use.
+This is a helper that will attempt to find a Static Mesh within our application package. In the struct as you can see we pass the name of the mesh we wish to use.
 We then create an instance of this struct. If the mesh exists we can then use this instance to access the pointer to our required mesh.
 
 Next we create a Static Mesh component and give it a name
@@ -159,8 +159,8 @@ This is an overridden function. It is passed a UInputComponent pointer. We use t
 We pass 3 parameters to this:
 "MoveRight"
 This is the name of the axis binding. This relates to input bindings that are defined in the editor. We shall cover how this is done later.
-this
-the this pointer that points to our object
+This 
+is the pointer that points to our object
 &AMarbleBall::MoveRight
 This is the function we wish to have called to handle input from the MoveRight binding.
 
@@ -168,7 +168,7 @@ The next line is pretty much the same, except that it binds the "MoveForward" bi
 The 2 functions we are passing here must of course match the delegate that is defined. For BindAxis this is a void function taking a single float. The float represents an axis value in the range -1 to 1
 
 [The last binding here is an action.](CODELINK:Private/[PROJECT]Ball.cpp,59,0)
-As you can see this take 4 parameters.
+As you can see, this take 4 parameters.
 The first is again the name of the binding.
 The next is the action we wish to bind - in this case pressed. The final 2 parameters are much the same as they were for the axis binding the only difference being that the delegate for this takes no parameters. 
 Jumping is just a simple key event
@@ -185,7 +185,7 @@ The Rolling Game Pawn - Movement functions
 There are 2 movement functions here. One for each axis.
 [MoveRight](CODELINK:Private/[PROJECT]Ball.cpp,62,0)
 This takes a float which will be in the range -1 to 1. This will also come in with ZERO as the parameter when the stick is unmoved.
-The function calculates a torque vector on the X axis based on the RollTorque member variable and multiplies this by the amount of movement on the stick. It then passes this to the ball component as torque using the AddTorque function. When the ball component is updated it will updated the components torque and recact accordingly.
+The function calculates a torque vector on the X axis based on the RollTorque member variable and multiplies this by the amount of movement on the stick. It then passes this to the ball component as torque using the AddTorque function. When the ball component is updated it will updated the components torque and react accordingly.
 [MoveForward](CODELINK:Private/[PROJECT]Ball.cpp,68,0)
 This does exactly the same as MoveRight except it creates a torque vector for the Y axis.
 
@@ -202,7 +202,7 @@ The Rolling Game Pawn - Impact
 [/VAR]
 
 [ReceiveHit](CODELINK:Private/[PROJECT]Ball.cpp,84,0)
-This is the final function in the file. It is called when the object collides in the world. It takes 8 parameters. There is a detailed explanation this functions parameters in the Actor API documentation. However we do not need to use any of these for the purposes of this template so we shall not go into them here.
+This is the final function in the file. It is called when the object collides in the world. It takes 8 parameters. There is a detailed explanation for this function's parameters in the Actor API documentation. However we do not need to use any of these for the purposes of this template so we shall not go into them here.
 So the first thing that we do is call up to the base class.
 Finally we simply clear the flag that we set when pressed the jump button. We could of course do something far more elaborate on impact and determine just what it is we have hit before deeming the jump to have ended but to keep the templates clear we simply treat any impact is the end of the jump.
 
@@ -236,7 +236,7 @@ If you now click on the arrow next to axis mappings a list of mappings will appe
 ![](AxisMap.png)
 
 So when either GamePadLeft X (which is the left stick on the controller), Left or Right is moved/pushed it will generate trigger the event resulting in the delegate being called.
-There are some subtle difference in these - because one of them is a controller binding and the other 2 are key bindings. 
+There are some subtle differences in these - because one of them is a controller binding and the other 2 are key bindings. 
 The controller binding has a scale of 1. This is fine, this will allow the stick to behave as we would expect. If we were to change this to -1 it would have the effect of flipping the axis. 
 The next 2 bindings are keys - these do not have a range of -1 to 1 like an axis stick on a controller - they return 0-1 (as they are keyboard keys not analogue buttons as found on a game controller they will in-fact be 0 or 1). The scale on the binding for Left means that it will convert the 1 to -1 - thus behaving as if the stick were moved left.
 
