#UE4doc #docs: sp: preview

diff --git a/Engine/Documentation/Source/Shared/Tutorials/TestCase/TestCase.INT.udn b/Engine/Documentation/Source/Shared/Tutorials/TestCase/TestCase.INT.udn
--- a/Engine/Documentation/Source/Shared/Tutorials/TestCase/TestCase.INT.udn
+++ b/Engine/Documentation/Source/Shared/Tutorials/TestCase/TestCase.INT.udn
@@ -63,7 +63,7 @@ Description: This is a test
 
 ![](TutorialTestImage06.png)
 
-1. Above the Components option you will find the Static Mesh options. Once you have found it, click on the bar next to the prview window.
+1. Above the Components option you will find the Static Mesh options. Once you have found it, click on the bar next to the preview window.
 1. In the new window that has appeared, type EditorCube in the search bar. Once it appears, add it to the StaticMeshComponent.
 
 [/EXCERPT:Six]
