#UE4doc #docs: reworked the sentence: Only for "Subsurface Profile" materials, is used during Screenspace Subsurface Scattering do not change at runtime.

diff --git a/Engine/Documentation/Source/Shared/Types/USubsurfaceProfile/USubsurfaceProfile.INT.udn b/Engine/Documentation/Source/Shared/Types/USubsurfaceProfile/USubsurfaceProfile.INT.udn
--- a/Engine/Documentation/Source/Shared/Types/USubsurfaceProfile/USubsurfaceProfile.INT.udn
+++ b/Engine/Documentation/Source/Shared/Types/USubsurfaceProfile/USubsurfaceProfile.INT.udn
@@ -5,7 +5,7 @@ Description:Overview of the SubsurfaceProfile class
 
 [EXCERPT:USubsurfaceProfile]
 
-Subsurface Scattering profile asset, can be specified at the material.  Only for "Subsurface Profile" materials, is used during Screenspace Subsurface Scattering do not change at runtime.  All properties in here are per material - texture like variations need to come from properties that are in the GBuffer.
+Subsurface Scattering profile asset, can be specified at the material.  Only "Subsurface Profile" materials are used during Screenspace Subsurface Scattering. These materials do not change at runtime.  All properties in here are per material - texture like variations need to come from properties that are in the GBuffer.
 
 [VAR:ToolTipFullLink]
 Engine\Rendering\Materials\LightingModels\SubSurface\SSSProfileLightingModel
