#UE4doc #docs: added grammar to help sentence flow better.

diff --git a/Engine/Documentation/Source/Shared/Types/UWidgetBlueprint/UWidgetBlueprint.INT.udn b/Engine/Documentation/Source/Shared/Types/UWidgetBlueprint/UWidgetBlueprint.INT.udn
--- a/Engine/Documentation/Source/Shared/Types/UWidgetBlueprint/UWidgetBlueprint.INT.udn
+++ b/Engine/Documentation/Source/Shared/Types/UWidgetBlueprint/UWidgetBlueprint.INT.udn
@@ -5,7 +5,7 @@ Description:Overview of the WidgetBlueprint class
 
 [EXCERPT:UWidgetBlueprint]
 
-The widget blueprint enables extending UUserWidget the user extensible UWidget.
+The widget blueprint enables extending UUserWidget, the user extensible UWidget.
 
 [VAR:ToolTipFullLink]
 
