#UE4 Docs: Adding a new option to the Proc Foliage step 3 QS that ensures that foliage meshes will not be spawned sideways.

diff --git a/Engine/Documentation/Source/Engine/OpenWorldTools/ProceduralFoliage/QuickStart/3/PFQS_3.INT.udn b/Engine/Documentation/Source/Engine/OpenWorldTools/ProceduralFoliage/QuickStart/3/PFQS_3.INT.udn
--- a/Engine/Documentation/Source/Engine/OpenWorldTools/ProceduralFoliage/QuickStart/3/PFQS_3.INT.udn
+++ b/Engine/Documentation/Source/Engine/OpenWorldTools/ProceduralFoliage/QuickStart/3/PFQS_3.INT.udn
@@ -29,7 +29,7 @@ We will continue to work with the **PFT_00** level that was used in the last ste
 
 1. Open up the **Tree_00** Foliage Type.
 
-1. Expand the **Placement** section and make sure that both **Align to Normal** and **Random Yaw** are enabled.
+1. Expand the **Placement** section and make sure that both **Align to Normal** and **Random Yaw** are enabled and set **Align Max Angle** to **3.0** to ensure that Static Meshes do not turn sideways when spawned.
 
 	![](T_Placement_Options.png)
 
