#UE4doc #docs: updated the year to match this year.

diff --git a/Engine/Documentation/Source/Support/Builds/Builds.INT.udn b/Engine/Documentation/Source/Support/Builds/Builds.INT.udn
--- a/Engine/Documentation/Source/Support/Builds/Builds.INT.udn
+++ b/Engine/Documentation/Source/Support/Builds/Builds.INT.udn
@@ -8,7 +8,7 @@ Description:Release notes for Unreal Engine builds.
 	[PARAM:icon]
 	[/PARAM]
 	[PARAM:title]
-	2014
+	2015
 	[/PARAM]
 	[PARAM:description]
 	[/PARAM]
