#UE4doc #docs: duplicate word: is

diff --git a/Engine/Documentation/Source/Support/UDNGettingStarted/SupportQuestions/SupportQuestions.INT.udn b/Engine/Documentation/Source/Support/UDNGettingStarted/SupportQuestions/SupportQuestions.INT.udn
--- a/Engine/Documentation/Source/Support/UDNGettingStarted/SupportQuestions/SupportQuestions.INT.udn
+++ b/Engine/Documentation/Source/Support/UDNGettingStarted/SupportQuestions/SupportQuestions.INT.udn
@@ -41,7 +41,7 @@ When typing the body of your question, this box will update to give you useful t
 
 ## Tags
 
-When asking a question, is is required that you add tags to help narrow down the systems related to your issue. You may add these manually by typing them into the tag line, or click on popular tag buttons to choose from some common ones. You may have up to 5 tags.
+When asking a question, is required that you add tags to help narrow down the systems related to your issue. You may add these manually by typing them into the tag line, or click on popular tag buttons to choose from some common ones. You may have up to 5 tags.
 
 ![](tags.png)
 
