#UE4 Docs: Adding a little blurb about material compile errors if you fail to connect anything to the grass input node.

diff --git a/Engine/Documentation/Source/Engine/OpenWorldTools/Grass/QuickStart/3/GTQS_3.INT.udn b/Engine/Documentation/Source/Engine/OpenWorldTools/Grass/QuickStart/3/GTQS_3.INT.udn
--- a/Engine/Documentation/Source/Engine/OpenWorldTools/Grass/QuickStart/3/GTQS_3.INT.udn
+++ b/Engine/Documentation/Source/Engine/OpenWorldTools/Grass/QuickStart/3/GTQS_3.INT.udn
@@ -84,6 +84,13 @@ If you would like to have a more in-depth look at how the Landscape terrain work
 
 	![](T_LLW_Setup.png)
 
+	[region:note]
+	When using the Grass Output the future make sure that the Grass output Material Expression is the last item you add to your Material Graph.
+	If you try and compile or save your Material without anything going into the Grass output node the compile will fail with the following error message in the **Stats** window.
+
+	![](T_Grass_Material_Warning.png)
+	[/region]
+
 1. When completed you should have a Material that looks like the following. As do not forget to press the **Apply** and **Save** buttons to compile and save the Material.
 
 	![](T_Final_Material.png)
