TypoCorrection

diff --git a/Engine/Documentation/Source/Engine/OpenWorldTools/ProceduralFoliage/QuickStart/1/PFQS_1.INT.udn b/Engine/Documentation/Source/Engine/OpenWorldTools/ProceduralFoliage/QuickStart/1/PFQS_1.INT.udn
--- a/Engine/Documentation/Source/Engine/OpenWorldTools/ProceduralFoliage/QuickStart/1/PFQS_1.INT.udn
+++ b/Engine/Documentation/Source/Engine/OpenWorldTools/ProceduralFoliage/QuickStart/1/PFQS_1.INT.udn
@@ -1,4 +1,4 @@
-Availability:Docs
+﻿Availability:Docs
 Title: 1. Creating Foliage Type Actors
 Crumbs:
 Description: Creating and adding all the required Actors.
@@ -35,7 +35,7 @@ In this step we will create a new Level, Landscape Terrain and all of the Actors
 	![](T_Create_Landscape.gif)(convert:false)
 
 	[region:note]
-	If you are not familiar with how Landscape Terrain works or would like to know more about how it works please check out the [Landscape Terrain documents](https://docs.unrealengine.com/latest/INT/Engine/Landscape/index.html) for more information.
+	If you are not familiar with how Landscape Terrain works or would like to know more about how it works please check out [](Engine/Landscape) for more information.
 	[/region]
 
 1. Exit the Landscape tool by clicking on the **Place** icon in the **Modes Panel**.
