#UE4doc #docs: sp: session, its, through,

diff --git a/Engine/Documentation/Source/Programming/Online/Interfaces/Session/Session.INT.udn b/Engine/Documentation/Source/Programming/Online/Interfaces/Session/Session.INT.udn
--- a/Engine/Documentation/Source/Programming/Online/Interfaces/Session/Session.INT.udn
+++ b/Engine/Documentation/Source/Programming/Online/Interfaces/Session/Session.INT.udn
@@ -64,7 +64,7 @@ the characteristics of the session. In the base implementation, these are things
 * Etc.
 
 
-Using the online game lobby example, each of those games is a session and has it's own session settings. 
+Using the online game lobby example, each of those games is a session and has its own session settings. 
 For example, some sessions may be player versus player (PvP) while others are cooperative multiplayer (Co-Op). 
 Different sessions could be playing different maps or playlists, require different numbers of players, etc.
 
@@ -100,7 +100,7 @@ When a session ends and is no longer needed, the session is destroyed using the
 function. When the destruction opteration has completed, the `OnDestroySessionComplete` delegate is fired enabling 
 you to perform cleanup operations.
 
-## Matchmaking - Finding sessons
+## Matchmaking - Finding sessions
 
 The Online Subsystem provides the basic building blocks necessary to match players up to active sessions. It does 
 not provide any specific type of matchmaking built-in with the base implementation. However, implementations on 
@@ -126,7 +126,7 @@ fired. From this delegate, you can iterate through the search results.
 ### Cloud-Based Matchmaking
 
 Cloud-based matchmaking refers to built-in matchmaking services that are available, and are generally platform-specific. 
-An example of this type of service is the TrueSkill system available though Microsoft's Xbox Live service. 
+An example of this type of service is the TrueSkill system available through Microsoft's Xbox Live service. 
 
 To initiate matchmaking on platforms that support it, you call `IOnlineSession::Startmatchmaking()` and pass 
 it the controller number of the player to matchmake for, the session name, session settings to use if creating a new 
