Changed file name

diff --git a/Engine/Documentation/Source/Engine/OpenWorldTools/Overview/OpenWorldToolsOverview.INT.udn b/Engine/Documentation/Source/Engine/OpenWorldTools/Overview/OpenWorldToolsOverview.INT.udn
--- a/Engine/Documentation/Source/Engine/OpenWorldTools/Overview/OpenWorldToolsOverview.INT.udn
+++ b/Engine/Documentation/Source/Engine/OpenWorldTools/Overview/OpenWorldToolsOverview.INT.udn
@@ -30,8 +30,8 @@ This enables you to quickly generate a natural looking forest comprised of many
 
 When talking about the Open World tools we are actually referring to the two different tool sets inside of UE4. 
 The tool set that you use depends on what you want to accomplish. 
-For example if you want to make an expansive pine and birch tree forest then you're going to want to use the **Procedural Foliage Tools**. 
-If you want to densely populate the Landscape Terrain used in that forest with grass, flowers, and small rocks, then you would use the **Grass Tools**.
+For example if you want to make an expansive pine and birch tree forest then you're going to want to use the **Procedural Foliage Tool**. 
+If you want to densely populate the Landscape Terrain used in that forest with grass, flowers, and small rocks, then you would use the **Grass Tool**.
 
 ### Procedural Foliage Tool 
 The Procedural Foliage Tool is used for simulating huge forests that are filled with lots of different kinds of trees and bushes. 
@@ -78,10 +78,10 @@ The chart below gives a very high level explanation of the four main sets of pro
 
 ### Grass Tool
 [region:note]
-Currently the Grass Tools will only work in conjunction with the Landscape Terrain system. 
+Currently the Grass Tool will only work in conjunction with the Landscape Terrain system. 
 [/region]
-The Grass Tools are used to densely cover the Landscape Terrain with things like grass, flowers, and small rocks. 
-You can use the Grass Tools to define a Static Mesh that will spawn when a certain Landscape Layer is painted on the Landscape Terrain. 
+The Grass Tool is used to densely cover the Landscape Terrain with things like grass, flowers, and small rocks. 
+You can use the Grass Tool to define a Static Mesh that will spawn when a certain Landscape Layer is painted on the Landscape Terrain. 
 This tool does not use any kind of simulation like the Procedural Foliage Spawner does. Instead, it simply fills the specified area at the specified density with the specified Static Mesh. 
 To learn more about how to use the Grass Tool, check out the Grass Tool Quick Start that is linked below.
 
