Minor edits

diff --git a/Engine/Documentation/Source/Engine/OpenWorldTools/Grass/QuickStart/GrassToolsQuickStart.INT.udn b/Engine/Documentation/Source/Engine/OpenWorldTools/Grass/QuickStart/GrassToolsQuickStart.INT.udn
--- a/Engine/Documentation/Source/Engine/OpenWorldTools/Grass/QuickStart/GrassToolsQuickStart.INT.udn
+++ b/Engine/Documentation/Source/Engine/OpenWorldTools/Grass/QuickStart/GrassToolsQuickStart.INT.udn
@@ -9,9 +9,9 @@ Version: 4.8
 ## Overview
 
 As with anything in life using the right tool for the job is often the determining factor between success and failure and the same can be said when talking about the Open World tools inside of UE4. 
-For example the Procedural Foliage tool is great at simulating a huge forests that are composed of many different species of trees and bushes. 
-However if you try and use the Procedural Foliage tool make the Landscape Terrain look like it was densely covered in grass or flowers you will quickly find out that getting the results you want can be very time consuming. 
-To address this issue and give you the ability to densely spawn grass and flowers the Open World Tools have an additional set of tools called the **Grass Tools**.
+For example the Procedural Foliage tool is great at simulating huge forests that are composed of many different species of trees and bushes. 
+However if you try and use the Procedural Foliage tool to make the Landscape Terrain look like it was densely covered in grass or flowers you will quickly find out that getting the results you want can be very time consuming. 
+To address this issue and give you the ability to densely spawn grass and flowers the Open World Tools have an optional tool called the **Grass Tool**.
 The Grass Tool allows you to define a Static Mesh to spawn wherever a certain Texture is painted on the Landscape terrain.
 In the following How - To we will cover all you need to know about using the Grass Tool in your levels.
 
@@ -24,13 +24,13 @@ Currently the Grass system will only work with the Landscape Terrain Actor. Tryi
 ## Prerequisites
 
 You will also need to download the **Open World Demo Collection** content pack from the **Unreal Engine Marketplace** as some of the content from the collection will be used in the following Quick Start.
-Once the Open World Demo Collection is downloaded add it to the project that you are using to follow along with this Quick Start by doing the following.
+Once the Open World Demo Collection is downloaded, add it to the project that you are using to follow along with this Quick Start by doing the following:
 
-1. From the Epic Games Launcher in the Marketplace locate and then download the **Open World Demo Collection**.
+1. From the Epic Games Launcher in the Marketplace, locate and then download the **Open World Demo Collection**.
 	
 	![](Engine\OpenWorldTools\T_OWT_OWDC.png)
 
-1. Go the the **Library** section of the launcher and then under **Vault** section locate the Open World Demo Collection . 
+1. Go the the **Library** section of the launcher and under **Vault** section locate the Open World Demo Collection . 
 	
 	![](Engine\OpenWorldTools\T_OWT_Add_Content_00.png)
 
@@ -38,7 +38,7 @@ Once the Open World Demo Collection is downloaded add it to the project that you
 	
 	![](Engine\OpenWorldTools\T_OWT_Add_To_Project.png)
 
-1. You will then be shown a list of projects that you can add this to, select the project that you are using to follow along with this Quick Start and then press the **Add to Project** button
+1. You will be shown a list of projects that you can add this to. Select the project that you are using to follow along with this Quick Start and then press the **Add to Project** button
 
 	![](Engine\OpenWorldTools\T_OWT_ATP.png)
 
