minor edits

diff --git a/Engine/Documentation/Source/Engine/OpenWorldTools/Grass/QuickStart/3/GTQS_3.INT.udn b/Engine/Documentation/Source/Engine/OpenWorldTools/Grass/QuickStart/3/GTQS_3.INT.udn
--- a/Engine/Documentation/Source/Engine/OpenWorldTools/Grass/QuickStart/3/GTQS_3.INT.udn
+++ b/Engine/Documentation/Source/Engine/OpenWorldTools/Grass/QuickStart/3/GTQS_3.INT.udn
@@ -1,4 +1,4 @@
-Availability:Docs
+Availability:Public
 Title:3. Landscape Materials and the Grass Tool
 Crumbs:
 Description: How to setup a Material that will work with both the Grass Tool and the Landscape terrain.
@@ -9,7 +9,7 @@ Version: 4.8
 [VAR:Steps]
 [OBJECT:Navigation]
 	[PARAM:previous]
-	[Previous Step](Engine\OpenWorldTools\Grass\QuickStart\1 "%Engine\OpenWorldTools\ProceduralFoliage\QuickStart\1:title%") 
+	[Previous Step](Engine\OpenWorldTools\Grass\QuickStart\2 "%Engine\OpenWorldTools\ProceduralFoliage\QuickStart\2:title%") 
 	[/PARAM]
 	[PARAM:current]
 	[/PARAM]
@@ -84,14 +84,7 @@ If you would like to have a more in-depth look at how the Landscape terrain work
 
 	![](T_LLW_Setup.png)
 
-	[region:note]
-	When using the Grass Output future make sure that the Grass output Material Expression is the last item you add to your Material Graph.
-	If you try and compile or save your Material without anything going into the Grass output node the compile will fail with the following error message in the **Stats** window.
-
-	![](T_Grass_Material_Warning.png)
-	[/region]
-
-1. When completed you should have a Material that looks like the following. As do not forget to press the **Apply** and **Save** buttons to compile and save the Material.
+1. When completed you should have a Material that looks like the following. As always, do not forget to press the **Apply** and **Save** buttons to compile and save the Material.
 
 	![](T_Final_Material.png)
 
diff --git a/Engine/Documentation/Source/Engine/OpenWorldTools/Grass/QuickStart/4/GTQS_4.INT.udn b/Engine/Documentation/Source/Engine/OpenWorldTools/Grass/QuickStart/4/GTQS_4.INT.udn
--- a/Engine/Documentation/Source/Engine/OpenWorldTools/Grass/QuickStart/4/GTQS_4.INT.udn
+++ b/Engine/Documentation/Source/Engine/OpenWorldTools/Grass/QuickStart/4/GTQS_4.INT.udn
@@ -1,4 +1,4 @@
-Availability:Docs
+Availability:Public
 Title:4. Using the Grass Tools
 Crumbs:
 Description: How to use the Grass Tools in .
@@ -28,7 +28,7 @@ In order to see the Grass system in action we need to apply the Material that wa
 In the following section we will go over how to apply our Material to the Landscape Terrain and then how to use the Landscape Painting Tools to define areas where grass should spawn.
 We will continue to work with the **GT_00** level that was created in the last step .
 
-1. First select the Landscape Actor in that was placed in the level by **clicking** on it with the in the viewport.
+1. First select the Landscape Actor that was placed in the level by **clicking** on it with the in the viewport.
 
 	![](T_Selected_Landscape.png)
 
diff --git a/Engine/Documentation/Source/Engine/OpenWorldTools/Grass/QuickStart/5/GTQS_5.INT.udn b/Engine/Documentation/Source/Engine/OpenWorldTools/Grass/QuickStart/5/GTQS_5.INT.udn
--- a/Engine/Documentation/Source/Engine/OpenWorldTools/Grass/QuickStart/5/GTQS_5.INT.udn
+++ b/Engine/Documentation/Source/Engine/OpenWorldTools/Grass/QuickStart/5/GTQS_5.INT.udn
@@ -1,4 +1,4 @@
-Availability:Docs
+Availability:Public
 Title: 5. On Your Own
 Crumbs:
 Description: What to do next with the Grass Tool.
@@ -9,7 +9,7 @@ Version: 4.8
 [VAR:Steps]
 [OBJECT:Navigation]
 	[PARAM:previous]
-	[Previous Step](Engine\OpenWorldTools\Grass\QuickStart\4 "%Engine\OpenWorldTools\ProceduralFoliage\QuickStart\3:title%") 
+	[Previous Step](Engine\OpenWorldTools\Grass\QuickStart\4 "%Engine\OpenWorldTools\ProceduralFoliage\QuickStart\4:title%") 
 	[/PARAM]
 	[PARAM:current]
 	[/PARAM]
@@ -23,14 +23,13 @@ Version: 4.8
 
 %Steps%
 
-Now that you have seen that power the Grass tool offers, try using the tools listed below in conjunction with what you just learned about the 
-Grass tool to try and make a level that looks like the following image.
+Now that you have seen the power that the Grass tool offers, try using the tools listed below in conjunction with what you have just learned to make a level that looks like the following image:
 
 ![](Engine\OpenWorldTools\ProceduralFoliage\QuickStart\6\T_On_Your_Own.png)
 
 * Use the **[Procedural Foliage Tool](Engine\OpenWorldTools\ProceduralFoliage\QuickStart)** to make the Landscape look like it is densely covered in grass, flowers and bushes.
 
-* Define the look and feel of the Landscape terrain using the **[Landscape Sculpt](Engine/Landscape/Editing/SculptMode)** tools add things like hills, mountains and lakes.
+* Define the look and feel of the Landscape terrain using the **[Landscape Sculpt](Engine/Landscape/Editing/SculptMode)** tools to add things like hills, mountains and lakes.
 
 * Give the surface of the Landscape more visual variety and detail by creating a **[Landscape Material](Engine/Landscape/Materials)** that has multiple Textures that can be painted on the Landscape terrain.
 
