#UE4doc #docs: Made public

diff --git a/Engine/Documentation/Source/Engine/OpenWorldTools/Grass/QuickStart/1/GTQS_1.INT.udn b/Engine/Documentation/Source/Engine/OpenWorldTools/Grass/QuickStart/1/GTQS_1.INT.udn
--- a/Engine/Documentation/Source/Engine/OpenWorldTools/Grass/QuickStart/1/GTQS_1.INT.udn
+++ b/Engine/Documentation/Source/Engine/OpenWorldTools/Grass/QuickStart/1/GTQS_1.INT.udn
@@ -1,4 +1,4 @@
-Availability:Docs
+Availability:Public
 Title: 1. Initial Level Setup
 Crumbs:
 Description: How To set up the Initial Actors the Grass Tool needs to work.
diff --git a/Engine/Documentation/Source/Engine/OpenWorldTools/Grass/QuickStart/2/GTQS_2.INT.udn b/Engine/Documentation/Source/Engine/OpenWorldTools/Grass/QuickStart/2/GTQS_2.INT.udn
--- a/Engine/Documentation/Source/Engine/OpenWorldTools/Grass/QuickStart/2/GTQS_2.INT.udn
+++ b/Engine/Documentation/Source/Engine/OpenWorldTools/Grass/QuickStart/2/GTQS_2.INT.udn
@@ -1,4 +1,4 @@
-Availability:Docs
+Availability:Public
 Title:2. Grass Tool Actor Creation and Setup
 Crumbs:
 Description: How to create and setup the required Actors for the Grass Tool
diff --git a/Engine/Documentation/Source/Engine/OpenWorldTools/OpenWorldTools.INT.udn b/Engine/Documentation/Source/Engine/OpenWorldTools/OpenWorldTools.INT.udn
--- a/Engine/Documentation/Source/Engine/OpenWorldTools/OpenWorldTools.INT.udn
+++ b/Engine/Documentation/Source/Engine/OpenWorldTools/OpenWorldTools.INT.udn
@@ -1,4 +1,4 @@
-Availability:Docs
+Availability:Public
 Title: Open World Tools
 Crumbs:
 Description: Using the Open World Tools
diff --git a/Engine/Documentation/Source/Engine/OpenWorldTools/Overview/OpenWorldToolsOverview.INT.udn b/Engine/Documentation/Source/Engine/OpenWorldTools/Overview/OpenWorldToolsOverview.INT.udn
--- a/Engine/Documentation/Source/Engine/OpenWorldTools/Overview/OpenWorldToolsOverview.INT.udn
+++ b/Engine/Documentation/Source/Engine/OpenWorldTools/Overview/OpenWorldToolsOverview.INT.udn
@@ -1,4 +1,4 @@
-Availability:Docs
+Availability:Public
 Title: Open World Tools Overview
 Crumbs:
 Description: A system for procedurally placing Static Meshes in your levels.
diff --git a/Engine/Documentation/Source/Engine/OpenWorldTools/ProceduralFoliage/QuickStart/1/PFQS_1.INT.udn b/Engine/Documentation/Source/Engine/OpenWorldTools/ProceduralFoliage/QuickStart/1/PFQS_1.INT.udn
--- a/Engine/Documentation/Source/Engine/OpenWorldTools/ProceduralFoliage/QuickStart/1/PFQS_1.INT.udn
+++ b/Engine/Documentation/Source/Engine/OpenWorldTools/ProceduralFoliage/QuickStart/1/PFQS_1.INT.udn
@@ -1,4 +1,4 @@
-﻿Availability:Docs
+Availability:Public
 Title: 1. Creating Foliage Type Actors
 Crumbs:
 Description: Creating and adding all the required Actors.
diff --git a/Engine/Documentation/Source/Engine/OpenWorldTools/ProceduralFoliage/QuickStart/2/PFQS_2.INT.udn b/Engine/Documentation/Source/Engine/OpenWorldTools/ProceduralFoliage/QuickStart/2/PFQS_2.INT.udn
--- a/Engine/Documentation/Source/Engine/OpenWorldTools/ProceduralFoliage/QuickStart/2/PFQS_2.INT.udn
+++ b/Engine/Documentation/Source/Engine/OpenWorldTools/ProceduralFoliage/QuickStart/2/PFQS_2.INT.udn
@@ -1,4 +1,4 @@
-Availability:Docs
+Availability:Public
 Title: 2. Telling the Spawner what to Spawn
 Crumbs:
 Description:Setting up the Procedural Foliage Spawner to spawn foliage
diff --git a/Engine/Documentation/Source/Engine/OpenWorldTools/ProceduralFoliage/QuickStart/3/PFQS_3.INT.udn b/Engine/Documentation/Source/Engine/OpenWorldTools/ProceduralFoliage/QuickStart/3/PFQS_3.INT.udn
--- a/Engine/Documentation/Source/Engine/OpenWorldTools/ProceduralFoliage/QuickStart/3/PFQS_3.INT.udn
+++ b/Engine/Documentation/Source/Engine/OpenWorldTools/ProceduralFoliage/QuickStart/3/PFQS_3.INT.udn
@@ -1,4 +1,4 @@
-Availability:Docs
+Availability:Public
 Title: 3. Tweaking Foliage Type Properties
 Crumbs:
 Description:In this step we tweak the settings of the Procedural Foliage Actors.
@@ -29,7 +29,7 @@ We will continue to work with the **PFT_00** level that was used in the last ste
 
 1. Open up the **Tree_00** Foliage Type.
 
-1. Expand the **Placement** section and make sure that both **Align to Normal** and **Random Yaw** are enabled and set **Align Max Angle** to **3.0** to ensure that Static Meshes do not turn sideways when spawned.
+1. Expand the **Placement** section and make sure that both **Align to Normal** and **Random Yaw** are enabled.
 
 	![](T_Placement_Options.png)
 
diff --git a/Engine/Documentation/Source/Engine/OpenWorldTools/ProceduralFoliage/QuickStart/6/PFQS_6.INT.udn b/Engine/Documentation/Source/Engine/OpenWorldTools/ProceduralFoliage/QuickStart/6/PFQS_6.INT.udn
--- a/Engine/Documentation/Source/Engine/OpenWorldTools/ProceduralFoliage/QuickStart/6/PFQS_6.INT.udn
+++ b/Engine/Documentation/Source/Engine/OpenWorldTools/ProceduralFoliage/QuickStart/6/PFQS_6.INT.udn
@@ -1,4 +1,4 @@
-Availability:Docs
+Availability:Public
 Title: 6. On your Own
 Crumbs:
 Description: What to do next with the Procedural Foliage tool.
