TypoCorrection

diff --git a/Engine/Documentation/Source/Engine/OpenWorldTools/ProceduralFoliage/QuickStart/4/PFQS_4.INT.udn b/Engine/Documentation/Source/Engine/OpenWorldTools/ProceduralFoliage/QuickStart/4/PFQS_4.INT.udn
--- a/Engine/Documentation/Source/Engine/OpenWorldTools/ProceduralFoliage/QuickStart/4/PFQS_4.INT.udn
+++ b/Engine/Documentation/Source/Engine/OpenWorldTools/ProceduralFoliage/QuickStart/4/PFQS_4.INT.udn
@@ -1,4 +1,4 @@
-Availability:Docs
+﻿Availability:Docs
 Title: 4. Using Multiple Foliage Types
 Crumbs:
 Description:In this step we will set up our Foliage Spawner to work with multiple Foliage Types.
@@ -53,7 +53,7 @@ We will be continuing to work with the **PFT_00** level that was used in the las
 	![](T_2_FT_In_Level.png)
 
 1. To make the forest look more interesting open up the **Tree_01** Foliage Type and adjust the following parameters with the following values. 
-	The numbers and options that are listed below where selected because they will produce a forest that has interesting clustering and growth interaction with the Foliage Types that are already used. 
+	The numbers and options that are listed below were selected because they will produce a forest that has interesting clustering and growth interaction with the Foliage Types that are already used. 
 	However please feel free to experiment with these numbers and settings till you get something that suits your needs. 
 
 	* **Num Steps:** 4 
