Made public

diff --git a/Engine/Documentation/Source/Engine/OpenWorldTools/ProceduralFoliage/QuickStart/4/PFQS_4.INT.udn b/Engine/Documentation/Source/Engine/OpenWorldTools/ProceduralFoliage/QuickStart/4/PFQS_4.INT.udn
--- a/Engine/Documentation/Source/Engine/OpenWorldTools/ProceduralFoliage/QuickStart/4/PFQS_4.INT.udn
+++ b/Engine/Documentation/Source/Engine/OpenWorldTools/ProceduralFoliage/QuickStart/4/PFQS_4.INT.udn
@@ -1,4 +1,4 @@
-﻿Availability:Docs
+﻿Availability:Public
 Title: 4. Using Multiple Foliage Types
 Crumbs:
 Description:In this step we will set up our Foliage Spawner to work with multiple Foliage Types.
diff --git a/Engine/Documentation/Source/Engine/OpenWorldTools/ProceduralFoliage/QuickStart/5/PFQS_5.INT.udn b/Engine/Documentation/Source/Engine/OpenWorldTools/ProceduralFoliage/QuickStart/5/PFQS_5.INT.udn
--- a/Engine/Documentation/Source/Engine/OpenWorldTools/ProceduralFoliage/QuickStart/5/PFQS_5.INT.udn
+++ b/Engine/Documentation/Source/Engine/OpenWorldTools/ProceduralFoliage/QuickStart/5/PFQS_5.INT.udn
@@ -1,4 +1,4 @@
-﻿Availability:Docs
+﻿Availability:Public
 Title: 5. Implementing Procedural Foliage Blocking Volumes
 Crumbs:
 Description:In this step we will set up a blocking volume to restrict where foliage can spawn.
diff --git a/Engine/Documentation/Source/Engine/OpenWorldTools/ProceduralFoliage/QuickStart/ProceduralFoliageQuickStart.INT.udn b/Engine/Documentation/Source/Engine/OpenWorldTools/ProceduralFoliage/QuickStart/ProceduralFoliageQuickStart.INT.udn
--- a/Engine/Documentation/Source/Engine/OpenWorldTools/ProceduralFoliage/QuickStart/ProceduralFoliageQuickStart.INT.udn
+++ b/Engine/Documentation/Source/Engine/OpenWorldTools/ProceduralFoliage/QuickStart/ProceduralFoliageQuickStart.INT.udn
@@ -1,4 +1,4 @@
-Availability:Docs
+Availability:Public
 Title: Procedural Foliage Tool Quick Start
 Crumbs:
 Description:How to set up and use the Procedural Foliage tool.
@@ -7,7 +7,7 @@ SkillLevel:Intermediate
 Version: 4.8
 
 
-In this Quick Start tutorial we will take a look at how the **Procedural Foliage Tool** that is now available with the 4.8 release of Unreal Engine 4 (UE4) work. 
+In this Quick Start tutorial we will take a look at how the **Procedural Foliage Tool** that is now available with the 4.8 release of Unreal Engine 4 (UE4) works. 
 Over the course of this Quick Start tutorial you will learn how to create, setup, and spawn an entire forests worth of trees inside of UE4 using only the Procedural Foliage tool. 
 You will also be introduced to key properties and settings that will help you shape your virtual forest to fit your project's needs. 
 
