fixed a missed formatting error

diff --git a/Engine/Documentation/Source/Support/Builds/ReleaseNotes/2015/4_8/ReleaseNotes_4_8.INT.udn b/Engine/Documentation/Source/Support/Builds/ReleaseNotes/2015/4_8/ReleaseNotes_4_8.INT.udn
--- a/Engine/Documentation/Source/Support/Builds/ReleaseNotes/2015/4_8/ReleaseNotes_4_8.INT.udn
+++ b/Engine/Documentation/Source/Support/Builds/ReleaseNotes/2015/4_8/ReleaseNotes_4_8.INT.udn
@@ -17,7 +17,7 @@ This was a really big focus for Unreal Engine 4.8.  With every Unreal Engine rel
 
 ## New:  SteamVR/HTC Vive Support
 
-## Unreal Engine now has comprehensive support for Valve's SteamVR API and HTC Vive hardware!
+Unreal Engine now has comprehensive support for Valve's SteamVR API and HTC Vive hardware!
 
 ![image alt text](image_0.jpg)
 
