#UE4 Docs: Adding the final section to the tips and tricks section.

diff --git a/Engine/Documentation/Source/Engine/OpenWorldTools/TipsAndTricks/TipsAndTricks.INT.udn b/Engine/Documentation/Source/Engine/OpenWorldTools/TipsAndTricks/TipsAndTricks.INT.udn
new file mode 100644
--- /dev/null
+++ b/Engine/Documentation/Source/Engine/OpenWorldTools/TipsAndTricks/TipsAndTricks.INT.udn
@@ -0,0 +1,114 @@
+Availability:Docs
+Title: Open World Tools Tips & Tricks
+Crumbs:
+Description: A collection of tips and tricks about how to get the most out of the Open World Tools.
+Navigation:topic
+SkillLevel:Advanced
+Version: 4.8
+
+![](T_Tip-Trick_Header.png)
+
+[TOC(start:2)]
+
+In the following section we will be taking a look at some tips and tricks that you can use to push the Open World tools to their limits. 
+The following tips and tricks come right from the Epic developers that used the Open World Tools to create the Kite Demo that was debuted at the 2015 Game Developers conference.
+Please keep in mind that the following is not meant to be a step by step tutorial but more of a high level look at what was done and how you could apply these techniques to your projects.
+
+[region:note]
+If you are looking for more basic information on how to use the Open World Tools check out the **[Procedural Foliage Quick Start](Engine\OpenWorldTools\ProceduralFoliage\QuickStart)** or the **[Grass Tools Quick Start](Engine\OpenWorldTools\Grass\QuickStart)** for more basic information on how to use these tools.
+[/region]
+
+## Procedural Foliage Volume Based Masking
+
+Having the ability to procedurally spawn an entire forest of trees can be a huge time saver but what if you want to restrict certain trees to only spawn in certain areas. 
+To accomplish this you could do a number of things like use multiple Procedural Foliage Volumes or use the Foliage tool to place the trees where you want but that adds a lot of complexity and upkeep to your project. 
+Luckily the **Foliage Types** have the option to restrict the placement of the Static Mesh that is assigned to them to a certain layer in the Landscape Material allowing only that Foliage Type to spawn when that layer is exposed while painting.
+
+[region:note]
+No mater which method you use you will always need to select the **Procedural Foliage Volume** in the level and click on the **Resimulate** button to update and see your changes.
+
+![](T_Resim_Button.png)
+[/region]
+
+The image below shows a very simple Landscape Material setup that will allow for three different Foliage Types to be used with three different Textures.
+Notice how each input in the **Layer Blend** is uniquely named. 
+This is important because this name is how we associate which Landscape terrain layer the Foliage Type should be restricted to.
+
+![](T_Landscape_Material.png)
+
+Once the Landscape Material has been setup we now need to setup the Foliage Type so that it can be associated with the correct Landscape layer.
+In the image below you can see an example of how a Foliage Type needs to be setup to restrict the placement of the assigned Static Mesh to a specific layer in the Landscape Material. 
+
+![](T_Restrict_To_Landscape_Layers.png)
+
+To set this up for your self you will need to do the following inside of the Foliage Type.
+1. Expand the **Show Advanced** arrow icon to expose the **Landscape Layers** option.
+1. Press the plus sign icon that that is to the right of the word **elements** to add a new entry in the array.
+1. Input the name of the layer in the Landscape Material that you want to restrict this Foliage Type to.
+
+Once you have the Material and Foliage Types setup make sure to assign the Material to the Landscape terrain and then start paint the Landscape as usual.
+When you are done painting or if you just want to see what you have so far, select the Procedural Foliage Volume in the level and then press the **Resimulate** button to update and apply any changes.
+When the completed the foliage that was spawned will be confined to the Landscape Texture it is associated with like in the following image.
+
+![](T_Content_Breakdown.png)
+
+If we take a closer look at the image above we can see that while everything has been procedurally placed like normal but none of the spawned meshes inter mix with one another. 
+For example the leaved trees only spawn where the grass Texture is applied to the Landscape.
+If you look closely you will the same applies for the bush and pine trees.
+
+## Texture Based Masking for the Grass Tool
+
+One of the other methods for controlling where Foliage meshes can or cannot spawn is to use a Texture that acts as a mask telling the Open World Tools to either spawn or not to spawn a Foliage mesh in a specific area. 
+The following image shows a very simple Material that demonstrates Texture based masking. 
+
+![](T_Texture_Based_Masking.png)
+
+Starting at the far left of this image we have the World Position Scaling section. 
+This is scaling the Mask Texture in world space so that it can be scaled big enough to fit the Landscape terrain. 
+Moving to the left we have the Mask Texture that is used to define where the foliage should be placed as well as the blending of the two Landscape Textures. 
+Here is a breakdown of what the different channels of the Mask Texture are used for.
+
+* **Red Channel**: This channel is used to define where the grass meshes should be spawned on the Landscape as well as to define how the Landscape base Textures should be blended.
+
+* **Green Channel**: This channel is used to define where the rock meshes should be spawned on the Landscape.
+
+Finally on the far right of the image we have the Main Material node and the Grass Output. 
+The Main Material node is setup in the same way as any other standard Material. 
+The Grass Output is setup in a similar manner to how you would do Texture based masking with the key difference being that you are going to use the portion of the Material that is masked or unmasked to spawn or not spawn Static Meshes instead of just showing a different Texture. 
+When this Material is applied to a Landscape in UE4 you should get results that look similar to the following image minus the mask Texture that was placed in the corner to show which mask color was doing what to the Landscape.
+
+![](T_Masking_In_Action.png)
+
+In the image above the grass Static Mesh, for the most part, are spawned only in areas that use the green channel of the mask Texture while the rocks use the red channel of the mask Texture to spawn. 
+As you might have noticed the grass and rocks are not entirely contained within the area of the mask Texture that spawns them. 
+This is expected behavior as the placement of the Static Meshes can be jittered and offset to make things look more natural.
+
+## Math Based Masking for the Grass Tool
+
+While using a Texture for masking can be convenient and quick there could be times where you can not afford the extra Texture data due to performance concerns or you just do not need that much control over the placement. 
+In cases like this you can power the placement using Math functions created in the Material Editor in place of Textures.
+
+![](T_Math_Based_Masking.png)
+
+ In the image above we have taken the Material from the last example but changed how the grass and rocks are spawned.
+For this example the grass and rocks are spawned in a checkerboard pattern by using the **Checker Pattern** Material Function to control placement. 
+ The **One Minus** was added to the rocks so that the rocks will spawn in the parts that the grass does not spawn in. 
+ When applied to the Landscape you will see something that looks like the following.
+
+![](T_Check_Pattern.png)
+
+## Non Random Placement for the Grass Tool
+
+With the right adjustments, the Grass tool can also be used to simulate things like crops growing in a field. 
+In the image below you can see an example of this type of behavior.
+
+![](T_Spawing_In_Row.png)
+
+To get the Grass tool to function like this all you need to do is make sure that **Use Grid** is checked, that **Placement Jitter** is set to **0** and **Grass Density** is set to something under 100. 
+The following image shows how the Landscape Grass Type was setup to achieve the results you saw in the first image.
+
+![](T_LGT_Crops_Setup.png)
+
+
+
+
