#UE4 DOCS: updated image and BP code snippet on step 12 of the Behavior Tree Quick start to avoid a technically correct Blueprint connection to one that makes sense

diff --git a/Engine/Documentation/Source/Engine/AI/BehaviorTrees/QuickStart/12/12.INT.udn b/Engine/Documentation/Source/Engine/AI/BehaviorTrees/QuickStart/12/12.INT.udn
--- a/Engine/Documentation/Source/Engine/AI/BehaviorTrees/QuickStart/12/12.INT.udn
+++ b/Engine/Documentation/Source/Engine/AI/BehaviorTrees/QuickStart/12/12.INT.udn
@@ -80,65 +80,66 @@ Here is what is going on in this graph:
 
 [VAR:bpcode]
 	Begin Object Class=K2Node_Event Name="K2Node_Event_845"
-	   Begin Object Class=EdGraphPin Name="EdGraphPin_4187"
+	   Begin Object Class=EdGraphPin Name="EdGraphPin_4172"
 	   End Object
-	   Begin Object Class=EdGraphPin Name="EdGraphPin_4188"
+	   Begin Object Class=EdGraphPin Name="EdGraphPin_4173"
 	   End Object
-	   Begin Object Class=EdGraphPin Name="EdGraphPin_4189"
+	   Begin Object Class=EdGraphPin Name="EdGraphPin_4174"
 	   End Object
-	   Begin Object Name="EdGraphPin_4187"
+	   Begin Object Name="EdGraphPin_4172"
 		  PinName="OutputDelegate"
 		  Direction=EGPD_Output
 		  PinType=(PinCategory="delegate",PinSubCategoryMemberReference=(MemberParentClass=Class'/Script/AIModule.BTTask_BlueprintBase',MemberName="ReceiveExecute"))
 	   End Object
-	   Begin Object Name="EdGraphPin_4188"
+	   Begin Object Name="EdGraphPin_4173"
 		  PinName="then"
 		  Direction=EGPD_Output
 		  PinType=(PinCategory="exec")
-		  LinkedTo(0)=EdGraphPin'K2Node_IfThenElse_95.EdGraphPin_20153'
+		  LinkedTo(0)=EdGraphPin'K2Node_IfThenElse_100.EdGraphPin_20069'
 	   End Object
-	   Begin Object Name="EdGraphPin_4189"
+	   Begin Object Name="EdGraphPin_4174"
 		  PinName="OwnerActor"
 		  PinToolTip="Owner Actor\nActor Reference"
 		  Direction=EGPD_Output
 		  PinType=(PinCategory="object",PinSubCategoryObject=Class'/Script/Engine.Actor')
-		  LinkedTo(0)=EdGraphPin'K2Node_DynamicCast_1392.EdGraphPin_6841'
+		  LinkedTo(0)=EdGraphPin'K2Node_DynamicCast_1392.EdGraphPin_6850'
 	   End Object
 	   EventSignatureName="ReceiveExecute"
 	   EventSignatureClass=Class'/Script/AIModule.BTTask_BlueprintBase'
 	   bOverrideFunction=True
-	   Pins(0)=EdGraphPin'EdGraphPin_4187'
-	   Pins(1)=EdGraphPin'EdGraphPin_4188'
-	   Pins(2)=EdGraphPin'EdGraphPin_4189'
-	   NodePosX=-304
-	   NodeGuid=EAEB21444709FA5B18806491D4728A43
+	   Pins(0)=EdGraphPin'EdGraphPin_4172'
+	   Pins(1)=EdGraphPin'EdGraphPin_4173'
+	   Pins(2)=EdGraphPin'EdGraphPin_4174'
+	   NodePosX=-352
+	   NodePosY=64
+	   NodeGuid=6739971B49CD20A6D2ACF6B441AF6A7A
 	End Object
 	Begin Object Class=K2Node_CallFunction Name="K2Node_CallFunction_123126"
-	   Begin Object Class=EdGraphPin Name="EdGraphPin_4091"
+	   Begin Object Class=EdGraphPin Name="EdGraphPin_4076"
 	   End Object
-	   Begin Object Class=EdGraphPin Name="EdGraphPin_4092"
+	   Begin Object Class=EdGraphPin Name="EdGraphPin_4077"
 	   End Object
-	   Begin Object Class=EdGraphPin Name="EdGraphPin_4093"
+	   Begin Object Class=EdGraphPin Name="EdGraphPin_4078"
 	   End Object
-	   Begin Object Class=EdGraphPin Name="EdGraphPin_4094"
+	   Begin Object Class=EdGraphPin Name="EdGraphPin_4079"
 	   End Object
-	   Begin Object Name="EdGraphPin_4091"
+	   Begin Object Name="EdGraphPin_4076"
 		  PinName="execute"
 		  PinType=(PinCategory="exec")
-		  LinkedTo(0)=EdGraphPin'K2Node_AIMoveTo_42.EdGraphPin_4127'
+		  LinkedTo(0)=EdGraphPin'K2Node_AIMoveTo_42.EdGraphPin_4112'
 	   End Object
-	   Begin Object Name="EdGraphPin_4092"
+	   Begin Object Name="EdGraphPin_4077"
 		  PinName="then"
 		  Direction=EGPD_Output
 		  PinType=(PinCategory="exec")
 	   End Object
-	   Begin Object Name="EdGraphPin_4093"
+	   Begin Object Name="EdGraphPin_4078"
 		  PinName="self"
 		  PinFriendlyName="Target"
 		  PinToolTip="Target\nBTTask Blueprint Base Reference"
 		  PinType=(PinCategory="object",PinSubCategoryObject=Class'/Script/AIModule.BTTask_BlueprintBase')
 	   End Object
-	   Begin Object Name="EdGraphPin_4094"
+	   Begin Object Name="EdGraphPin_4079"
 		  PinName="bSuccess"
 		  PinToolTip="Success\nBoolean"
 		  PinType=(PinCategory="bool")
@@ -146,260 +147,257 @@ Here is what is going on in this graph:
 		  AutogeneratedDefaultValue="false"
 	   End Object
 	   FunctionReference=(MemberName="FinishExecute",bSelfContext=True)
-	   Pins(0)=EdGraphPin'EdGraphPin_4091'
-	   Pins(1)=EdGraphPin'EdGraphPin_4092'
-	   Pins(2)=EdGraphPin'EdGraphPin_4093'
-	   Pins(3)=EdGraphPin'EdGraphPin_4094'
-	   NodePosX=1456
-	   NodePosY=-48
-	   NodeGuid=C331B0A749C86B7E532866B449BF1FA5
+	   Pins(0)=EdGraphPin'EdGraphPin_4076'
+	   Pins(1)=EdGraphPin'EdGraphPin_4077'
+	   Pins(2)=EdGraphPin'EdGraphPin_4078'
+	   Pins(3)=EdGraphPin'EdGraphPin_4079'
+	   NodePosX=1440
+	   NodeGuid=81BF51914D95508F36CB2387949AA34F
 	End Object
 	Begin Object Class=K2Node_DynamicCast Name="K2Node_DynamicCast_1392"
-	   Begin Object Class=EdGraphPin Name="EdGraphPin_6838"
+	   Begin Object Class=EdGraphPin Name="EdGraphPin_6847"
 	   End Object
-	   Begin Object Class=EdGraphPin Name="EdGraphPin_6839"
+	   Begin Object Class=EdGraphPin Name="EdGraphPin_6848"
 	   End Object
-	   Begin Object Class=EdGraphPin Name="EdGraphPin_6840"
+	   Begin Object Class=EdGraphPin Name="EdGraphPin_6849"
 	   End Object
-	   Begin Object Class=EdGraphPin Name="EdGraphPin_6841"
+	   Begin Object Class=EdGraphPin Name="EdGraphPin_6850"
 	   End Object
-	   Begin Object Class=EdGraphPin Name="EdGraphPin_6842"
+	   Begin Object Class=EdGraphPin Name="EdGraphPin_6851"
 	   End Object
-	   Begin Object Class=EdGraphPin Name="EdGraphPin_6843"
+	   Begin Object Class=EdGraphPin Name="EdGraphPin_6852"
 	   End Object
-	   Begin Object Name="EdGraphPin_6838"
+	   Begin Object Name="EdGraphPin_6847"
 		  PinName="execute"
 		  PinType=(PinCategory="exec")
-		  LinkedTo(0)=EdGraphPin'K2Node_IfThenElse_95.EdGraphPin_20156'
+		  LinkedTo(0)=EdGraphPin'K2Node_IfThenElse_100.EdGraphPin_20072'
 	   End Object
-	   Begin Object Name="EdGraphPin_6839"
+	   Begin Object Name="EdGraphPin_6848"
 		  PinName="then"
 		  Direction=EGPD_Output
 		  PinType=(PinCategory="exec")
-		  LinkedTo(0)=EdGraphPin'K2Node_VariableSet_421.EdGraphPin_4158'
+		  LinkedTo(0)=EdGraphPin'K2Node_VariableSet_421.EdGraphPin_4143'
 	   End Object
-	   Begin Object Name="EdGraphPin_6840"
+	   Begin Object Name="EdGraphPin_6849"
 		  PinName="CastFailed"
 		  Direction=EGPD_Output
 		  PinType=(PinCategory="exec")
 	   End Object
-	   Begin Object Name="EdGraphPin_6841"
+	   Begin Object Name="EdGraphPin_6850"
 		  PinName="Object"
 		  PinType=(PinCategory="object",PinSubCategoryObject=Class'/Script/CoreUObject.Object')
-		  LinkedTo(0)=EdGraphPin'K2Node_Event_845.EdGraphPin_4189'
+		  LinkedTo(0)=EdGraphPin'K2Node_Event_845.EdGraphPin_4174'
 	   End Object
-	   Begin Object Name="EdGraphPin_6842"
+	   Begin Object Name="EdGraphPin_6851"
 		  PinName="AsFollower AI CON"
 		  Direction=EGPD_Output
 		  PinType=(PinCategory="object",PinSubCategoryObject=BlueprintGeneratedClass'/Game/Blueprints/AI/Follower_AI_Con.Follower_AI_CON_C')
-		  LinkedTo(0)=EdGraphPin'K2Node_CallFunction_123506.EdGraphPin_4119'
-		  LinkedTo(1)=EdGraphPin'K2Node_VariableSet_421.EdGraphPin_4160'
+		  LinkedTo(0)=EdGraphPin'K2Node_VariableSet_421.EdGraphPin_4145'
 	   End Object
-	   Begin Object Name="EdGraphPin_6843"
+	   Begin Object Name="EdGraphPin_6852"
 		  PinName="bSuccess"
 		  Direction=EGPD_Output
 		  PinType=(PinCategory="bool")
 		  bHidden=True
 	   End Object
 	   TargetType=BlueprintGeneratedClass'/Game/Blueprints/AI/Follower_AI_Con.Follower_AI_CON_C'
-	   Pins(0)=EdGraphPin'EdGraphPin_6838'
-	   Pins(1)=EdGraphPin'EdGraphPin_6839'
-	   Pins(2)=EdGraphPin'EdGraphPin_6840'
-	   Pins(3)=EdGraphPin'EdGraphPin_6841'
-	   Pins(4)=EdGraphPin'EdGraphPin_6842'
-	   Pins(5)=EdGraphPin'EdGraphPin_6843'
-	   NodePosX=416
-	   NodePosY=16
-	   NodeGuid=779D303548E480DF1D94709C00DF4CC6
+	   Pins(0)=EdGraphPin'EdGraphPin_6847'
+	   Pins(1)=EdGraphPin'EdGraphPin_6848'
+	   Pins(2)=EdGraphPin'EdGraphPin_6849'
+	   Pins(3)=EdGraphPin'EdGraphPin_6850'
+	   Pins(4)=EdGraphPin'EdGraphPin_6851'
+	   Pins(5)=EdGraphPin'EdGraphPin_6852'
+	   NodePosX=400
+	   NodePosY=64
+	   NodeGuid=7FB506C8481E3974BAB151A7A93C47BA
 	End Object
 	Begin Object Class=K2Node_CallFunction Name="K2Node_CallFunction_123506"
-	   Begin Object Class=EdGraphPin Name="EdGraphPin_4119"
+	   Begin Object Class=EdGraphPin Name="EdGraphPin_4104"
 	   End Object
-	   Begin Object Class=EdGraphPin Name="EdGraphPin_4120"
+	   Begin Object Class=EdGraphPin Name="EdGraphPin_4105"
 	   End Object
-	   Begin Object Name="EdGraphPin_4119"
+	   Begin Object Name="EdGraphPin_4104"
 		  PinName="self"
 		  PinFriendlyName="Target"
 		  PinToolTip="Target\nController Reference"
 		  PinType=(PinCategory="object",PinSubCategoryObject=Class'/Script/Engine.Controller')
-		  LinkedTo(0)=EdGraphPin'K2Node_DynamicCast_1392.EdGraphPin_6842'
+		  LinkedTo(0)=EdGraphPin'K2Node_VariableGet_265.EdGraphPin_20067'
 	   End Object
-	   Begin Object Name="EdGraphPin_4120"
+	   Begin Object Name="EdGraphPin_4105"
 		  PinName="ReturnValue"
 		  PinToolTip="Return Value\nPawn Reference"
 		  Direction=EGPD_Output
 		  PinType=(PinCategory="object",PinSubCategoryObject=Class'/Script/Engine.Pawn')
-		  LinkedTo(0)=EdGraphPin'K2Node_AIMoveTo_42.EdGraphPin_4131'
+		  LinkedTo(0)=EdGraphPin'K2Node_AIMoveTo_42.EdGraphPin_4116'
 	   End Object
 	   bIsPureFunc=True
 	   bIsConstFunc=True
 	   FunctionReference=(MemberParentClass=Class'/Script/Engine.Controller',MemberName="K2_GetPawn")
-	   Pins(0)=EdGraphPin'EdGraphPin_4119'
-	   Pins(1)=EdGraphPin'EdGraphPin_4120'
-	   NodePosX=688
-	   NodePosY=112
-	   NodeGuid=94B808B141A90BBF2045408DC96B5EA4
+	   Pins(0)=EdGraphPin'EdGraphPin_4104'
+	   Pins(1)=EdGraphPin'EdGraphPin_4105'
+	   NodePosX=784
+	   NodePosY=176
+	   NodeGuid=27431F034F88EB67FFA2BF9C9F7C776D
 	End Object
 	Begin Object Class=K2Node_AIMoveTo Name="K2Node_AIMoveTo_42"
-	   Begin Object Class=EdGraphPin Name="EdGraphPin_4125"
+	   Begin Object Class=EdGraphPin Name="EdGraphPin_4110"
 	   End Object
-	   Begin Object Class=EdGraphPin Name="EdGraphPin_4126"
+	   Begin Object Class=EdGraphPin Name="EdGraphPin_4111"
 	   End Object
-	   Begin Object Class=EdGraphPin Name="EdGraphPin_4127"
+	   Begin Object Class=EdGraphPin Name="EdGraphPin_4112"
 	   End Object
-	   Begin Object Class=EdGraphPin Name="EdGraphPin_4128"
+	   Begin Object Class=EdGraphPin Name="EdGraphPin_4113"
 	   End Object
-	   Begin Object Class=EdGraphPin Name="EdGraphPin_4129"
+	   Begin Object Class=EdGraphPin Name="EdGraphPin_4114"
 	   End Object
-	   Begin Object Class=EdGraphPin Name="EdGraphPin_4130"
+	   Begin Object Class=EdGraphPin Name="EdGraphPin_4115"
 	   End Object
-	   Begin Object Class=EdGraphPin Name="EdGraphPin_4131"
+	   Begin Object Class=EdGraphPin Name="EdGraphPin_4116"
 	   End Object
-	   Begin Object Class=EdGraphPin Name="EdGraphPin_4132"
+	   Begin Object Class=EdGraphPin Name="EdGraphPin_4117"
 	   End Object
-	   Begin Object Class=EdGraphPin Name="EdGraphPin_4133"
+	   Begin Object Class=EdGraphPin Name="EdGraphPin_4118"
 	   End Object
-	   Begin Object Class=EdGraphPin Name="EdGraphPin_4134"
+	   Begin Object Class=EdGraphPin Name="EdGraphPin_4119"
 	   End Object
-	   Begin Object Class=EdGraphPin Name="EdGraphPin_4135"
+	   Begin Object Class=EdGraphPin Name="EdGraphPin_4120"
 	   End Object
-	   Begin Object Name="EdGraphPin_4125"
+	   Begin Object Name="EdGraphPin_4110"
 		  PinName="execute"
 		  PinType=(PinCategory="exec")
-		  LinkedTo(0)=EdGraphPin'K2Node_VariableSet_421.EdGraphPin_4159'
-		  LinkedTo(1)=EdGraphPin'K2Node_IfThenElse_95.EdGraphPin_20155'
+		  LinkedTo(0)=EdGraphPin'K2Node_VariableSet_421.EdGraphPin_4144'
+		  LinkedTo(1)=EdGraphPin'K2Node_IfThenElse_100.EdGraphPin_20071'
 	   End Object
-	   Begin Object Name="EdGraphPin_4126"
+	   Begin Object Name="EdGraphPin_4111"
 		  PinName="then"
 		  Direction=EGPD_Output
 		  PinType=(PinCategory="exec")
 	   End Object
-	   Begin Object Name="EdGraphPin_4127"
+	   Begin Object Name="EdGraphPin_4112"
 		  PinName="OnSuccess"
 		  Direction=EGPD_Output
 		  PinType=(PinCategory="exec")
-		  LinkedTo(0)=EdGraphPin'K2Node_CallFunction_123126.EdGraphPin_4091'
+		  LinkedTo(0)=EdGraphPin'K2Node_CallFunction_123126.EdGraphPin_4076'
 	   End Object
-	   Begin Object Name="EdGraphPin_4128"
+	   Begin Object Name="EdGraphPin_4113"
 		  PinName="OnFail"
 		  Direction=EGPD_Output
 		  PinType=(PinCategory="exec")
 	   End Object
-	   Begin Object Name="EdGraphPin_4129"
+	   Begin Object Name="EdGraphPin_4114"
 		  PinName="MovementResult"
 		  Direction=EGPD_Output
 		  PinType=(PinCategory="byte",PinSubCategoryObject=Enum'/Script/AIModule.PathFollowingComponent:EPathFollowingResult')
 	   End Object
-	   Begin Object Name="EdGraphPin_4130"
+	   Begin Object Name="EdGraphPin_4115"
 		  PinName="WorldContextObject"
 		  PinType=(PinCategory="object",PinSubCategoryObject=Class'/Script/CoreUObject.Object')
 		  DefaultValue="0"
 		  bHidden=True
 	   End Object
-	   Begin Object Name="EdGraphPin_4131"
+	   Begin Object Name="EdGraphPin_4116"
 		  PinName="Pawn"
 		  PinType=(PinCategory="object",PinSubCategoryObject=Class'/Script/Engine.Pawn')
-		  LinkedTo(0)=EdGraphPin'K2Node_CallFunction_123506.EdGraphPin_4120'
+		  LinkedTo(0)=EdGraphPin'K2Node_CallFunction_123506.EdGraphPin_4105'
 	   End Object
-	   Begin Object Name="EdGraphPin_4132"
+	   Begin Object Name="EdGraphPin_4117"
 		  PinName="Destination"
 		  PinType=(PinCategory="struct",PinSubCategoryObject=ScriptStruct'/Script/CoreUObject.Object:Vector')
 		  DefaultValue="0, 0, 0"
 		  AutogeneratedDefaultValue="0, 0, 0"
 	   End Object
-	   Begin Object Name="EdGraphPin_4133"
+	   Begin Object Name="EdGraphPin_4118"
 		  PinName="TargetActor"
 		  PinType=(PinCategory="object",PinSubCategoryObject=Class'/Script/Engine.Actor')
-		  LinkedTo(0)=EdGraphPin'K2Node_CallFunction_26123.EdGraphPin_4173'
+		  LinkedTo(0)=EdGraphPin'K2Node_CallFunction_26123.EdGraphPin_4158'
 	   End Object
-	   Begin Object Name="EdGraphPin_4134"
+	   Begin Object Name="EdGraphPin_4119"
 		  PinName="AcceptanceRadius"
 		  PinType=(PinCategory="float")
 		  DefaultValue="50.000"
 		  AutogeneratedDefaultValue="5.000000"
 	   End Object
-	   Begin Object Name="EdGraphPin_4135"
+	   Begin Object Name="EdGraphPin_4120"
 		  PinName="bStopOnOverlap"
 		  PinType=(PinCategory="bool")
 		  DefaultValue="false"
 		  AutogeneratedDefaultValue="false"
 	   End Object
-	   Pins(0)=EdGraphPin'EdGraphPin_4125'
-	   Pins(1)=EdGraphPin'EdGraphPin_4126'
-	   Pins(2)=EdGraphPin'EdGraphPin_4127'
-	   Pins(3)=EdGraphPin'EdGraphPin_4128'
-	   Pins(4)=EdGraphPin'EdGraphPin_4129'
-	   Pins(5)=EdGraphPin'EdGraphPin_4130'
-	   Pins(6)=EdGraphPin'EdGraphPin_4131'
-	   Pins(7)=EdGraphPin'EdGraphPin_4132'
-	   Pins(8)=EdGraphPin'EdGraphPin_4133'
-	   Pins(9)=EdGraphPin'EdGraphPin_4134'
-	   Pins(10)=EdGraphPin'EdGraphPin_4135'
-	   NodePosX=1088
-	   NodePosY=-48
+	   Pins(0)=EdGraphPin'EdGraphPin_4110'
+	   Pins(1)=EdGraphPin'EdGraphPin_4111'
+	   Pins(2)=EdGraphPin'EdGraphPin_4112'
+	   Pins(3)=EdGraphPin'EdGraphPin_4113'
+	   Pins(4)=EdGraphPin'EdGraphPin_4114'
+	   Pins(5)=EdGraphPin'EdGraphPin_4115'
+	   Pins(6)=EdGraphPin'EdGraphPin_4116'
+	   Pins(7)=EdGraphPin'EdGraphPin_4117'
+	   Pins(8)=EdGraphPin'EdGraphPin_4118'
+	   Pins(9)=EdGraphPin'EdGraphPin_4119'
+	   Pins(10)=EdGraphPin'EdGraphPin_4120'
+	   NodePosX=1104
 	   ErrorType=1
 	   ErrorMsg="Error Unexpected node type K2Node_AIMoveTo encountered at  AI MoveTo "
-	   NodeGuid=CCC96BFC4C9294A884E008A26E53FBBE
+	   NodeGuid=E94A59B64BE284DD7C92E6B51454D074
 	End Object
 	Begin Object Class=K2Node_VariableSet Name="K2Node_VariableSet_421"
-	   Begin Object Class=EdGraphPin Name="EdGraphPin_4158"
+	   Begin Object Class=EdGraphPin Name="EdGraphPin_4143"
 	   End Object
-	   Begin Object Class=EdGraphPin Name="EdGraphPin_4159"
+	   Begin Object Class=EdGraphPin Name="EdGraphPin_4144"
 	   End Object
-	   Begin Object Class=EdGraphPin Name="EdGraphPin_4160"
+	   Begin Object Class=EdGraphPin Name="EdGraphPin_4145"
 	   End Object
-	   Begin Object Class=EdGraphPin Name="EdGraphPin_4161"
+	   Begin Object Class=EdGraphPin Name="EdGraphPin_4146"
 	   End Object
-	   Begin Object Class=EdGraphPin Name="EdGraphPin_4162"
+	   Begin Object Class=EdGraphPin Name="EdGraphPin_4147"
 	   End Object
-	   Begin Object Name="EdGraphPin_4158"
+	   Begin Object Name="EdGraphPin_4143"
 		  PinName="execute"
 		  PinType=(PinCategory="exec")
-		  LinkedTo(0)=EdGraphPin'K2Node_DynamicCast_1392.EdGraphPin_6839'
+		  LinkedTo(0)=EdGraphPin'K2Node_DynamicCast_1392.EdGraphPin_6848'
 	   End Object
-	   Begin Object Name="EdGraphPin_4159"
+	   Begin Object Name="EdGraphPin_4144"
 		  PinName="then"
 		  Direction=EGPD_Output
 		  PinType=(PinCategory="exec")
-		  LinkedTo(0)=EdGraphPin'K2Node_AIMoveTo_42.EdGraphPin_4125'
+		  LinkedTo(0)=EdGraphPin'K2Node_AIMoveTo_42.EdGraphPin_4110'
 	   End Object
-	   Begin Object Name="EdGraphPin_4160"
+	   Begin Object Name="EdGraphPin_4145"
 		  PinName="AI_CON_Ref"
 		  PinType=(PinCategory="object",PinSubCategoryObject=BlueprintGeneratedClass'/Game/Blueprints/AI/Follower_AI_Con.Follower_AI_CON_C')
-		  LinkedTo(0)=EdGraphPin'K2Node_DynamicCast_1392.EdGraphPin_6842'
+		  LinkedTo(0)=EdGraphPin'K2Node_DynamicCast_1392.EdGraphPin_6851'
 	   End Object
-	   Begin Object Name="EdGraphPin_4161"
+	   Begin Object Name="EdGraphPin_4146"
 		  PinName="Output_Get"
 		  PinToolTip="Retrieves the value of the variable, can use instead of a separate Get node"
 		  Direction=EGPD_Output
 		  PinType=(PinCategory="object",PinSubCategoryObject=BlueprintGeneratedClass'/Game/Blueprints/AI/Follower_AI_Con.Follower_AI_CON_C')
 	   End Object
-	   Begin Object Name="EdGraphPin_4162"
+	   Begin Object Name="EdGraphPin_4147"
 		  PinName="self"
 		  PinFriendlyName="Target"
 		  PinType=(PinCategory="object",PinSubCategoryObject=BlueprintGeneratedClass'/Game/Blueprints/AI/RapidMoveTo.RapidMoveTo_C')
 		  bHidden=True
 	   End Object
 	   VariableReference=(MemberName="AI_CON_Ref",bSelfContext=True)
-	   Pins(0)=EdGraphPin'EdGraphPin_4158'
-	   Pins(1)=EdGraphPin'EdGraphPin_4159'
-	   Pins(2)=EdGraphPin'EdGraphPin_4160'
-	   Pins(3)=EdGraphPin'EdGraphPin_4161'
-	   Pins(4)=EdGraphPin'EdGraphPin_4162'
-	   NodePosX=688
-	   NodePosY=16
-	   NodeGuid=E57170A748D2EAA0551017B3922F6FD0
+	   Pins(0)=EdGraphPin'EdGraphPin_4143'
+	   Pins(1)=EdGraphPin'EdGraphPin_4144'
+	   Pins(2)=EdGraphPin'EdGraphPin_4145'
+	   Pins(3)=EdGraphPin'EdGraphPin_4146'
+	   Pins(4)=EdGraphPin'EdGraphPin_4147'
+	   NodePosX=672
+	   NodePosY=64
+	   NodeGuid=D120D77248C9744E274BB8925D0FBE23
 	End Object
 	Begin Object Class=K2Node_CallFunction Name="K2Node_CallFunction_26123"
-	   Begin Object Class=EdGraphPin Name="EdGraphPin_4170"
+	   Begin Object Class=EdGraphPin Name="EdGraphPin_4155"
 	   End Object
-	   Begin Object Class=EdGraphPin Name="EdGraphPin_4171"
+	   Begin Object Class=EdGraphPin Name="EdGraphPin_4156"
 	   End Object
-	   Begin Object Class=EdGraphPin Name="EdGraphPin_4172"
+	   Begin Object Class=EdGraphPin Name="EdGraphPin_4157"
 	   End Object
-	   Begin Object Class=EdGraphPin Name="EdGraphPin_4173"
+	   Begin Object Class=EdGraphPin Name="EdGraphPin_4158"
 	   End Object
-	   Begin Object Name="EdGraphPin_4170"
+	   Begin Object Name="EdGraphPin_4155"
 		  PinName="self"
 		  PinFriendlyName="Target"
 		  PinToolTip="Target\nBTFunction Library Reference"
@@ -407,174 +405,197 @@ Here is what is going on in this graph:
 		  DefaultObject=Default__BTFunctionLibrary
 		  bHidden=True
 	   End Object
-	   Begin Object Name="EdGraphPin_4171"
+	   Begin Object Name="EdGraphPin_4156"
 		  PinName="NodeOwner"
 		  PinToolTip="Node Owner\nBTNode Reference"
 		  PinType=(PinCategory="object",PinSubCategoryObject=Class'/Script/AIModule.BTNode')
 		  bHidden=True
 	   End Object
-	   Begin Object Name="EdGraphPin_4172"
+	   Begin Object Name="EdGraphPin_4157"
 		  PinName="Key"
 		  PinToolTip="Key\nBlackboard Key Selector Structure (by ref)"
 		  PinType=(PinCategory="struct",PinSubCategoryObject=ScriptStruct'/Script/AIModule.BehaviorTreeTypes:BlackboardKeySelector',bIsReference=True,bIsConst=True)
-		  LinkedTo(0)=EdGraphPin'K2Node_VariableGet_150.EdGraphPin_4182'
+		  LinkedTo(0)=EdGraphPin'K2Node_VariableGet_150.EdGraphPin_4167'
 		  bDefaultValueIsIgnored=True
 	   End Object
-	   Begin Object Name="EdGraphPin_4173"
+	   Begin Object Name="EdGraphPin_4158"
 		  PinName="ReturnValue"
 		  PinToolTip="Return Value\nActor Reference"
 		  Direction=EGPD_Output
 		  PinType=(PinCategory="object",PinSubCategoryObject=Class'/Script/Engine.Actor')
-		  LinkedTo(0)=EdGraphPin'K2Node_AIMoveTo_42.EdGraphPin_4133'
+		  LinkedTo(0)=EdGraphPin'K2Node_AIMoveTo_42.EdGraphPin_4118'
 	   End Object
 	   bIsPureFunc=True
 	   FunctionReference=(MemberParentClass=Class'/Script/AIModule.BTFunctionLibrary',MemberName="GetBlackboardValueAsActor")
-	   Pins(0)=EdGraphPin'EdGraphPin_4170'
-	   Pins(1)=EdGraphPin'EdGraphPin_4171'
-	   Pins(2)=EdGraphPin'EdGraphPin_4172'
-	   Pins(3)=EdGraphPin'EdGraphPin_4173'
-	   NodePosX=672
-	   NodePosY=208
+	   Pins(0)=EdGraphPin'EdGraphPin_4155'
+	   Pins(1)=EdGraphPin'EdGraphPin_4156'
+	   Pins(2)=EdGraphPin'EdGraphPin_4157'
+	   Pins(3)=EdGraphPin'EdGraphPin_4158'
+	   NodePosX=752
+	   NodePosY=256
 	   ErrorType=1
 	   ErrorMsg="Error Default value \'\' for  Key  is invalid: \'Literal values are not allowed for pass-by-reference parameters.\'"
-	   NodeGuid=E71862544FC69B9329F123B202591DE5
+	   NodeGuid=98F411F940A58634CB63C1BF782C619B
 	End Object
 	Begin Object Class=K2Node_VariableGet Name="K2Node_VariableGet_150"
-	   Begin Object Class=EdGraphPin Name="EdGraphPin_4182"
+	   Begin Object Class=EdGraphPin Name="EdGraphPin_4167"
 	   End Object
-	   Begin Object Class=EdGraphPin Name="EdGraphPin_4183"
+	   Begin Object Class=EdGraphPin Name="EdGraphPin_4168"
 	   End Object
-	   Begin Object Name="EdGraphPin_4182"
+	   Begin Object Name="EdGraphPin_4167"
 		  PinName="TargetToFollow"
 		  Direction=EGPD_Output
 		  PinType=(PinCategory="struct",PinSubCategoryObject=ScriptStruct'/Script/AIModule.BehaviorTreeTypes:BlackboardKeySelector')
-		  LinkedTo(0)=EdGraphPin'K2Node_CallFunction_26123.EdGraphPin_4172'
+		  LinkedTo(0)=EdGraphPin'K2Node_CallFunction_26123.EdGraphPin_4157'
 	   End Object
-	   Begin Object Name="EdGraphPin_4183"
+	   Begin Object Name="EdGraphPin_4168"
 		  PinName="self"
 		  PinFriendlyName="Target"
 		  PinType=(PinCategory="object",PinSubCategoryObject=BlueprintGeneratedClass'/Game/Blueprints/AI/RapidMoveTo.RapidMoveTo_C')
 		  bHidden=True
 	   End Object
 	   VariableReference=(MemberName="TargetToFollow",bSelfContext=True)
-	   Pins(0)=EdGraphPin'EdGraphPin_4182'
-	   Pins(1)=EdGraphPin'EdGraphPin_4183'
-	   NodePosX=512
-	   NodePosY=208
-	   NodeGuid=2A8535054EDA5546714B1D9F3B88DCB8
+	   Pins(0)=EdGraphPin'EdGraphPin_4167'
+	   Pins(1)=EdGraphPin'EdGraphPin_4168'
+	   NodePosX=592
+	   NodePosY=256
+	   NodeGuid=361AB0164C7B72B6F80BD986EFD799A2
 	End Object
-	Begin Object Class=K2Node_VariableGet Name="K2Node_VariableGet_252"
-	   Begin Object Class=EdGraphPin Name="EdGraphPin_20148"
+	Begin Object Class=K2Node_VariableGet Name="K2Node_VariableGet_265"
+	   Begin Object Class=EdGraphPin Name="EdGraphPin_20067"
 	   End Object
-	   Begin Object Class=EdGraphPin Name="EdGraphPin_20147"
+	   Begin Object Class=EdGraphPin Name="EdGraphPin_20068"
 	   End Object
-	   Begin Object Name="EdGraphPin_20148"
+	   Begin Object Name="EdGraphPin_20067"
+		  PinName="AI_CON_Ref"
+		  Direction=EGPD_Output
+		  PinType=(PinCategory="object",PinSubCategoryObject=BlueprintGeneratedClass'/Game/Blueprints/AI/Follower_AI_Con.Follower_AI_CON_C')
+		  LinkedTo(0)=EdGraphPin'K2Node_CallFunction_123506.EdGraphPin_4104'
+	   End Object
+	   Begin Object Name="EdGraphPin_20068"
 		  PinName="self"
 		  PinFriendlyName="Target"
 		  PinType=(PinCategory="object",PinSubCategoryObject=BlueprintGeneratedClass'/Game/Blueprints/AI/RapidMoveTo.RapidMoveTo_C')
 		  bHidden=True
 	   End Object
-	   Begin Object Name="EdGraphPin_20147"
-		  PinName="AI_CON_Ref"
-		  Direction=EGPD_Output
-		  PinType=(PinCategory="object",PinSubCategoryObject=BlueprintGeneratedClass'/Game/Blueprints/AI/Follower_AI_Con.Follower_AI_CON_C')
-		  LinkedTo(0)=EdGraphPin'K2Node_CallFunction_2631.EdGraphPin_20150'
-	   End Object
 	   VariableReference=(MemberName="AI_CON_Ref",bSelfContext=True)
-	   Pins(0)=EdGraphPin'EdGraphPin_20147'
-	   Pins(1)=EdGraphPin'EdGraphPin_20148'
-	   NodePosX=-304
-	   NodePosY=-112
-	   NodeGuid=C902475B4023D7B2028EF292B6FA0539
+	   Pins(0)=EdGraphPin'EdGraphPin_20067'
+	   Pins(1)=EdGraphPin'EdGraphPin_20068'
+	   NodePosX=592
+	   NodePosY=208
+	   NodeGuid=A10A3F5048238FE8D6B7E9A8DFF21873
 	End Object
-	Begin Object Class=K2Node_CallFunction Name="K2Node_CallFunction_2631"
-	   Begin Object Class=EdGraphPin Name="EdGraphPin_20152"
+	Begin Object Class=K2Node_IfThenElse Name="K2Node_IfThenElse_100"
+	   Begin Object Class=EdGraphPin Name="EdGraphPin_20069"
 	   End Object
-	   Begin Object Class=EdGraphPin Name="EdGraphPin_20151"
+	   Begin Object Class=EdGraphPin Name="EdGraphPin_20070"
 	   End Object
-	   Begin Object Class=EdGraphPin Name="EdGraphPin_20150"
+	   Begin Object Class=EdGraphPin Name="EdGraphPin_20071"
 	   End Object
-	   Begin Object Class=EdGraphPin Name="EdGraphPin_20149"
+	   Begin Object Class=EdGraphPin Name="EdGraphPin_20072"
 	   End Object
-	   Begin Object Name="EdGraphPin_20152"
-		  PinName="ReturnValue"
-		  PinToolTip="Return Value\nBoolean"
-		  Direction=EGPD_Output
+	   Begin Object Name="EdGraphPin_20069"
+		  PinName="execute"
+		  PinType=(PinCategory="exec")
+		  LinkedTo(0)=EdGraphPin'K2Node_Event_845.EdGraphPin_4173'
+	   End Object
+	   Begin Object Name="EdGraphPin_20070"
+		  PinName="Condition"
 		  PinType=(PinCategory="bool")
 		  DefaultValue="false"
 		  AutogeneratedDefaultValue="false"
-		  LinkedTo(0)=EdGraphPin'K2Node_IfThenElse_95.EdGraphPin_20154'
+		  LinkedTo(0)=EdGraphPin'K2Node_CallFunction_2670.EdGraphPin_20078'
 	   End Object
-	   Begin Object Name="EdGraphPin_20151"
-		  PinName="B"
-		  PinToolTip="B\nObject Reference"
-		  PinType=(PinCategory="object",PinSubCategoryObject=Class'/Script/CoreUObject.Object')
+	   Begin Object Name="EdGraphPin_20071"
+		  PinName="then"
+		  PinFriendlyName="true"
+		  Direction=EGPD_Output
+		  PinType=(PinCategory="exec")
+		  LinkedTo(0)=EdGraphPin'K2Node_AIMoveTo_42.EdGraphPin_4110'
 	   End Object
-	   Begin Object Name="EdGraphPin_20150"
-		  PinName="A"
-		  PinToolTip="A\nObject Reference"
-		  PinType=(PinCategory="object",PinSubCategoryObject=Class'/Script/CoreUObject.Object')
-		  LinkedTo(0)=EdGraphPin'K2Node_VariableGet_252.EdGraphPin_20147'
+	   Begin Object Name="EdGraphPin_20072"
+		  PinName="else"
+		  PinFriendlyName="false"
+		  Direction=EGPD_Output
+		  PinType=(PinCategory="exec")
+		  LinkedTo(0)=EdGraphPin'K2Node_DynamicCast_1392.EdGraphPin_6847'
+	   End Object
+	   Pins(0)=EdGraphPin'EdGraphPin_20069'
+	   Pins(1)=EdGraphPin'EdGraphPin_20070'
+	   Pins(2)=EdGraphPin'EdGraphPin_20071'
+	   Pins(3)=EdGraphPin'EdGraphPin_20072'
+	   NodePosX=112
+	   NodeGuid=5DE038EB4BDB2B34F7C21B9FA0A97743
+	End Object
+	Begin Object Class=K2Node_VariableGet Name="K2Node_VariableGet_267"
+	   Begin Object Class=EdGraphPin Name="EdGraphPin_20073"
+	   End Object
+	   Begin Object Class=EdGraphPin Name="EdGraphPin_20074"
+	   End Object
+	   Begin Object Name="EdGraphPin_20073"
+		  PinName="AI_CON_Ref"
+		  Direction=EGPD_Output
+		  PinType=(PinCategory="object",PinSubCategoryObject=BlueprintGeneratedClass'/Game/Blueprints/AI/Follower_AI_Con.Follower_AI_CON_C')
+		  LinkedTo(0)=EdGraphPin'K2Node_CallFunction_2670.EdGraphPin_20076'
 	   End Object
-	   Begin Object Name="EdGraphPin_20149"
+	   Begin Object Name="EdGraphPin_20074"
 		  PinName="self"
 		  PinFriendlyName="Target"
-		  PinToolTip="Target\nKismet Math Library Reference"
-		  PinType=(PinCategory="object",PinSubCategoryObject=Class'/Script/Engine.KismetMathLibrary')
-		  DefaultObject=Default__KismetMathLibrary
+		  PinType=(PinCategory="object",PinSubCategoryObject=BlueprintGeneratedClass'/Game/Blueprints/AI/RapidMoveTo.RapidMoveTo_C')
 		  bHidden=True
 	   End Object
-	   bIsPureFunc=True
-	   FunctionReference=(MemberParentClass=Class'/Script/Engine.KismetMathLibrary',MemberName="NotEqual_ObjectObject")
-	   Pins(0)=EdGraphPin'EdGraphPin_20149'
-	   Pins(1)=EdGraphPin'EdGraphPin_20150'
-	   Pins(2)=EdGraphPin'EdGraphPin_20151'
-	   Pins(3)=EdGraphPin'EdGraphPin_20152'
-	   NodePosX=-128
-	   NodePosY=-112
-	   NodeGuid=12101D19401BB291B5226A92EA65C5CB
+	   VariableReference=(MemberName="AI_CON_Ref",bSelfContext=True)
+	   Pins(0)=EdGraphPin'EdGraphPin_20073'
+	   Pins(1)=EdGraphPin'EdGraphPin_20074'
+	   NodePosX=-352
+	   NodePosY=-64
+	   NodeGuid=DD83F3A04F10972DE201B1A9948967EF
 	End Object
-	Begin Object Class=K2Node_IfThenElse Name="K2Node_IfThenElse_95"
-	   Begin Object Class=EdGraphPin Name="EdGraphPin_20156"
+	Begin Object Class=K2Node_CallFunction Name="K2Node_CallFunction_2670"
+	   Begin Object Class=EdGraphPin Name="EdGraphPin_20075"
 	   End Object
-	   Begin Object Class=EdGraphPin Name="EdGraphPin_20155"
+	   Begin Object Class=EdGraphPin Name="EdGraphPin_20076"
 	   End Object
-	   Begin Object Class=EdGraphPin Name="EdGraphPin_20154"
+	   Begin Object Class=EdGraphPin Name="EdGraphPin_20077"
 	   End Object
-	   Begin Object Class=EdGraphPin Name="EdGraphPin_20153"
+	   Begin Object Class=EdGraphPin Name="EdGraphPin_20078"
 	   End Object
-	   Begin Object Name="EdGraphPin_20156"
-		  PinName="else"
-		  PinFriendlyName="false"
-		  Direction=EGPD_Output
-		  PinType=(PinCategory="exec")
-		  LinkedTo(0)=EdGraphPin'K2Node_DynamicCast_1392.EdGraphPin_6838'
+	   Begin Object Name="EdGraphPin_20075"
+		  PinName="self"
+		  PinFriendlyName="Target"
+		  PinToolTip="Target\nKismet Math Library Reference"
+		  PinType=(PinCategory="object",PinSubCategoryObject=Class'/Script/Engine.KismetMathLibrary')
+		  DefaultObject=Default__KismetMathLibrary
+		  bHidden=True
 	   End Object
-	   Begin Object Name="EdGraphPin_20155"
-		  PinName="then"
-		  PinFriendlyName="true"
-		  Direction=EGPD_Output
-		  PinType=(PinCategory="exec")
-		  LinkedTo(0)=EdGraphPin'K2Node_AIMoveTo_42.EdGraphPin_4125'
+	   Begin Object Name="EdGraphPin_20076"
+		  PinName="A"
+		  PinToolTip="A\nObject Reference"
+		  PinType=(PinCategory="object",PinSubCategoryObject=Class'/Script/CoreUObject.Object')
+		  LinkedTo(0)=EdGraphPin'K2Node_VariableGet_267.EdGraphPin_20073'
 	   End Object
-	   Begin Object Name="EdGraphPin_20154"
-		  PinName="Condition"
+	   Begin Object Name="EdGraphPin_20077"
+		  PinName="B"
+		  PinToolTip="B\nObject Reference"
+		  PinType=(PinCategory="object",PinSubCategoryObject=Class'/Script/CoreUObject.Object')
+	   End Object
+	   Begin Object Name="EdGraphPin_20078"
+		  PinName="ReturnValue"
+		  PinToolTip="Return Value\nBoolean"
+		  Direction=EGPD_Output
 		  PinType=(PinCategory="bool")
 		  DefaultValue="false"
 		  AutogeneratedDefaultValue="false"
-		  LinkedTo(0)=EdGraphPin'K2Node_CallFunction_2631.EdGraphPin_20152'
+		  LinkedTo(0)=EdGraphPin'K2Node_IfThenElse_100.EdGraphPin_20070'
 	   End Object
-	   Begin Object Name="EdGraphPin_20153"
-		  PinName="execute"
-		  PinType=(PinCategory="exec")
-		  LinkedTo(0)=EdGraphPin'K2Node_Event_845.EdGraphPin_4188'
-	   End Object
-	   Pins(0)=EdGraphPin'EdGraphPin_20153'
-	   Pins(1)=EdGraphPin'EdGraphPin_20154'
-	   Pins(2)=EdGraphPin'EdGraphPin_20155'
-	   Pins(3)=EdGraphPin'EdGraphPin_20156'
-	   NodePosX=144
-	   NodePosY=-48
-	   NodeGuid=3A5CA2834B6B140CE69753AD1FAE89C6
+	   bIsPureFunc=True
+	   FunctionReference=(MemberParentClass=Class'/Script/Engine.KismetMathLibrary',MemberName="NotEqual_ObjectObject")
+	   Pins(0)=EdGraphPin'EdGraphPin_20075'
+	   Pins(1)=EdGraphPin'EdGraphPin_20076'
+	   Pins(2)=EdGraphPin'EdGraphPin_20077'
+	   Pins(3)=EdGraphPin'EdGraphPin_20078'
+	   NodePosX=-160
+	   NodePosY=-64
+	   NodeGuid=F90931054D514458FE3DDA997BE5C7C2
 	End Object
 [/VAR]
\ No newline at end of file
