#UE4 Docs: fixed release notes error and builds page organization

diff --git a/Engine/Documentation/Source/Support/Builds/Builds.INT.udn b/Engine/Documentation/Source/Support/Builds/Builds.INT.udn
--- a/Engine/Documentation/Source/Support/Builds/Builds.INT.udn
+++ b/Engine/Documentation/Source/Support/Builds/Builds.INT.udn
@@ -13,6 +13,19 @@ Description:Release notes for Unreal Engine builds.
 	[PARAM:description]
 	[/PARAM]
 	[PARAM:links]
+	* [](Support/Builds/ReleaseNotes/2015/4_8)
+	[/PARAM]
+[/OBJECT]
+
+[OBJECT:TopicButtonList]
+	[PARAM:icon]
+	[/PARAM]
+	[PARAM:title]
+	2014
+	[/PARAM]
+	[PARAM:description]
+	[/PARAM]
+	[PARAM:links]
 	* [](Support/Builds/ReleaseNotes/2015/4_7)
 	* [](Support/Builds/ReleaseNotes/2014/4_6)
 	* [](Support/Builds/ReleaseNotes/2014/4_5)
