#UE4 Docs: saved after fix this time

diff --git a/Engine/Documentation/Source/Support/Builds/ReleaseNotes/2015/4_8/ReleaseNotes_4_8.INT.udn b/Engine/Documentation/Source/Support/Builds/ReleaseNotes/2015/4_8/ReleaseNotes_4_8.INT.udn
--- a/Engine/Documentation/Source/Support/Builds/ReleaseNotes/2015/4_8/ReleaseNotes_4_8.INT.udn
+++ b/Engine/Documentation/Source/Support/Builds/ReleaseNotes/2015/4_8/ReleaseNotes_4_8.INT.udn
@@ -541,7 +541,7 @@ Cooked builds can use a dedicated thread for package streaming instead of using
 
 * New: Added new modes for EQS test's scoring/filtering against multiple contexts.
 
-* New: Added new modes for selecting single result from EQS query: random from best items (top 5%) and random from good items (top 25%)
+* New: Added new modes for selecting single result from EQS query: random from best items (top `5%`) and random from good items (top `25%`)
 
 * New: Added support for physx based projection in EQS generators.
 
