#UE4 Docs: formatting fixes

diff --git a/Engine/Documentation/Source/Support/Builds/ReleaseNotes/2015/4_8/ReleaseNotes_4_8.INT.udn b/Engine/Documentation/Source/Support/Builds/ReleaseNotes/2015/4_8/ReleaseNotes_4_8.INT.udn
--- a/Engine/Documentation/Source/Support/Builds/ReleaseNotes/2015/4_8/ReleaseNotes_4_8.INT.udn
+++ b/Engine/Documentation/Source/Support/Builds/ReleaseNotes/2015/4_8/ReleaseNotes_4_8.INT.udn
@@ -1,21 +1,21 @@
 Availability: Public
-Title: 4.8 Release Notes
+Title: Unreal Engine 4.8 Release Notes
 Crumbs: %ROOT%, Engine
 Description:4.8 Release Notes
 
-Unreal Engine 4.8 Release Notes
+[TOC(start:2 end:2)]
 
 This release includes 189 great changes that were submitted from Unreal Engine's amazing community of developers.  Thanks to all of these contributors to Unreal Engine 4.8:
 
 [REDACTED_CONTRIBUTORS]
 
-# Major Features
+## Major Features
 
-## Improved Editor Stability and Quality of Life
+### Improved Editor Stability and Quality of Life
 
 This was a really big focus for Unreal Engine 4.8.  With every Unreal Engine release we fix hundreds of bugs and add many new features, but with Unreal Engine 4.8 we've worked hard to improve editor and game stability, as well as to address many of the most common engine bugs and small feature requests.  We hope that the result is a very solid engine runtime and editor, and we'll continue to make stability a big focus of future releases.
 
-## New:  SteamVR/HTC Vive Support
+### New:  SteamVR/HTC Vive Support
 
 Unreal Engine now has comprehensive support for Valve's SteamVR API and HTC Vive hardware!
 
@@ -27,7 +27,7 @@ Unreal Engine now has comprehensive support for Valve's SteamVR API and HTC Vive
 
 * Full support for both room scale and seated experiences.
 
-## New:  Playstation 4 Morpheus VR with 120Hz Reprojection
+### New:  Playstation 4 Morpheus VR with 120Hz Reprojection
 
 The latest Sony Morpheus hardware is now supported, along with support for 120Hz rendering using temporal reprojection.
 
@@ -37,7 +37,7 @@ The latest Sony Morpheus hardware is now supported, along with support for 120Hz
 
 * 3D audio playback in VR has also been greatly improved for this release.
 
-## New:  Landscape Grass Renderer
+### New:  Landscape Grass Renderer
 
 The new grass system dynamically renders instanced grass and other ground clutter meshes on the landscape around the player.
 
@@ -53,7 +53,7 @@ The new grass system dynamically renders instanced grass and other ground clutte
 
 * See the "Landscape Grass Tools" section below for information about how to set this up!
 
-## New:  Hierarchical Level-Of-Detail System
+### New:  Hierarchical Level-Of-Detail System
 
 This new Hierarchical LOD feature automatically creates LODs for groups of meshes!  This greatly reduces draw calls and polygon counts when they are visible from a distance. This is great for open world games!
 
@@ -67,7 +67,7 @@ This new Hierarchical LOD feature automatically creates LODs for groups of meshe
 
 * LOD meshes are generated as part of the 'Level Build' process, or by selecting *Build LODs* in the Build menu.
 
-## New:  Procedural Foliage Generator (Experimental)
+### New:  Procedural Foliage Generator (Experimental)
 
 There is a new experimental feature for procedurally generating huge (or tiny, or medium-sized) regions of foliage instances quickly.
 
@@ -77,7 +77,7 @@ There is a new experimental feature for procedurally generating huge (or tiny, o
 
 * Multiple procedural foliage types can be simulated alongside one another to generate the foundations of an entire forest!
 
-## New:  Scrubbable Network Replays
+### New:  Scrubbable Network Replays
 
 Unreal Engine games have full game replay recording and playback.  New in 4.8, you can now "scrub" to a specific time while viewing network replays!
 
@@ -113,7 +113,7 @@ While viewing a replay:
 
 Out of the box, the engine is configured to store replays locally. We have a system that allows you to easily swap out replay streamers. This allows you to use a streamer that saves replays locally, or you can use a streamer that is designed to save replays to the cloud with full streaming support. More on that in future releases!
 
-## New:  High Quality Depth of Field and Tone Mapping
+### New:  High Quality Depth of Field and Tone Mapping
 
 ![image alt text](image_6.gif)(convert:false)
 
@@ -135,7 +135,7 @@ Out of the box, the engine is configured to store replays locally. We have a sys
 
     * r.TonemapperFilm 1 to enable
 
-## New:  Redesigned Foliage Editor
+### New:  Redesigned Foliage Editor
 
 The foliage editor has recently undergone a redesign!  Besides improved usability, the new layout gives you the ability to more easily work with a large number of foliage types in the palette at once.
 
@@ -147,7 +147,7 @@ The foliage editor has recently undergone a redesign!  Besides improved usabilit
 
 * You can now activate or deactivate multiple foliage types at once with multi-select.
 
-## New:  Dynamic Navigation Mesh for Open Worlds
+### New:  Dynamic Navigation Mesh for Open Worlds
 
 We have a brand new system for supporting dynamic NavMeshes in large open worlds.
 
@@ -157,7 +157,7 @@ We have a brand new system for supporting dynamic NavMeshes in large open worlds
 
 * To enable this, configure your project's Navigation System for this mode and add a Navigation Invoker Component to relevant actors.
 
-## New:  New Unreal Example Character
+### New:  New Unreal Example Character
 
 We've designed a new, high resolution version of the Unreal Example Character included with our templates!
 
@@ -167,7 +167,7 @@ We've designed a new, high resolution version of the Unreal Example Character in
 
 * The high resolution version is included -- you can grab the low resolution character from the Marketplace for free!
 
-## New:  Curve Editing for Unreal Motion Graphics UI
+### New:  Curve Editing for Unreal Motion Graphics UI
 
 Unreal Motion Graphics now has a curve editor!
 
@@ -177,7 +177,7 @@ Unreal Motion Graphics now has a curve editor!
 
 * You can also play, rewind and loop your animation while previewing it using the new Transport control buttons.
 
-## New:  Screen-Space Reflections for Translucent Surfaces
+### New:  Screen-Space Reflections for Translucent Surfaces
 
 Translucent surfaces can now dynamically reflect other objects on the screen!
 
@@ -185,7 +185,7 @@ Translucent surfaces can now dynamically reflect other objects on the screen!
 
 ![image alt text](image_11.jpg)
 
-## New:  Quick Window Switching (Control+Tab)
+### New:  Quick Window Switching (Control+Tab)
 
 You can now summon a pop-up menu to quickly switch between different asset editors and tool tabs by pressing Ctrl+Tab.
 
@@ -195,7 +195,7 @@ You can now summon a pop-up menu to quickly switch between different asset edito
 
 * You can also activate a tool tab in the editor for the current selected asset or show the asset in the content browser by using the mouse
 
-## New:  Landscape Grass Tools
+### New:  Landscape Grass Tools
 
 New editor features have been added for setting up instanced meshes for the grass rendering system.
 
@@ -207,7 +207,7 @@ New editor features have been added for setting up instanced meshes for the gras
 
 * Connecting a Landscape Layer Sample node directly will allow a landscape layer to specify the grass density. You can also connect a network of nodes to create more complex placement rules such as automatically removing grass based on surface normal or altitude or using multiple landscape layers.
 
-## New:  Edit Properties on Instances of Actor Blueprints
+### New:  Edit Properties on Instances of Actor Blueprints
 
 Just as you can in subclasses, properties of components added in the components section of the Blueprint can be modified on a per-actor instance basis.
 
@@ -215,7 +215,7 @@ Just as you can in subclasses, properties of components added in the components
 
 * A few caveats:  Components created by the Construction Script cannot be modified, and properties that are modified on components by the Construction Script cannot be modified on the instance and will appear as read-only in the Actor details panel. However, if the Construction Script explicitly sets the property value to the default, we cannot distinguish that from not having set it at all and it will be editable at the instance level.
 
-## New:  Asset Size Mapper
+### New:  Asset Size Mapper
 
 Ever wondered how much memory an asset takes up? Well, we have a new editor feature called "**Size Mapper**" that helps you figure that out!
 
@@ -235,13 +235,13 @@ Ever wondered how much memory an asset takes up? Well, we have a new editor feat
 
 * **Any asset editor**:  Open the File menu and choose "Size Map."
 
-## Shared Editor Layout and Keybinds
+### Shared Editor Layout and Keybinds
 
 Your editor layout and keybinding settings are now shared between all of your projects!  
 
 Previously, every project had its own saved layout/keybinding settings, and it was cumbersome to rebuild your layout with every project you switch to.  In the future we'll be moving even more editor preferences to be globally shared between all projects.
 
-## New:  Interactive Wires in Graph Editors
+### New:  Interactive Wires in Graph Editors
 
 You can now interact with wires directly in any graph editor!
 
@@ -255,7 +255,7 @@ You can now interact with wires directly in any graph editor!
 
 * Double-click to insert a reroute node (Blueprint graphs only)
 
-## New:  Procedural Mesh Component (Experimental)
+### New:  Procedural Mesh Component (Experimental)
 
 The new Procedural Mesh Component lets you create meshes at runtime using C++ or Blueprints. 
 
@@ -265,13 +265,13 @@ The new Procedural Mesh Component lets you create meshes at runtime using C++ or
 
 * We have also added a utility library, including a function to generate smooth normals and tangents, given mesh vertex positions and UVs.
 
-## New:  Improved Automatic Convex Collision Meshes
+### New:  Improved Automatic Convex Collision Meshes
 
 When using the Static Mesh Editor to automatically generate a collision mesh for your model, the generated mesh will be of much higher quality and will complete more quickly!
 
 ![image alt text](image_17.jpg)
 
-## New:  Configurable Animation Notify Triggers
+### New:  Configurable Animation Notify Triggers
 
 There are helpful new settings on all notifies to control when animation notifies trigger:
 
@@ -285,7 +285,7 @@ There are helpful new settings on all notifies to control when animation notifie
 
     * *LOD *- Only trigger this notify if the owning mesh is below this LOD. For example a setting of 2 will mean that the LOD will only trigger if the mesh component is rendering at LODs 0 and 1
 
-## New:  On-Demand Asset Cooking ("Cook on the Side")
+### New:  On-Demand Asset Cooking ("Cook on the Side")
 
 * "Cook on the Side" is a new Project Setting that allows your game to startup on a device very quickly.
 
@@ -297,7 +297,7 @@ There are helpful new settings on all notifies to control when animation notifie
 
 * Previously to 4.8, this feature was available as an Experimental setting.  It's now fully supported in 4.8!
 
-## New:  **Downloadable Content (DLC) Support**
+### New:  Downloadable Content (DLC) Support
 
 * Added support for cooking DLC (Downloadable Content).  This allows you to create Pak files that can be distributed to users later.
 
@@ -305,7 +305,7 @@ There are helpful new settings on all notifies to control when animation notifie
 
     * In Unreal Frontend, make a custom profile, then select Cook -> Release / DLC / Patch Settings -> Include Engine Content. 
 
-## New:  **Content Patching Support**
+### New:  Content Patching Support
 
 You can now generate Patch Files for your game that have only the content files that have changed since a previous release! 
 
@@ -317,7 +317,7 @@ You can now generate Patch Files for your game that have only the content files
 
 * To generate a patch off that release use the "generate patch" check box in project launcher
 
-## New:  Localization Dashboard
+### New:  Localization Dashboard
 
 The Localization Dashboard is now usable as an experimental feature. Use this feature to manage your project's localization pipeline.
 
@@ -327,7 +327,7 @@ The Localization Dashboard is now usable as an experimental feature. Use this fe
 
 * You can also manage outstanding issues with any of the existing localization data, as well as control all aspects of localization data gathering here.
 
-## Paper2D:  Paint Directly on Sprite Textures
+### Paper2D:  Paint Directly on Sprite Textures
 
 Sprite components can be selected and painted on when in Mesh Paint mode
 
@@ -337,7 +337,7 @@ Sprite components can be selected and painted on when in Mesh Paint mode
 
 * Note: Only texture painting mode is supported, vertex painting does not currently work with sprite components.
 
-## Paper2D:  Tile Maps are promoted out of experimental into Early Access Preview
+### Paper2D:  Tile Maps are promoted out of experimental into Early Access Preview
 
 ![image alt text](image_22.png)
 
@@ -355,7 +355,7 @@ Sprite components can be selected and painted on when in Mesh Paint mode
 
 	![image alt text](image_23.png)
 
-## Paper2D:  Support for sprites that require multiple paired textures
+### Paper2D:  Support for sprites that require multiple paired textures
 
 Materials for sprites can define additional texture 'slots', with the actual textures defined in each sprite asset, allowing multi-texture sprites to be easily used.
 
@@ -371,7 +371,7 @@ Materials for sprites can define additional texture 'slots', with the actual tex
 
 * Note: Each texture does not have to be the same size, but the portion corresponding to a specific sprite must be at the same relative location in each (there is only one source region defined per sprite)
 
-## Paper2D:  Snap layers for speedy 2D level layout
+### Paper2D:  Snap layers for speedy 2D level layout
 
 2D Snap Layers are an easy way to manage level layout when working with a 2D game (or a 3D game with 2D gameplay)
 
@@ -405,7 +405,7 @@ Materials for sprites can define additional texture 'slots', with the actual tex
 
 * Cycle active snap layer backwards (Alt+PageDown)
 
-## Paper2D:  Interactive Extract Sprites Dialog
+### Paper2D:  Interactive Extract Sprites Dialog
 
 The Extract sprites dialog allows more control over how sprites are split out from a texture, making it quicker and easier to import/prepare assets for use with Paper2D
 
@@ -421,7 +421,7 @@ The Extract sprites dialog allows more control over how sprites are split out fr
 
     * Grid-based allows you to specify the width/height of one grid cell, and optionally how many cells wide/tall to extract (in case there are different sized grids sharing the same texture)
 
-## Paper2D:  Grouped Sprite Components
+### Paper2D:  Grouped Sprite Components
 
  A grouped sprite component allows many instances of one or more sprite assets to be drawn together in a minimal number of batches (separate materials or textures will still split batches).
 
@@ -445,7 +445,7 @@ The Extract sprites dialog allows more control over how sprites are split out fr
 
     * These kinds of sorting issues only apply to Translucent materials, Masked materials don't have the same issues but they only work for binary (0 or 1) opacity.
 
-## New:  'Interpolate To' Component
+### New:  'Interpolate To' Component
 
 The new "Interpolate To" Movement Component can be dropped onto actors.  This will let the actor move automatically!
 
@@ -453,7 +453,7 @@ The new "Interpolate To" Movement Component can be dropped onto actors.  This wi
 
 * You can specify multiple points to interpolate between, and a number of different behaviour options are available (e.g. looping, ping pong, etc.)
 
-##New:  Support for Large Android Games
+### New:  Support for Large Android Games
 
 Added support for downloading .obb (large data) files from the Google Play Store has been added.
 
@@ -463,7 +463,7 @@ Added support for downloading .obb (large data) files from the Google Play Store
 
 * Allows for a custom background image support for the download screen.
 
-##New:  Build iOS Projects from Windows (Remotely)
+### New:  Build iOS Projects from Windows (Remotely)
 
 Added support for building iOS Code Projects on a remote Mac using SSH.
 
@@ -483,7 +483,7 @@ Added support for building iOS Code Projects on a remote Mac using SSH.
 
     * The previous, but unsupported, method (UnrealRemoteTool) can still be used if you set bUseRsync to false DefaultEngine.ini
 
-## New:  Motion Blur and Anti-aliasing for Vertex Animation Materials
+### New:  Motion Blur and Anti-aliasing for Vertex Animation Materials
 
 Materials that animate vertices using World Position Offset can now output accurate motion vectors.  The motion vectors allow per-pixel motion blur from the procedural animation. Additionally Temporal Anti-Aliasing and Distance Field Ambient Occlusion can correctly reproject to prevent undesirable blurring.  
 
@@ -499,7 +499,7 @@ Visualization of per-pixel velocities from a tree animated with World Position O
 
 * Note that this feature can only support velocities from changing time, not parameters set by blueprint or gameplay code. 
 
-## New:  Multi-threaded Asynchronous Asset Loading
+### New:  Multi-threaded Asynchronous Asset Loading
 
 Cooked builds can use a dedicated thread for package streaming instead of using time-slicing approach.  You can enable this in your project with a new config option.  This can make asynchronous loading of assets up to three times faster!
 
@@ -521,9 +521,9 @@ Cooked builds can use a dedicated thread for package streaming instead of using
 
 * Async load percentage accessible via 'stat levels' now takes PostLoad into account
 
-# Release Notes
+## Release Notes
 
-## AI
+### AI
 
 * AI: Blueprint Support for EQS
 
@@ -583,13 +583,13 @@ Cooked builds can use a dedicated thread for package streaming instead of using
 
 * Removed redundant data gathering from NavMesh, to optimize NavMesh rendering.
 
-### Debugging Tools
+#### Debugging Tools
 
 * Fixed a bug in Gameplay Debugger that would result in a crash when used tried to debug a Pawn and there was no navigation data present on the map.
 
 * Fixed crash on collecting Visual Log snapshot from uninitialized blackboard.
 
-### Navigation
+#### Navigation
 
 * New: A way to limit number of simultaneously rebuild NavMesh tiles has been added.
 
@@ -697,7 +697,7 @@ Cooked builds can use a dedicated thread for package streaming instead of using
 
 * Navmesh generation issues with World Composition have been fixed
 
-## Animation
+### Animation
 
 * New: Improved Animation Slot Editing
 
@@ -895,7 +895,7 @@ Cooked builds can use a dedicated thread for package streaming instead of using
 
 * Prevented a crash when a skeletal mesh LOD has zero vertices
 
-## Audio
+### Audio
 
 * New: A unique audio device instance will now be created for each PIE window. 
 
@@ -933,7 +933,7 @@ Cooked builds can use a dedicated thread for package streaming instead of using
 
 * Stop any sound cue previews playing when sound cue editor closes.
 
-## Automation
+### Automation
 
 * New: Added additive flag to rotation multiplier bone controller node in animation graphs
 
@@ -963,7 +963,7 @@ Cooked builds can use a dedicated thread for package streaming instead of using
 
 * The Load All Maps in Editor automation test will no longer run on the maps found in the engine content folder.
 
-## Blueprints
+### Blueprints
 
 * Blueprints: Drag and Drop to Create New Pins
 
@@ -1581,7 +1581,7 @@ Cooked builds can use a dedicated thread for package streaming instead of using
 
 * When pasting Event nodes in Blueprints, if there is a ghost node version of the Event, that node will be removed in favor of the newly pasted Event.
 
-# Core
+### Core
 
 * New: Added HSV-based linear color interpolation function (Blueprint and C++)
 
@@ -1777,7 +1777,7 @@ Cooked builds can use a dedicated thread for package streaming instead of using
 
 * XML parser can now parse nested '=' characters and named entities correctly in attribute strings.
 
-## Editor and Tools
+### Editor and Tools
 
 * New: Added the ability to specify build settings and enable plugins for content-only projects.
 
@@ -2393,7 +2393,7 @@ Cooked builds can use a dedicated thread for package streaming instead of using
 
 * Writing a default config file with no hierarchy no longer causes a crash.
 
-## Content Browser
+### Content Browser
 
 * New: Added a prototype configurable asset tag filter
 
@@ -2443,7 +2443,7 @@ Cooked builds can use a dedicated thread for package streaming instead of using
 
 * When discovering assets, we are now prioritizing background search results in addition to background search files after selecting a path in an asset view.
 
-## Landscape
+### Landscape
 
 * New: Landscape splines can now split their meshes between landscape streaming levels
 
@@ -2541,7 +2541,7 @@ Cooked builds can use a dedicated thread for package streaming instead of using
 
 * The landscape tools no longer crash when toggling a sublevel's visibility in the Levels window while Landscape mode is active
 
-## Material Editor
+### Material Editor
 
 * New: A new material function 'TransformToZVector' is useful for transforming vectors from worldspace into an a space defined by a Z vector. The X and Y basis vectors will be generated as orthogonal to the input Z vector.
 
@@ -2581,7 +2581,7 @@ Cooked builds can use a dedicated thread for package streaming instead of using
 
 * When compiling materials, the engine now no longer completely re-registers any components using the material, and instead just recreates the components' render state. This reduces the time taken when pressing the Apply button in the material editor.
 
-## Matinee
+### Matinee
 
 * New: Added ability to use a custom Anim Blueprint for Skeletal Meshes driven by Matinee
 
@@ -2615,7 +2615,7 @@ Cooked builds can use a dedicated thread for package streaming instead of using
 
 * Undoing a rename of a camera group now correctly undoes the same name shown in the director group
 
-## Cascade
+### Cascade
 
 * Emitters with disabled LODs to remain alive so that when re-entering LOD range, the effect became visible again.
 
@@ -2629,7 +2629,7 @@ Cooked builds can use a dedicated thread for package streaming instead of using
 
 * Fixed the Spawn Using Surface Normal feature in the Skel/Vert/Surf module.
 
-## Persona
+### Persona
 
 * New: Added a scroll bar to the notify panel to scroll through the track when zoomed
 
@@ -2673,7 +2673,7 @@ Cooked builds can use a dedicated thread for package streaming instead of using
 
 * Fixed undo for changes in anim notify state duration not undoing the change
 
-## PhAT
+### PhAT
 
 * New: Added PhAT support for scaled bones.
 
@@ -2681,13 +2681,13 @@ Cooked builds can use a dedicated thread for package streaming instead of using
 
 * When building collision for a bone, correctly apply component space matrix to propagate scale.
 
-## Project Browser
+### Project Browser
 
 * New: Made the Project Browser window a little taller so that it's clear that there's description text below the template image.
 
 * New projects now have a random project ID set by default.
 
-## World Outliner
+### World Outliner
 
 * Deleting folders in the World Outliner now promotes all children to the level above instead of directly to the root.
 
@@ -2699,7 +2699,7 @@ Cooked builds can use a dedicated thread for package streaming instead of using
 
     * As before, the Level Blueprint can still be opened from the Blueprints drop-down in the toolbar.
 
-## Static Mesh Editor
+### Static Mesh Editor
 
 * New: Added socket snap-to-vertex, and socket metadata.
 
@@ -2715,7 +2715,7 @@ Cooked builds can use a dedicated thread for package streaming instead of using
 
 * Improved the performance of scaling collision volumes in the static mesh editor.
 
-## Source Control
+### Source Control
 
 * New: Added the ability to commit file deletions from the editor.
 
@@ -2741,7 +2741,7 @@ Cooked builds can use a dedicated thread for package streaming instead of using
 
 * Allows Git provider to display revisions as hashes.
 
-## World Browser
+### World Browser
 
 * New: Sub-levels now have a property to hide them from world composition tile view. "Hide in Tile View" in sub-level details panel.
 
@@ -2765,7 +2765,7 @@ Cooked builds can use a dedicated thread for package streaming instead of using
 
 * Play in Editor (PIE) session will no longer crash when the player initiates travel from a single map to a world composition map.
 
-## Foliage
+### Foliage
 
 * New: Added right-click menu items to the Foliage Type list allowing you to Select or Deselect all foliage instances of the foliage type
 
@@ -2823,7 +2823,7 @@ Cooked builds can use a dedicated thread for package streaming instead of using
 
 * The foliage tool now detects if a shared Foliage Type's properties (eg collision or lighting) were changed while a level using the shared foliage type was not loaded. Any changes are now propagated when the level is next loaded in the editor.
 
-## Cooker
+### Cooker
 
 * New: Added ability to specify a subset of maps to cook when packaging. This is used to make File/Package _not_ cook your entire content. See Packaging section of Project Settings.
 
@@ -2869,7 +2869,7 @@ Cooked builds can use a dedicated thread for package streaming instead of using
 
 * Made some adjustments to make cooking faster for projects with a high object count.
 
-## Gameplay Framework
+### Gameplay Framework
 
 * New: Game Mode Customization in Blueprints
 
@@ -2919,7 +2919,7 @@ Cooked builds can use a dedicated thread for package streaming instead of using
 
 * Sub-levels marked as "AlwaysLoaded" will now be fully loaded before gameplay code attempts to find spawn point for players.
 
-## Learning Resources
+### Learning Resources
 
 * New: Added a brief tutorial on using feature packs
 
@@ -2929,13 +2929,13 @@ Cooked builds can use a dedicated thread for package streaming instead of using
 
 * Templates: Fixed an issue with animation state in the first person BP template
 
-## In-Editor Help
+### In-Editor Help
 
 * Fixed docked tutorial browser never showing tutorials
 
 * Fixed stalls in tutorials
 
-## Content Examples
+### Content Examples
 
 * Added missing space in text at section Audio 1.8 of Content Examples
 
@@ -2955,13 +2955,13 @@ Cooked builds can use a dedicated thread for package streaming instead of using
 
 * Updated and repaired normals on CobbleStone material in StarterContent.
 
-## Sample Content
+### Sample Content
 
 * New: 2D and 3D side scroller templates now enable snapping layers by default, with three pre-defined layers at the correct depth for the sample content
 
 * Repaired z-fighting occurring on directional arrows in StrategyGame.
 
-## Localization
+### Localization
 
 * New: All localization configuration scripts now treat path settings consistently as being relative to the root directory of the project.
 
@@ -3013,7 +3013,7 @@ Cooked builds can use a dedicated thread for package streaming instead of using
 
 * Translation Editor: When saving, if source control is disabled, try write the file anyway.
 
-## Networking
+### Networking
 
 * New: Ability to disable replays on command line. Use -noreplays.
 
@@ -3071,7 +3071,7 @@ Cooked builds can use a dedicated thread for package streaming instead of using
 
 * When shutting down actor channels, all registered sub-objects will now get destroyed as well.
 
-## Online
+### Online
 
 * New: Added a "Relevant For Network Replays" property on Actor which, if set to false, will prevent the actor from saving out to the network replay.
 
@@ -3085,7 +3085,7 @@ Cooked builds can use a dedicated thread for package streaming instead of using
 
 * Network replays now only load streaming levels with the "Should Be Loaded" option selected.
 
-## Other
+### Other
 
 * New: Added a Flurry analytics plugin for IOS
 
@@ -3241,7 +3241,7 @@ Cooked builds can use a dedicated thread for package streaming instead of using
 
 * Tweaked layout of a few commonly-used data structures to reduce memory footprint and increase efficiency.
 
-## Paper2D
+### Paper2D
 
 * New: Sprite Editor: Added a corner text display to the viewport indicating the current mode when not in View mode
 
@@ -3531,7 +3531,7 @@ Cooked builds can use a dedicated thread for package streaming instead of using
 
 * Tile Maps: Fixed the pivot show flag not working in the stand-alone tile map editor
 
-## Physics
+### Physics
 
 * New: Added ability to enable and disable constraint projection at runtime
 
@@ -3667,7 +3667,7 @@ Cooked builds can use a dedicated thread for package streaming instead of using
 
     * Fixes Characters being unable to step up simple box geometry, among other issues.
 
-## Platforms
+### Platforms
 
 * New: Added EXPERIMENTAL "Desktop (Win + Mac + Linux)" Target Platform which will allow you to cook content that can run on all 3 platforms.
 
@@ -3717,7 +3717,7 @@ Cooked builds can use a dedicated thread for package streaming instead of using
 
 * The Virtual Joystick no longer takes takes keyboard focus. This prevents the problem of it sealing focus from the game viewport and resetting the touch state of fingers that have been pressed outside of the joystick.
 
-### Android
+#### Android
 
 * New: Added application display name (app_name) setting for Android to Project Settings
 
@@ -3833,7 +3833,7 @@ Cooked builds can use a dedicated thread for package streaming instead of using
 
 * Protection added against in-app purchase disabled to onActivityResult
 
-### iOS
+#### iOS
 
 * New: Added Delta-copy install location as a configurable property in the IOS project settings.
 
@@ -3877,7 +3877,7 @@ Cooked builds can use a dedicated thread for package streaming instead of using
 
 * Properly prefer iTunes 12 Mobile Device dll over the iTunes 11 dll. This fixed some crashes related to picking the wrong version of the dll when both are present.
 
-### Linux
+#### Linux
 
 * New: Available GPU memory size (when an appropriate extension is provided by the driver) will now be used for better control over texture streaming.
 
@@ -3979,7 +3979,7 @@ Cooked builds can use a dedicated thread for package streaming instead of using
 
 * -Wreorder warnings will no longer cause the build to fail (these warnings aren't suppressed, just exempted from -Werror).
 
-### Mac
+#### Mac
 
 * New: Added compatibility with Mono 3.12
 
@@ -4141,7 +4141,7 @@ Cooked builds can use a dedicated thread for package streaming instead of using
 
 * Use 'which mono' to discover local Mono installations on Mac OS X rather than assuming a fixed path.
 
-### Playstation 4
+#### Playstation 4
 
 * New: Debug builds now link in the debug versions of Gnm and Gnmx.
 
@@ -4271,7 +4271,7 @@ Cooked builds can use a dedicated thread for package streaming instead of using
 
 * Shadow Quality console variable is now defaulted correctly.
 
-### Xbox One
+#### Xbox One
 
 * New: Added catching of unhandled exceptions on Xbox One.
 
@@ -4355,7 +4355,7 @@ Cooked builds can use a dedicated thread for package streaming instead of using
 
 * Removed an unnecessary and slow memory barrier function.
 
-### VR
+#### VR
 
 * New: GearVR performance settings can now be set in Engine.ini. Under [GearVR.Settings] values can be set for CpuLevel and GpuLevel. Defaults are 2,2. See the Oculus Mobile SDK documentation for valid values.
 
@@ -4375,7 +4375,7 @@ Cooked builds can use a dedicated thread for package streaming instead of using
 
 * Fixed running start for SteamVR.
 
-### HTML5
+#### HTML5
 
 * New: Experimental: Allow HTML5 games to download maps on the fly. For smaller startup packages and time to boot [ Can be enabled from Project Settings -> HTML5 ]
 
@@ -4411,7 +4411,7 @@ Cooked builds can use a dedicated thread for package streaming instead of using
 
 * Mutilpe fixes for HTML5's implementation of HttpModules.
 
-### Windows
+#### Windows
 
 * New: 32-bit and 64-bit Windows are now both fully supported as target platforms in the binary release (Win32 support was previously limited to shipping configurations, and Win64 support was previously limited to development configurations)
 
@@ -4427,7 +4427,7 @@ Cooked builds can use a dedicated thread for package streaming instead of using
 
 * Integrated fix for portmidi crash on Windows 8.1
 
-## Rendering
+### Rendering
 
 * New: Added ability to have some object (landscape or FarShadow marked objects) casing shadows in extra shadow cascades
 
@@ -4651,7 +4651,7 @@ Cooked builds can use a dedicated thread for package streaming instead of using
 
 * Worked around a shader bug that occurs on some mobile GPUs in which an incorrect normalize result can be produced by large view vectors.
 
-## FX
+### FX
 
 * New: Added support for Ctrl+B and Ctrl+E shortcuts to find/edit the associated asset for selected Niagara actors
 
@@ -4667,7 +4667,7 @@ Cooked builds can use a dedicated thread for package streaming instead of using
 
 * Particle Systems spawned from Blueprints and directly in code via SpawnEmitterAtLocation now work correctly with collision.
 
-## Lighting
+### Lighting
 
 * New: Added color temperature property for lights
 
@@ -4775,7 +4775,7 @@ Cooked builds can use a dedicated thread for package streaming instead of using
 
 * Fixed some shadows disappearing with large terrain
 
-## Materials
+### Materials
 
 * New: Added a "mutli-lerp" function that blends up to 4 float 3 values.
 
@@ -4839,7 +4839,7 @@ Cooked builds can use a dedicated thread for package streaming instead of using
 
 * Unlit materials can now use the Normal input when also using refraction.
 
-## Postprocessing
+### Postprocessing
 
 * New: Added another bloom stage (gaussian) for larger bloom
 
@@ -4861,7 +4861,7 @@ Cooked builds can use a dedicated thread for package streaming instead of using
 
 * Fixed PostProcessMaterial SwitchParameter behavior
 
-## Optimizations
+### Optimizations
 
 * New: Added a new 'BestPlus' configuration for SpeedTree shaders. 'Best' is now faster.
 
@@ -4879,7 +4879,7 @@ Cooked builds can use a dedicated thread for package streaming instead of using
 
 * Removed ClearCoat material handling from shaders if there are no ClearCoat materials in the given view.
 
-## UI
+### UI
 
 * New: Added a new ListItemAlignment option to the SListPanel that allows for STileView items to stretch horizontally instead of increasing the padding between the items.
 
@@ -4903,7 +4903,7 @@ Cooked builds can use a dedicated thread for package streaming instead of using
 
 * Greatly improved performance of drawing wrapped text using the canvas.
 
-## Slate
+### Slate
 
 * New: Added constructor to FSlateTextureData that takes a pointer to texture data
 
@@ -4967,7 +4967,7 @@ Cooked builds can use a dedicated thread for package streaming instead of using
 
 * We no longer clear focus on child widgets of the game viewport when the window loses focus by default. This means controllers will still be able to interact with the game when you give focus to another app, like visual studio.
 
-## Unreal Motion Graphics (UMG UI Editor)
+### Unreal Motion Graphics (UMG UI Editor)
 
 * New:  Split-Screen UI with Unreal Motion Graphics
 
@@ -5191,9 +5191,9 @@ Cooked builds can use a dedicated thread for package streaming instead of using
 
 * We now use bilinear sampling rather than point sampling.
 
-# Programming Release Notes
+## Programming Release Notes
 
-## AI
+### AI
 
 * New: Added versions of UEnvQueryTest::GetItem[Actor/Location/Rotation] parametric with FEnvQueryInstance::ItemIterator to avoid conversion-related issues when accessing query's items with item iterator
 
@@ -5231,7 +5231,7 @@ Cooked builds can use a dedicated thread for package streaming instead of using
 
     * A serialization issue of NavigationDataClass in Navigation System's properties has been solved so that game-specific classes can now be used for these properties.
 
-### Behavior Tree
+#### Behavior Tree
 
 * New: Added instancing for blackboard key types to support custom memory allocations.
 
@@ -5241,13 +5241,13 @@ Cooked builds can use a dedicated thread for package streaming instead of using
 
 * Removed failsafe Tick call from TimeLimit decorator, since it's no longer needed.
 
-### Debugging Tools
+#### Debugging Tools
 
 * New: Added a function to VisualLog allowing users to add arbitrary FVisualLogShapeElement instances to given log entry which results in even more flexibility in logging spatial data.
 
 * New: Made AIController automatically redirect all visual logs of its components to controller's log
 
-### Navigation
+#### Navigation
 
 * New: Added a generic support for navigation relevant actors being able to supply NavMesh generation process with subsections of their collision geometry
 
@@ -5295,7 +5295,7 @@ Cooked builds can use a dedicated thread for package streaming instead of using
 
 * Moved navigation-related undo/redo handling from LandscapeProxy to Actor. Now all actors should handle navigation-related undo/redo consistently.
 
-## Animation
+### Animation
 
 * Performance: USkinnedMeshComponent::AnimUpdateRateTick() optimized to reduce potential dynamic allocations. 0.40ms -> 0.26ms in test scene with many characters.
 
@@ -5367,7 +5367,7 @@ Cooked builds can use a dedicated thread for package streaming instead of using
 
 * USkeletalMeshComponent::RecalcRequiredBones's check for Sockets could add the same bone indices multiple times, triggering an assert in FillSpaceBases() later. Made sure all bones added to ForceAnimatedSocketBones and NeededBonesForFillSpaceBases are unique.
 
-## Audio
+### Audio
 
 * New: Export AudioVolume class to allow subclassing.
 
@@ -5381,7 +5381,7 @@ Cooked builds can use a dedicated thread for package streaming instead of using
 
 * Removed the unused sound preview dialog and sound preview thread
 
-## Blueprints
+### Blueprints
 
 * New: Added a fast calling mechanism for no-argument thunks into the event graph
 
@@ -5425,7 +5425,7 @@ Cooked builds can use a dedicated thread for package streaming instead of using
 
 * Fixed UScriptFactory not returning created script blueprint object
 
-## Core
+### Core
 
 * Initial support for TMap properties!
 
@@ -5847,7 +5847,7 @@ Cooked builds can use a dedicated thread for package streaming instead of using
 
 * UnrealHeaderTool: Added limited support for parsing character literals to avoid incorrectly erroring on things like '\"'
 
-## Editor and Tools
+### Editor and Tools
 
 * New: Improved the performance of FText::AsNumber
 
@@ -6075,7 +6075,7 @@ Cooked builds can use a dedicated thread for package streaming instead of using
 
 * UnrealHeaderTool will no longer try to automatically declare a constructor with FObjectInitializer if the base class does not declare one too.
 
-## Content Browser
+### Content Browser
 
 * New: Added support for saving unique per-filter state in a frontend filter by implementing SaveSettings and LoadSettings
 
@@ -6087,7 +6087,7 @@ Cooked builds can use a dedicated thread for package streaming instead of using
 
 * New: Opening a C++ class via the Content Browser now opens the source file as well as the header file
 
-## Landscape
+### Landscape
 
 * Removed "#pragma clang diagnostic ignored "-Wreorder"" in landscape editor data caches
 
@@ -6097,21 +6097,21 @@ Cooked builds can use a dedicated thread for package streaming instead of using
 
 * Fixed incorrect spelling of the ULandscapeHeightfieldCollisionComponent::CookCollisionData function name in C++ code
 
-## Project Browser
+### Project Browser
 
 * New: Added some verbose logging to help diagonse issues when creating new project templates
 
-## Static Mesh Editor
+### Static Mesh Editor
 
 * Deleted some dead lightmap mesh import/export code in StaticMeshEditor and UnrealEd
 
-## Foliage
+### Foliage
 
 * New: Moved Foliage into its own module
 
 * Removed the InstanceInverseTransform element from the FInstancingUserData struct used by InstancedStaticMeshComponent. This data was incorrectly calculated, wasted memory and was not used by the engine.
 
-## Gameplay Framework
+### Gameplay Framework
 
 * New Component 'Editable When Inherited' flag
 
@@ -6179,11 +6179,11 @@ Cooked builds can use a dedicated thread for package streaming instead of using
 
 * Added new attachment location option called EAttachLocation::SnapToTargetIncludingScale ("Snap to Target, Including Scale" in the editor).
 
-## Documentation
+### Documentation
 
 * Added warning to documentation for AActor::GetComponents() about the possibility of the array changing while iterating its contents.
 
-## Localization
+### Localization
 
 * New: The text from UI_COMMAND macros will be nested under the namespace specified by LOCTEXT_NAMESPACE, nested under the "UICommands" namespace. IE: "UICommands.MyNamespace"
 
@@ -6191,7 +6191,7 @@ Cooked builds can use a dedicated thread for package streaming instead of using
 
 * Tidied up some old localization code by removing the following: PPF_LocalizedOnly, CPF_Localized, CLASS_Localized, UProperty::IsLocalized.
 
-## Networking
+### Networking
 
 * Network Replay Compatibility
 
@@ -6257,7 +6257,7 @@ Cooked builds can use a dedicated thread for package streaming instead of using
 
 * UpdatePing and RecalculateAvgPing are now virtual
 
-## Online
+### Online
 
 * New: Added a new function to the online sessions API to find a session by id. Not implemented by most platforms yet.
 
@@ -6273,7 +6273,7 @@ Cooked builds can use a dedicated thread for package streaming instead of using
 
 * Crash fix for BuildPatchTool following FRunnableThread::InitializeTls requirement
 
-## Other
+### Other
 
 * New: Added navigation filter support in EQS pathfinding tests.
 
@@ -6291,7 +6291,7 @@ Cooked builds can use a dedicated thread for package streaming instead of using
 
 * New: Added new console command ?'ShowDebugForReticleTargetToggle <DesiredClass>' Toggles 'ShowDebug' from showing debug info between reticle target actor (of subclass <DesiredClass>) and camera view target.
 
-## Paper2D
+### Paper2D
 
 * New: Added/improved comments for various Paper2D types
 
@@ -6327,7 +6327,7 @@ Cooked builds can use a dedicated thread for package streaming instead of using
 
     * Removed support for the old binary collision layers in Tile Maps
 
-## Physics
+### Physics
 
 * PhysX 3.3.3 source code is available for everyone. Check out https://wiki.unrealengine.com/PhysX_Source_Guide for more info
 
@@ -6519,13 +6519,13 @@ Cooked builds can use a dedicated thread for package streaming instead of using
 
 * Exposed AActor::IsOverlappingActor(), UPrimitiveComponent::IsOverlappingActor(), UPrimitiveComponent::IsOverlappingComponent() to blueprints and improved comments
 
-## Platforms
+### Platforms
 
 * Fixed default directory iterator assignment in FFileIterator.
 
 * Process management has been reworked, deprecating redundant methods of FProcHandle in favor of FPlatformProcess API.
 
-### Android
+#### Android
 
 * New: Added GExternalFilePath to allow access to Android getExternalFilesDir path
 
@@ -6533,11 +6533,11 @@ Cooked builds can use a dedicated thread for package streaming instead of using
 
 * Added two Blueprint nodes SetVolumeButtonsHandledBySystem and ?GetVolumeButtonsHandledBySystem in Utilities/Platform
 
-### iOS
+#### iOS
 
 * Fixed orientation determination when targeting iOS 7 or earlier
 
-### Mac
+#### Mac
 
 * New: A Manifest_DebugFiles.txt is now generated on Mac similar to other platforms. This contains file paths to debug and test apps built in addition to the primary (typically shipping) app.
 
@@ -6587,7 +6587,7 @@ Cooked builds can use a dedicated thread for package streaming instead of using
 
 * Remove an unnecessary workaround from the Mac OpenGL implementation of MapBuffer that reduced performance and was masking a bug in the ResourcePool code.
 
-### Playstation 4
+#### Playstation 4
 
 * New: Added console variable 'r.PS4StallsOnMarkers' which will cause thread traces to line up with markers on RazorGPU thread trace timeline. Adds GPU overhead.
 
@@ -6615,7 +6615,7 @@ Cooked builds can use a dedicated thread for package streaming instead of using
 
 * Removed printout for benign warning in Morpheus head tracking.
 
-### Xbox One
+#### Xbox One
 
 * New: Added support for RQT_Undefined in Xbox query code. Like on Windows, this waits for the GPU to catch up.
 
@@ -6633,15 +6633,15 @@ Cooked builds can use a dedicated thread for package streaming instead of using
 
 * Fixed a D3D11 reference count leak in the DECLARE_ISBOUNDSHADER macro.
 
-### VR
+#### VR
 
 * New: The SteamVR plugin is now in the Runtime plugins folder rather than Experimental.
 
-### Windows
+#### Windows
 
 * New: Allowing for checking major AND minor version of Windows
 
-## Programming
+### Programming
 
 * New: Added API for getting handle of top-level window on Windows
 
@@ -6727,7 +6727,7 @@ Cooked builds can use a dedicated thread for package streaming instead of using
 
 * Slight optimization to FName less than operator
 
-## Rendering
+### Rendering
 
 * New: Added a Thickness option to DrawFlatArrow, DrawDirectionalArrow, and DrawCoordinateSystem
 
@@ -6755,7 +6755,7 @@ Cooked builds can use a dedicated thread for package streaming instead of using
 
 * Removed dead class FFontMaterialRenderProxy
 
-## UI
+### UI
 
 * New: Converted UTextRenderComponent to use FText so it can respond to live-culture updates
 
@@ -6769,7 +6769,7 @@ Cooked builds can use a dedicated thread for package streaming instead of using
 
 * Correct initialization of bMenuHiding and bMenuHidden in SGameMenuPageWidget::Construct(). Fixes issues with clicking on player name input widget in PlatformerGame.
 
-## Slate
+### Slate
 
 * The FString passthrough functionality of SLATE_TEXT_ATTRIBUTE is deprecated as of 4.8, which means that you're now going to get a deprecation warning when passing text to Slate as FString rather than FText.
 
@@ -6831,7 +6831,7 @@ Cooked builds can use a dedicated thread for package streaming instead of using
 
 * List view should not attempt to scroll if that would lead outside scroll range. This check is performed only if overscroll is not allowed.
 
-## UMG
+### UMG
 
 * New: Adding a reference implementation of a web browser wrapper around the SWebBrowser Slate widget to the WebBrowserAsPlugin module.
 
@@ -6845,9 +6845,9 @@ Cooked builds can use a dedicated thread for package streaming instead of using
 
 * Now using IsBindableEvent="True" for widgets that expose bindable single cast delegates. Using the old trick of scanning for ones ending with "Event", will eventually stop working.
 
-# Upgrade Notes
+## Upgrade Notes
 
-## Programming
+### Programming
 
 * Shadow variable warnings enabled.  Almost all shadowed variables have been eliminated from the UE4 source and shadow variable warning has been enabled.
 
@@ -6857,23 +6857,23 @@ Cooked builds can use a dedicated thread for package streaming instead of using
 
     * If you need to disable the warning for just a specific section of code or the includes from a third-party library you cannot modify you can wrap the code with PRAGMA_DISABLE_SHADOW_VARIABLE_WARNINGS and PRAGMA_ENABLE_SHADOW_VARIABLE_WARNINGS
 
-## Navigation
+### Navigation
 
 * Since we had to replace bRebuildAtRuntime boolean property with a new enum, you should update navigation rebuild settings in your project configuration.
 
-## Animation
+### Animation
 
 * If you have any Anim Blueprint derived from either of those classes, it won't work.
 
-## Automation
+### Automation
 
 * Running automation tests through command line for larger projects may take longer before the tests are triggered.
 
-## Blueprints
+### Blueprints
 
 * Users may be required to select the top most comment node in a nested to set once in order to fix older blueprints that have persistent comment selection issues
 
-## Core
+### Core
 
 * Standardised naming convention for class used to managed memory operations, the same names for template and pointer based functions
 
@@ -6909,37 +6909,37 @@ Cooked builds can use a dedicated thread for package streaming instead of using
 	
 	![image alt text](image_66.png)
 
-## Editor and Tools
+### Editor and Tools
 
 * Any project defaults for EditorUserSettings will no longer apply. Please rename any (Default)EditorUserSettings.ini configs to (Default)EditorPerProjectUserSettings.ini
 
     * On loading 4.8 for the first time, we will migrate the first loaded project's editor layout and keybindings to the default location. These settings will then be common to all projects.
 
-## Content Browser
+### Content Browser
 
 * The named FExtender sections in the Content Browser Filters menu have been replaced by the added ability to register new advanced asset categories.
 
-## Cooker
+### Cooker
 
 * Games supporting DLC will now receive a warning when engine content is accessed during cooking if they disable Include Engine Content.  This setting is enabled/disabled in the Release / DLC / Patching Settings from within Unreal Frontend.
 
-## Gameplay Framework
+### Gameplay Framework
 
 * Loading savegame files from previous versions may cause crashes in some instances. If you encounter this, the easiest solution is to delete your savegames (located in MyProject/Saved/SaveGames). If you wish to try and recover this data, look for the "Note for 4.8 and beyond" comment in LoadGameFromSlot for instructions on how to temporarily modify your code to enable loading of old data, so you can then resave in the new version.
 
-## Localization
+### Localization
 
 * Various localization configuration script settings have been renamed or no longer attempt to silently interpret ambiguous values.
 
-## Networking
+### Networking
 
 * Double check for any overloaded uses of these functions and make sure the signatures match.
 
-## Online
+### Online
 
 * The GameInstance class now automatically subscribes to the OnSessionUserInviteAccepted delegate, and GameInstance::OnSessionUserInviteAccepted can be overriden for custom handling of accepted invites.
 
-## Paper2D
+### Paper2D
 
 * Paper2D importing / sprite creation settings have been moved to a new location in Project Settings .. Paper2D - Import.
 
@@ -6947,21 +6947,21 @@ Cooked builds can use a dedicated thread for package streaming instead of using
 
     * The default pixels per unreal unit setting for Paper2D assets has been changed to 1.0, so if you create new sprites and they are at an unexpected scale to existing ones in your project, this should be the first thing to check (the previous default was 2.56 in the engine or 1.28 in 2D templates). It can be configured in the project settings as well.
 
-## Physics
+### Physics
 
 * The new 6 degrees of freedom encompasses the old LockedAxis behavior. The properties have been renamed accordingly.
 
-## Platforms
+### Platforms
 
-### Android
+#### Android
 
 * If using a custom AndroidDeviceProfileSelectorModule, the new Android profiles in this version should be added to handle proper compatibility fallbacks
 
-### Playstation 4
+#### Playstation 4
 
 * The .ini heading and setting names for Morpheus have changed, so existing settings will not work. Please use the new settings in PS4 ProjectSettings to set them up.
 
-### Xbox One
+#### Xbox One
 
 * Any custom modification of AppxManifest templates will need to be migrated to configuration settings.
 
@@ -6975,15 +6975,15 @@ Cooked builds can use a dedicated thread for package streaming instead of using
 
     * Look in Engine\Config\XboxOne\XboxOneEngine.ini and Samples\Games\ShooterGame\Config\XboxOne\XboxOneEngine.ini for examples of how to configure these settings.
 
-### VR
+#### VR
 
 * Older projects may still need to manually enable the GearVR plugin.
 
-### Programming
+#### Programming
 
 * See FWindowsApplication::AddMessageHandler()
 
-### Rendering
+#### Rendering
 
 * if your project was previously configured in this way then you may wish to uncheck an orientation in order to avoid unwanted orientation changes.
 
@@ -6991,19 +6991,19 @@ Cooked builds can use a dedicated thread for package streaming instead of using
 
     * if your project was previously configured in this way then you may wish to uncheck an orientation in order to avoid unwanted orientation changes.
 
-### Lighting
+#### Lighting
 
 * Lighting done in the bass pass does not properly use the new shading model yet. All dynamic lighting does: deferred direct lights, ambient cubemap, and dynamic skylight.
 
-## Postprocessing
+### Postprocessing
 
 * Best to use SSAO defaults settings as older settings might not look as good. New method isn't as good with small SSAO radius but we intend to improve that again
 
-## Optimizations
+### Optimizations
 
 * Check to see if your SpeedTrees need to change to use the new 'BestPlus' configuration.
 
-## Slate
+### Slate
 
 * Any existing custom editor widgets that rely on Tick() to be called should be updated to use active timers instead that register only when the tick is required.
 
@@ -7011,11 +7011,11 @@ Cooked builds can use a dedicated thread for package streaming instead of using
 
     * Any existing custom editor widgets that rely on Tick() to be called should be updated to use active timers instead that register only when the tick is required.
 
-## UMG
+### UMG
 
 * If you have been using UMG for a static image load screen between full level transitions those will no longer work. Using UMG for loading screens when using level streaming will continue to work.
 
-## C++ API Changes
+### C++ API Changes
 
 * New: Various fixes to the in-game / in-editor screenshot facility:
 
