fixed formatting

diff --git a/Engine/Documentation/Source/Support/Builds/ReleaseNotes/2015/4_8/ReleaseNotes_4_8.INT.udn b/Engine/Documentation/Source/Support/Builds/ReleaseNotes/2015/4_8/ReleaseNotes_4_8.INT.udn
--- a/Engine/Documentation/Source/Support/Builds/ReleaseNotes/2015/4_8/ReleaseNotes_4_8.INT.udn
+++ b/Engine/Documentation/Source/Support/Builds/ReleaseNotes/2015/4_8/ReleaseNotes_4_8.INT.udn
@@ -3727,17 +3727,14 @@ Cooked builds can use a dedicated thread for package streaming instead of using
 
 * New: Added support for ASTC texture format on Android ES 2
 
-* New: Allow adding additional tags to AndoidManifest.xml <manifest>, <application>, and <activity> nodes
+* New: Allow adding additional tags to AndroidManifest.xml `<manifest>`, `<application>`, and `<activity>` nodes
 
 * New: Allow better control of AndroidManifest.xml generation
 
-* Extra Activity Settings are also read from [project]/Build/Android/ManifestActivityAdditions.txt
-
-* Extra Application Settings are also read from [project]/Build/Android/ManifestApplicationAdditions.txt
-
-* Extra Requirements are also read from [project]/Build/Android/ManifestRequirementsAdditions.txt and are not restricted to only permissions
-
-* Entire Requirements section may be overridden by contents of [project]/Build/Android/ManifestRequirementsOverride.txt (be sure to include <uses-sdk android:minSdkVersion="9"/> and <uses-feature android:glEsVersion="0x00020000" android:required="true"/> as these are also replaced!)
+	* Extra Activity Settings are also read from [project]/Build/Android/ManifestActivityAdditions.txt
+	* Extra Application Settings are also read from [project]/Build/Android/ManifestApplicationAdditions.txt
+	* Extra Requirements are also read from [project]/Build/Android/ManifestRequirementsAdditions.txt and are not restricted to only permissions
+	* Entire Requirements section may be overridden by contents of [project]/Build/Android/ManifestRequirementsOverride.txt (be sure to include `<uses-sdk android:minSdkVersion="9"/>` and `<uses-feature android:glEsVersion="0x00020000" android:required="true"/>` as these are also replaced!)
 
 * New: Allow package and project name to be different and converts dashes to underscores if necessary to generate a valid package name
 
