Credits (out of order due to circumstances).

diff --git a/Engine/Documentation/Source/Support/Builds/ReleaseNotes/2015/4_8/ReleaseNotes_4_8.INT.udn b/Engine/Documentation/Source/Support/Builds/ReleaseNotes/2015/4_8/ReleaseNotes_4_8.INT.udn
--- a/Engine/Documentation/Source/Support/Builds/ReleaseNotes/2015/4_8/ReleaseNotes_4_8.INT.udn
+++ b/Engine/Documentation/Source/Support/Builds/ReleaseNotes/2015/4_8/ReleaseNotes_4_8.INT.udn
@@ -7,7 +7,7 @@ Description:4.8 Release Notes
 
 This release includes 189 great changes that were submitted from Unreal Engine's amazing community of developers.  Thanks to all of these contributors to Unreal Engine 4.8:
 
-Adrian Gasinski (gasnica), Adrian Irwin (adrianirwin), Ahsan Muzaheed (muzaheed57), Albino Rodrigues (Bino-), Alexander O'Mara (AlexanderOMara), Andrew Zhilin (zoon), Angelo Miner (angjminer), Ash Matheson (NuclearFossil), Audiokinetic Inc. (audiokinetic), Ben Rog-Wilhelm (zorbathut), BhaaL (BhaaLseN), Black Phoenix (PhoenixBlack), Błażej Szczygieł (zaps166), Chris Grebeldinger (cgrebeldinger), Christopher P. Yarger (cpyarger), codetackler, Craig Wilson (hamo2k1), Dainius MasiliÅ«nas (GreatEmerald), Daniel, Thielman (dthielman), Daniel Zeligman (dzeligman), David Palacios (Hevedy), Derek van Vliet (derekvanvliet), devbm, Dewitt Buckingham (DaJoker29), Eugene Elkin (eelkin2000), Ezekiel Warren (zaucy), Fazouane Marouane (fazouane-marouane), gatools, Gigantoad, Henrik Danielsson (TwoD), Jackson Egan (Sneagan), JaredTherriault, Jason (Abatron), Jason Maskell (TamedTornado), Jay Mattis (braindx), jdh84, Jeremy (OPTIPlanet), jlamarche, Joshua Sheard (jsheard), kafumanto, karolz, Korshikov Anton (tnako), korypostma, Lee Berger (MrCrowbar), Leon Rosengarten (lion03), Leszek Godlewski (inequation), Lukasz Baran (iniside) Luna Duclos (PSG-Luna), Maciej Kacper Jagiello (cician), Marat Radchenko (slonopotamus), Marco Antonio Alvarez (surakin), marynate, Maurice Laveaux (mlaveaux), nsf, Omar Pakker (Omar007), Paul Eremeeff (PaulEremeeff), PjotrSvetachov, Quadtree, radu-danciu, Rene Rivera (grafikrobot), Richard Biely (richardbiely), Rene Rivera (grafikrobot), Richard Biely (richardbiely), Richard Nemeth (hsdxpro), Robert Chubb (SaxonRah), Robert Williams (Rottbott), Sasha Petrovich (xulture), Satheesh (ryanjon2040), Simon Taylor (simontaylor81), Spencer Brown (JedTheKrampus), threadlock, Tore Lervik (hallatore), wackoisgod, Will Stahl (merlin91), x414e54, Yohann Martel (ymartel06)
+3dluvr (amigo), Cengiz Terzibas (yaakuro), salamanderrake, Adrian Gasinski (gasnica), Adrian Irwin (adrianirwin), Ahsan Muzaheed (muzaheed57), Albino Rodrigues (Bino-), Alexander O'Mara (AlexanderOMara), Andrew Zhilin (zoon), Angelo Miner (angjminer), Ash Matheson (NuclearFossil), Audiokinetic Inc. (audiokinetic), Ben Rog-Wilhelm (zorbathut), BhaaL (BhaaLseN), Black Phoenix (PhoenixBlack), Błażej Szczygieł (zaps166), Chris Grebeldinger (cgrebeldinger), Christopher P. Yarger (cpyarger), codetackler, Craig Wilson (hamo2k1), Dainius MasiliÅ«nas (GreatEmerald), Daniel, Thielman (dthielman), Daniel Zeligman (dzeligman), David Palacios (Hevedy), Derek van Vliet (derekvanvliet), devbm, Dewitt Buckingham (DaJoker29), Eugene Elkin (eelkin2000), Ezekiel Warren (zaucy), Fazouane Marouane (fazouane-marouane), gatools, Gigantoad, Henrik Danielsson (TwoD), Jackson Egan (Sneagan), JaredTherriault, Jason (Abatron), Jason Maskell (TamedTornado), Jay Mattis (braindx), jdh84, Jeremy (OPTIPlanet), jlamarche, Joshua Sheard (jsheard), kafumanto, karolz, Korshikov Anton (tnako), korypostma, Lee Berger (MrCrowbar), Leon Rosengarten (lion03), Leszek Godlewski (inequation), Lukasz Baran (iniside) Luna Duclos (PSG-Luna), Maciej Kacper Jagiello (cician), Marat Radchenko (slonopotamus), Marco Antonio Alvarez (surakin), marynate, Maurice Laveaux (mlaveaux), nsf, Omar Pakker (Omar007), Paul Eremeeff (PaulEremeeff), PjotrSvetachov, Quadtree, radu-danciu, Rene Rivera (grafikrobot), Richard Biely (richardbiely), Rene Rivera (grafikrobot), Richard Biely (richardbiely), Richard Nemeth (hsdxpro), Robert Chubb (SaxonRah), Robert Williams (Rottbott), Sasha Petrovich (xulture), Satheesh (ryanjon2040), Simon Taylor (simontaylor81), Spencer Brown (JedTheKrampus), threadlock, Tore Lervik (hallatore), wackoisgod, Will Stahl (merlin91), x414e54, Yohann Martel (ymartel06)
 
 ## Major Features
 
