#UE4 DOCS: updated EQS node reference to include all of the Trace Test's properties

diff --git a/Engine/Documentation/Source/Engine/AI/EnvironmentQuerySystem/NodeReference/NodeReference.INT.udn b/Engine/Documentation/Source/Engine/AI/EnvironmentQuerySystem/NodeReference/NodeReference.INT.udn
--- a/Engine/Documentation/Source/Engine/AI/EnvironmentQuerySystem/NodeReference/NodeReference.INT.udn
+++ b/Engine/Documentation/Source/Engine/AI/EnvironmentQuerySystem/NodeReference/NodeReference.INT.udn
@@ -44,6 +44,9 @@ Description:Break downs and descriptions of the various nodes available in the E
 
 [toc]
 
+[VAR:spacer]![](%ROOT%/spacer.png)(w:150 convert:false)[/VAR]
+
+[VAR:spacer2]![](%ROOT%/spacer.png)(w:35 convert:false)[/VAR]
 
 This is a quick reference for the Nodes and Tests available in the Environment Query Editor.
 
@@ -57,11 +60,13 @@ Finds all of the Actors of a class within the Search Radius around the Context.
 
 ![](gen_actor_props.png)
 
-| Property | Description |
-| --- | --- |
-| **Search Radius** | The max distance to look for the desired Actors. |
-| **Searched Actor Class** | The Class of the Actor to look for (Pawn, Character, etc...) |
-| **Search Center** | The context to center the search on, such as from the Querier or perhaps some other context. |
+[REGION:raw]
+| Property | %spacer% | Description |
+| --- | --- | :--- |
+| **Search Radius** || The max distance to look for the desired Actors. |
+| **Searched Actor Class** || The Class of the Actor to look for (Pawn, Character, etc...) |
+| **Search Center** || The context to center the search on, such as from the Querier or perhaps some other context. |
+[/REGION]
 
 ### On circle
 
@@ -71,15 +76,17 @@ Radiates traces out from the Context, returning the hits or the edge of the radi
 
 ![](gen_circle_props.png)
 
-| Property | Description |
-| --- | --- |
-| **Circle Radius** | The radius of the circle that will expand out from the Circle Center Context. |
-| **Space Between** | The space (in cm) between the items generated on the outer edge of the circle. |
-| **Arc Direction** | This will enable the use of the Arc Direction feature which will enable you to focus the arc in a direction based either on a rotation of a context or the vector formed between to locations. |
-| **Arc Angle** | The angle of the arc measured in degrees. |
-| **Circle Center** | The context used as the center for the circle. |
-| **Trace Data** | The type of trace to use. |
-| **Projection Data** | If the resulting items should be projected onto the NavMesh (and which NavMesh data set to use). |
+[REGION:raw]
+| Property | %spacer% | Description |
+| --- | --- | :--- |
+| **Circle Radius** || The radius of the circle that will expand out from the Circle Center Context. |
+| **Space Between** || The space (in cm) between the items generated on the outer edge of the circle. |
+| **Arc Direction** || This will enable the use of the Arc Direction feature which will enable you to focus the arc in a direction based either on a rotation of a context or the vector formed between to locations. |
+| **Arc Angle** || The angle of the arc measured in degrees. |
+| **Circle Center** || The context used as the center for the circle. |
+| **Trace Data** || The type of trace to use. |
+| **Projection Data** || If the resulting items should be projected onto the NavMesh (and which NavMesh data set to use). |
+[/REGION]
 
 ### Pathing Grid
 
@@ -89,13 +96,15 @@ Like Grid, but this paths to different points on the NavMesh, only returning Ite
 
 ![](gen_pathing_props.png)
 
-| Property | Description |
-| --- | --- |
-| **Max Distance** | How far to path away from the Querier. |
-| **Distance Between** | The distance between the generated items. |
-| **Path to Item** | Whether to path **to** or **from** the item. |
-| **Generate Around** | The context to generate the pathing grid around. |
-| **Navigation Filter** | The Navigation Filter to use.  | 
+[REGION:raw]
+| Property | %spacer% | Description |
+| --- | --- | :--- |
+| **Max Distance** || How far to path away from the Querier. |
+| **Distance Between** || The distance between the generated items. |
+| **Path to Item** || Whether to path **to** or **from** the item. |
+| **Generate Around** || The context to generate the pathing grid around. |
+| **Navigation Filter** || The Navigation Filter to use.  | 
+[/REGION]
 
 ### Simple Grid
 
@@ -105,12 +114,14 @@ Generates a grid of Items around the Querier.
 
 ![](gen_grid_props.png)
 
-| Property | Description |
-| --- | --- |
-| **Grid Size** | The height and width of the grid of Items to generate. |
-| **Space Between** | The distance between the grid Items. |
-| **Generated Around** | The context used to Generate the Grid. |
-| **Projected Data** | Should the grid be projected onto the NavMesh? This will move Items that may be inside walls or blocked, back onto the NavMesh, which can create bunching if a grid line happens to be just beyond the edge of the NavMesh. |
+[REGION:raw]
+| Property | %spacer% | Description |
+| --- | --- | :--- |
+| **Grid Size** || The height and width of the grid of Items to generate. |
+| **Space Between** || The distance between the grid Items. |
+| **Generated Around** || The context used to Generate the Grid. |
+| **Projected Data** || Should the grid be projected onto the NavMesh? This will move Items that may be inside walls or blocked, back onto the NavMesh, which can create bunching if a grid line happens to be just beyond the edge of the NavMesh. |
+[/REGION]
 
 
 ## Tests
@@ -123,20 +134,21 @@ This will return the direct distance between the Item and the chosen Context. if
 
 ![](test_dist_props.png)
 
-| Property | Description |
-| --- | --- |
-| **Test Mode** | If the Distance Check should be done in 3D, as a 2D XY plane, or just along the Z axis. |
-| **Distance To** | The Context that should be used to measure distance to. |
-| **Test Purpose** | Should the test Filter, Score, or Filter and Score the result? Filtering will remove Items that fail the test, while scoring will weight an Item. |
-| **Filter Type** | Change the type of filter to be a minimum, maximum, or range of values. Any values outside the values set in **Float Value Min** and/or **Float Value Max**, will be culled. |
-| **Float Value Min** | Filter any value less than or equal to this value. |
-| **Float Value Max** | Filter any value greater than or equal to this value. |
-| **Clamping** | If a value isn't filtered, you can clamp it to limit a test's influence. Change the Clamp Min/Max Type properties to enable this. |
-| **Clamp Min Type** | How should the lower bound for normalization of the raw test value before applying the scoring formula be determined? Should it use the lowest value found (rested), the lower threshold for filtering, or a separate specified normalization minimum. |
-| **Clamp Max Type** | How should the upper bound for normalization of the raw test value before applying the scoring formula be determined? Should it use the highest value found (rested), the higher threshold for filtering, or a separate specified normalization maximum. |
-| **Scoring Equation** | This modifies the score of the test to adhere to a curve. This can be a **Constant**, **Linear**, **Squared**, or **Inverse Linear** curve.  |
-| **Scoring Factor** | The weight this test has among all the tests. Can be negative. |
-
+[REGION:raw]
+| Property | %spacer% | Description |
+| --- | --- | :--- |
+| **Test Mode** || If the Distance Check should be done in 3D, as a 2D XY plane, or just along the Z axis. |
+| **Distance To** || The Context that should be used to measure distance to. |
+| **Test Purpose** || Should the test Filter, Score, or Filter and Score the result? Filtering will remove Items that fail the test, while scoring will weight an Item. |
+| **Filter Type** || Change the type of filter to be a minimum, maximum, or range of values. Any values outside the values set in **Float Value Min** and/or **Float Value Max**, will be culled. |
+| **Float Value Min** || Filter any value less than or equal to this value. |
+| **Float Value Max** || Filter any value greater than or equal to this value. |
+| **Clamping** || If a value isn't filtered, you can clamp it to limit a test's influence. Change the Clamp Min/Max Type properties to enable this. |
+| **Clamp Min Type** || How should the lower bound for normalization of the raw test value before applying the scoring formula be determined? Should it use the lowest value found (rested), the lower threshold for filtering, or a separate specified normalization minimum. |
+| **Clamp Max Type** || How should the upper bound for normalization of the raw test value before applying the scoring formula be determined? Should it use the highest value found (rested), the higher threshold for filtering, or a separate specified normalization maximum. |
+| **Scoring Equation** || This modifies the score of the test to adhere to a curve. This can be a **Constant**, **Linear**, **Squared**, or **Inverse Linear** curve.  |
+| **Scoring Factor** || The weight this test has among all the tests. Can be negative. |
+[/REGION]
 
 ### Dot
 
@@ -146,21 +158,23 @@ Calculates the Dot Product of two vectors. These can be Context Rotations, or th
 
 ![](test_dot_props.png)
 
-| Property | Description |
-| --- | --- |
-| **Line A** | [INCLUDE:#dotline] |
-| **Line B** | [INCLUDE:#dotline] |
-| **Test Mode** | Whether the test should calculate using the complete 3D vector or just the 2D heading vector of the Line A and Line B vectors. |
-| **Absolute Value** | A Dot product returns a value from -1.0 to 1.0, this will make the test return the Absolute Value of the Dot product.  |
-| **Test Purpose** | Should the test Filter, Score, or Filter and Score the result? Filtering will remove Items that fail the test, while scoring will weight an Item. |
-| **Filter Type** | Change the type of filter to be a minimum, maximum, or range of values. Any values outside the values set in **Float Value Min** and/or **Float Value Max**, will be culled. |
-| **Float Value Min** | Filter any value less than or equal to this value. |
-| **Float Value Max** | Filter any value greater than or equal to this value. |
-| **Clamping** | If a value isn't filtered, you can clamp it to limit a test's influence. Change the Clamp Min/Max Type properties to enable this. |
-| **Clamp Min Type** | How should the lower bound for normalization of the raw test value before applying the scoring formula be determined? Should it use the lowest value found (rested), the lower threshold for filtering, or a separate specified normalization minimum. |
-| **Clamp Max Type** | How should the upper bound for normalization of the raw test value before applying the scoring formula be determined? Should it use the highest value found (rested), the higher threshold for filtering, or a separate specified normalization maximum. |
-| **Scoring Equation** | This modifies the score of the test to adhere to a curve. This can be a **Constant**, **Linear**, **Squared**, or **Inverse Linear** curve.  |
-| **Scoring Factor** | The weight this test has among all the tests. |
+[REGION:raw]
+| Property | %spacer% | Description |
+| --- | --- | :--- |
+| **Line A** || [INCLUDE:#dotline] |
+| **Line B** || [INCLUDE:#dotline] |
+| **Test Mode** || Whether the test should calculate using the complete 3D vector or just the 2D heading vector of the Line A and Line B vectors. |
+| **Absolute Value** || A Dot product returns a value from -1.0 to 1.0, this will make the test return the Absolute Value of the Dot product.  |
+| **Test Purpose** || Should the test Filter, Score, or Filter and Score the result? Filtering will remove Items that fail the test, while scoring will weight an Item. |
+| **Filter Type** || Change the type of filter to be a minimum, maximum, or range of values. Any values outside the values set in **Float Value Min** and/or **Float Value Max**, will be culled. |
+| **Float Value Min** || Filter any value less than or equal to this value. |
+| **Float Value Max** || Filter any value greater than or equal to this value. |
+| **Clamping** || If a value isn't filtered, you can clamp it to limit a test's influence. Change the Clamp Min/Max Type properties to enable this. |
+| **Clamp Min Type** || How should the lower bound for normalization of the raw test value before applying the scoring formula be determined? Should it use the lowest value found (rested), the lower threshold for filtering, or a separate specified normalization minimum. |
+| **Clamp Max Type** || How should the upper bound for normalization of the raw test value before applying the scoring formula be determined? Should it use the highest value found (rested), the higher threshold for filtering, or a separate specified normalization maximum. |
+| **Scoring Equation** || This modifies the score of the test to adhere to a curve. This can be a **Constant**, **Linear**, **Squared**, or **Inverse Linear** curve.  |
+| **Scoring Factor** || The weight this test has among all the tests. |
+[/REGION]
 
 <!--
 [EXCERPT:dotLine]
@@ -185,13 +199,15 @@ Calculates the Dot Product of two vectors. These can be Context Rotations, or th
 
 ![](test_tags_props.png)
 
-| Property | Description |
-| --- | --- |
-| **Tags to Match** | Whether to match **Any** or **All** tags from the Gameplay Tags property. Basically an **or** or **and** statement. |
-| **Gameplay Tags** | The tags that should be matched. |
-| **Test Purpose** | Should the test Filter, Score, or Filter and Score the result? Filtering will remove Items that fail the test, while scoring will weight an Item. |
-| **Bool Match** | Essentially, are you looking for a true match or a false match? If the item returns that it has a desired tag, then this would allow you to _remove_ it instead of keep it. |
-| **Scoring Factor** | The weight this test has among all the tests. |
+[REGION:raw]
+| Property | %spacer% | Description |
+| --- | --- | :--- |
+| **Tags to Match** || Whether to match **Any** or **All** tags from the Gameplay Tags property. Basically an **or** or **and** statement. |
+| **Gameplay Tags** || The tags that should be matched. |
+| **Test Purpose** || Should the test Filter, Score, or Filter and Score the result? Filtering will remove Items that fail the test, while scoring will weight an Item. |
+| **Bool Match** || Essentially, are you looking for a true match or a false match? If the item returns that it has a desired tag, then this would allow you to _remove_ it instead of keep it. |
+| **Scoring Factor** || The weight this test has among all the tests. |
+[/REGION]
 
 
 ### Path Finding
@@ -200,29 +216,31 @@ Calculates the Dot Product of two vectors. These can be Context Rotations, or th
 
 ![](test_path_props.png)
 
-| Property | Description |
-| --- | --- |
-| **Test Mode** | [INCLUDE:#path1] |
-| **Context** | The Context to Path to or from. |
-| **Path From Context** | Should the pathfinder go to of from the Context. |
-| **Skip Unreachable** | Filter out any Item that is unreachable. |
-| **Use Hierarchical Pathfinding** | Should Hierarchical Pathfinding be used, as it's faster, though less accurate. |
-| **Test Purpose** | Should the test Filter, Score, or Filter and Score the result? Filtering will remove Items that fail the test, while scoring will weight an Item. |
-| **Filter Type** | Change the type of filter to be a minimum, maximum, or range of values. Any values outside the values set in **Float Value Min** and/or **Float Value Max**, will be culled. |
-| **Float Value Min** | Filter any value less than or equal to this value. |
-| **Float Value Max** | Filter any value greater than or equal to this value. |
-| **Clamping** | If a value isn't filtered, you can clamp it to limit a test's influence. Change the Clamp Min/Max Type properties to enable this. |
-| **Clamp Min Type** | How should the lower bound for normalization of the raw test value before applying the scoring formula be determined? Should it use the lowest value found (rested), the lower threshold for filtering, or a separate specified normalization minimum. |
-| **Clamp Max Type** | How should the upper bound for normalization of the raw test value before applying the scoring formula be determined? Should it use the highest value found (rested), the higher threshold for filtering, or a separate specified normalization maximum. |
-| **Scoring Equation** | This modifies the score of the test to adhere to a curve. This can be a **Constant**, **Linear**, **Squared**, or **Inverse Linear** curve.  |
-| **Scoring Factor** | The weight this test has among all the tests. |
+[REGION:raw]
+| Property | %spacer% | Description |
+| --- | --- | :--- |
+| **Test Mode** || [INCLUDE:#path1] |
+| **Context** || The Context to Path to or from. |
+| **Path From Context** || Should the pathfinder go to of from the Context. |
+| **Skip Unreachable** || Filter out any Item that is unreachable. |
+| **Use Hierarchical Pathfinding** || Should Hierarchical Pathfinding be used, as it's faster, though less accurate. |
+| **Test Purpose** || Should the test Filter, Score, or Filter and Score the result? Filtering will remove Items that fail the test, while scoring will weight an Item. |
+| **Filter Type** || Change the type of filter to be a minimum, maximum, or range of values. Any values outside the values set in **Float Value Min** and/or **Float Value Max**, will be culled. |
+| **Float Value Min** || Filter any value less than or equal to this value. |
+| **Float Value Max** || Filter any value greater than or equal to this value. |
+| **Clamping** || If a value isn't filtered, you can clamp it to limit a test's influence. Change the Clamp Min/Max Type properties to enable this. |
+| **Clamp Min Type** || How should the lower bound for normalization of the raw test value before applying the scoring formula be determined? Should it use the lowest value found (rested), the lower threshold for filtering, or a separate specified normalization minimum. |
+| **Clamp Max Type** || How should the upper bound for normalization of the raw test value before applying the scoring formula be determined? Should it use the highest value found (rested), the higher threshold for filtering, or a separate specified normalization maximum. |
+| **Scoring Equation** || This modifies the score of the test to adhere to a curve. This can be a **Constant**, **Linear**, **Squared**, or **Inverse Linear** curve.  |
+| **Scoring Factor** || The weight this test has among all the tests. |
+[/REGION]
 
 <!--
 [EXCERPT:path1]
-| --- | --- |
-| **Path Exist** | Does the path exist to (or from) the Context. |
-| **Path Cost** | How expensive is the path to (or from) the Context? |
-| **Path Length** | How long is the path to (or from) the Context? |
+| --- | --- | :--- |
+| **Path Exist** || Does the path exist to (or from) the Context. |
+| **Path Cost** || How expensive is the path to (or from) the Context? |
+| **Path Length** || How long is the path to (or from) the Context? |
 [/EXCERPT:path1]
 -->
 
@@ -235,18 +253,19 @@ Less of a test (though it can be used to randomly filter out Items), and more of
 
 ![](test_rand_props.png)
 
-| Property | Description |
-| --- | --- |
-| **Test Purpose** | Should the test Filter, Score, or Filter and Score the result? Filtering will remove Items that fail the test, while scoring will weight an Item. |
-| **Filter Type** | Change the type of filter to be a minimum, maximum, or range of values. Any values outside the values set in **Float Value Min** and/or **Float Value Max**, will be culled. |
-| **Float Value Min** | Filter any value less than or equal to this value. |
-| **Float Value Max** | Filter any value greater than or equal to this value. |
-| **Clamping** | If a value isn't filtered, you can clamp it to limit a test's influence. Change the Clamp Min/Max Type properties to enable this. |
-| **Clamp Min Type** | How should the lower bound for normalization of the raw test value before applying the scoring formula be determined? Should it use the lowest value found (rested), the lower threshold for filtering, or a separate specified normalization minimum. |
-| **Clamp Max Type** | How should the upper bound for normalization of the raw test value before applying the scoring formula be determined? Should it use the highest value found (rested), the higher threshold for filtering, or a separate specified normalization maximum. |
-| **Scoring Equation** | This modifies the score of the test to adhere to a curve. This can be a **Constant**, **Linear**, **Squared**, or **Inverse Linear** curve.  |
-| **Scoring Factor** | The weight this test has among all the tests. |
-
+[REGION:raw]
+| Property | %spacer% | Description |
+| --- | --- | :--- |
+| **Test Purpose** || Should the test Filter, Score, or Filter and Score the result? Filtering will remove Items that fail the test, while scoring will weight an Item. |
+| **Filter Type** || Change the type of filter to be a minimum, maximum, or range of values. Any values outside the values set in **Float Value Min** and/or **Float Value Max**, will be culled. |
+| **Float Value Min** || Filter any value less than or equal to this value. |
+| **Float Value Max** || Filter any value greater than or equal to this value. |
+| **Clamping** || If a value isn't filtered, you can clamp it to limit a test's influence. Change the Clamp Min/Max Type properties to enable this. |
+| **Clamp Min Type** || How should the lower bound for normalization of the raw test value before applying the scoring formula be determined? Should it use the lowest value found (rested), the lower threshold for filtering, or a separate specified normalization minimum. |
+| **Clamp Max Type** || How should the upper bound for normalization of the raw test value before applying the scoring formula be determined? Should it use the highest value found (rested), the higher threshold for filtering, or a separate specified normalization maximum. |
+| **Scoring Equation** || This modifies the score of the test to adhere to a curve. This can be a **Constant**, **Linear**, **Squared**, or **Inverse Linear** curve.  |
+| **Scoring Factor** || The weight this test has among all the tests. |
+[/REGION]
 
 ### Trace
 
@@ -256,17 +275,29 @@ This will trace to or from an Item to a Context and return if it hit it or not.
 
 ![](test_vis_props.png)
 
-| Property | Description |
-| --- | --- |
-| **Test Purpose** | Should the test Filter, Score, or Filter and Score the result? Filtering will remove Items that fail the test, while scoring will weight an Item. |
-| **Filter Type** | Change the type of filter to be a minimum, maximum, or range of values. Any values outside the values set in **Float Value Min** and/or **Float Value Max**, will be culled. |
-| **Float Value Min** | Filter any value less than or equal to this value. |
-| **Float Value Max** | Filter any value greater than or equal to this value. |
-| **Clamping** | If a value isn't filtered, you can clamp it to limit a test's influence. Change the Clamp Min/Max Type properties to enable this. |
-| **Clamp Min Type** | How should the lower bound for normalization of the raw test value before applying the scoring formula be determined? Should it use the lowest value found (rested), the lower threshold for filtering, or a separate specified normalization minimum. |
-| **Clamp Max Type** | How should the upper bound for normalization of the raw test value before applying the scoring formula be determined? Should it use the highest value found (rested), the higher threshold for filtering, or a separate specified normalization maximum. |
-| **Scoring Equation** | This modifies the score of the test to adhere to a curve. This can be a **Constant**, **Linear**, **Squared**, or **Inverse Linear** curve.  |
-| **Scoring Factor** | The weight this test has among all the tests. |
+[REGION:raw]
+| Property | %spacer% | Description |
+| --- | --- | :--- |
+| [REGION:tablesection]**Trace Data**[/REGION] |||
+| **Trance Channel** || The channel to trace against, by default Unreal Engine 4 has Visibility and Camera. Trace Channels can be added and adjusted in the **Edit Menu** -> **Project Settings** -> **Physics** -> **Trace Channels** |
+| **Trace Shape** || Line, Sphere, Box, Capsule |
+| **Trace Complex** || If the trace should be against the mesh (complex) or just the simple collision |
+| **Only Blocking Hits** || Whether the trace uses blocking or non-blocking traces in its results. | 
+| [REGION:tablesection]%spacer2%[/REGION] |||
+| **Trace From Context** || The context to trace from, such as the the Querier, an Item, or any Custom Contexts you may have created. |
+| **Context** || |
+| **Item Height Offset** || This will add a Z offset, in cm, to the Item the test is tracing to/from. |
+| **Context Height Offset** || This will add a Z offset, in cm, to the Context the test is tracing to/from. |
+| **Test Purpose** || Should the test Filter, Score, or Filter and Score the result? Filtering will remove Items that fail the test, while scoring will weight an Item. |
+| **Filter Type** || Change the type of filter to be a minimum, maximum, or range of values. Any values outside the values set in **Float Value Min** and/or **Float Value Max**, will be culled. |
+| **Float Value Min** || Filter any value less than or equal to this value. |
+| **Float Value Max** || Filter any value greater than or equal to this value. |
+| **Clamping** || If a value isn't filtered, you can clamp it to limit a test's influence. Change the Clamp Min/Max Type properties to enable this. |
+| **Clamp Min Type** || How should the lower bound for normalization of the raw test value before applying the scoring formula be determined? Should it use the lowest value found (rested), the lower threshold for filtering, or a separate specified normalization minimum. |
+| **Clamp Max Type** || How should the upper bound for normalization of the raw test value before applying the scoring formula be determined? Should it use the highest value found (rested), the higher threshold for filtering, or a separate specified normalization maximum. |
+| **Scoring Equation** || This modifies the score of the test to adhere to a curve. This can be a **Constant**, **Linear**, **Squared**, or **Inverse Linear** curve.  |
+| **Scoring Factor** || The weight this test has among all the tests. |
+[/REGION]
 
 ## Contexts
 
