#UE4 Docs: Fixing up the GearVR docs with the feedback from Oculus as well as updating some information to reflect some new changes in 4.8

diff --git a/Engine/Documentation/Source/Platforms/GearVR/BestPractices/GearVRBestPractices.INT.udn b/Engine/Documentation/Source/Platforms/GearVR/BestPractices/GearVRBestPractices.INT.udn
--- a/Engine/Documentation/Source/Platforms/GearVR/BestPractices/GearVRBestPractices.INT.udn
+++ b/Engine/Documentation/Source/Platforms/GearVR/BestPractices/GearVRBestPractices.INT.udn
@@ -23,7 +23,7 @@ using the Samsung Gear Virtual Reality(VR) Headset and Samsung Galaxy Note 4 Sma
 ## Gear VR Rendering Options
 While combining the Gear VR Headset with the Samsung Galaxy Note 4 Smartphone can produce some very immersive VR experiences, 
 there are some rendering features that UE4 can not perform while using the Gear VR Headset due to hardware limitations. 
-Many of the restrictions that you will encounter are similar to restrictions for [mobile devices](Platforms/Mobile/Performance).
+Many of the restrictions that you will encounter are similar to restrictions for **[mobile devices](Platforms/Mobile/Performance)**.
 However because of the unique rendering requirements for the Gear VR Headset there are few extra things that must be considered.
 In the following sections we will cover what these requirements are and how they can effect your Gear VR project. 
 
@@ -58,7 +58,7 @@ For example only the following Material inputs will work with the Gear VR Headse
 | **Opacity Mask** |Opacity Mask helps to define if a certain pixel is see through or not. |
 
 [REGION:note]
-Please keep in mind that the both **Translucent** and **Masked** Materials can be very expensive to render, especially on hardware like Smartphones.
+Please keep in mind that the both **Translucent** and **Masked** Materials can be very expensive to render, especially on hardware like Smartphone's.
 Whenever possible you should use [Opaque Materials](Engine/Rendering/Materials/MaterialProperties/BlendModes) as they are much cheaper to render.
 [/region]
 
@@ -98,7 +98,7 @@ The settings contained here are for helping to lessen Material rendering costs o
 | bUseLightmapDirectionality | When disabled, lighting from light maps will be flat but cheaper. |
 
 ### Customized UVs
-[Customized UVs](Engine/Rendering/Materials/CustomizedUVs) allow you to perform certain shader operations in vertex shaders instead of pixel shaders to help avoid dependent texture fetches. 
+**[Customized UVs](Engine/Rendering/Materials/CustomizedUVs)** allow you to perform certain shader operations in vertex shaders instead of pixel shaders to help avoid dependent texture fetches. 
 This offers an excellent way to claw back much needed performance when doing simple things like UV manipulation to make a texture title more than it currently does. 
 Whenever possible make sure to do all UV manipulation using Customized UVs especially when developing content for the Gear VR Headset.
 
@@ -107,7 +107,7 @@ Customized UVs are currently only implemented for certain component types such a
 [/region]
 
 ### Reflections
-The Gear VR does not support Screen Space Reflections(SSR) which means you will only be able generate reflections that come from [Reflection Capture Actors](Resources/ContentExamples/Reflections/1_4/) that are placed in levels. 
+The Gear VR does not support Screen Space Reflections(SSR) which means you will only be able generate reflections that come from **[Reflection Capture Actors](Resources/ContentExamples/Reflections/1_4/)** that are placed in levels. 
 Also keep in mind that the reflections that are generated from the reflections actors will not be of very high fidelity due to hardware limitations. 
 
 ## Gear VR & Content Setup
@@ -124,35 +124,35 @@ Failing to follow these restrictions will result in your project now working.
 ### Mesh Considerations
 When creating Meshes to use with your Gear VR Project the following information needs to be kept in mind to ensure that what you are creating will work with the Gear VR Headset. 
 
-* Use as few [Material IDs](Engine/Content/FBX/Materials) on 3D meshes as possible to save on draw calls.
-* Make sure to make use of [Static Mesh LODs](Engine/Content/FBX/StaticMeshes) whenever possible.
-* Keep an eye on your 3D meshes polygon count and try and keep it as [Low Poly](https://en.wikipedia.org/wiki/Low_poly) as possible.
-* Make sure that all static models have a 2nd UV set so they can make use of [Light Mapping](Engine/Content/Types/StaticMeshes/LightmapUnwrapping).
-* Whenever possible try and make use if [Imposter Sprites](Engine/Content/Tools/RenderToTextureTools/3) to replace far away Static Meshes as they are much cheaper to render.
+* Use as few **[Material IDs](Engine/Content/FBX/Materials)** on 3D meshes as possible to save on draw calls.
+* Make sure to make use of **[Static Mesh LODs](Engine/Content/FBX/StaticMeshes)** whenever possible.
+* Keep an eye on your 3D meshes polygon count and try and keep it as **[Low Poly](https://en.wikipedia.org/wiki/Low_poly)** as possible.
+* Make sure that all static models have a 2nd UV set so they can make use of **[Light Mapping](Engine/Content/Types/StaticMeshes/LightmapUnwrapping)**.
+* Whenever possible try and make use if **[Imposter Sprites](Engine/Content/Tools/RenderToTextureTools/3)** to replace far away Static Meshes as they are much cheaper to render.
 
 ### Texture Considerations
 When creating Textures to use with your Gear VR Project the following information needs to be kept in mind to ensure that what you are creating will not be detrimental to your projects overall performance. 
 
-* Textures use the [ECT2](http://en.wikipedia.org/wiki/Ericsson_Texture_Compression) Texture compression format.
+* Textures use the **[ECT2](http://en.wikipedia.org/wiki/Ericsson_Texture_Compression)** Texture compression format.
 * Textures **can not** be bigger than 2048 pixels in either the X or Y dimension.
 * Textures sizes **MUST BE** in a power of 2 (2, 4, 8, 16, 32, 64, 128, 256, 512, 1024, 2048)
 * Loading up to 512 MB of Textures into memory is possible with the Gear VR however it is recommended that you do not load more that 128 MB of Textures into memory.
 * Use square Textures whenever possible as them make a more efficient use of memory. 
-* Wherever possible try and use [Texture Atlasing](http://en.wikipedia.org/wiki/Texture_atlas) to reduce the number of Textures needed in memory.
-* [RGB Texture Masking / Packing](Engine/Rendering/Materials/HowTo/Masking) is a great way to cut down on unique Textures in memory.
+* Wherever possible try and use **[Texture Atlasing](http://en.wikipedia.org/wiki/Texture_atlas)** to reduce the number of Textures needed in memory.
+* **[RGB Texture Masking / Packing](Engine/Rendering/Materials/HowTo/Masking)** is a great way to cut down on unique Textures in memory.
 * Only **TC_Default** and **TC_NormalMap** compression settings should be used as they are the only two that will work on mobile platforms. 
 
 ### Level Considerations
 Getting the maximum performance out of your levels on the Gear VR is a delicate balancing act that will require time to get just right. 
 The following suggestions / things to look out for can help you quickly identify areas of your levels that might cause performance issues when viewed on the Gear VR Headset.
 
-* Always make sure that the level lighting has been [built](Engine/Rendering/LightingAndShadows/Lightmass) to ensure you get the best performance possible. If you see a red **Lighting needs to be rebuilt** error message on the screen, this means that your lighting needs to be built.
-* Make sure that all the lights in you level are set to [Static](Engine/Rendering/LightingAndShadows/LightMobility/StaticLights) as the other lighting options can be very expensive to render.
-* Check to see if your level is making use of [Precomputed Visibility](Engine/Actors/Volumes) and if not make sure to add one to the level and re-build lighting to enable it in the level's World Settings. 
-* Draw calls for the entire scene should not be greater than 100 for any view. This can be seen with [stat openglrhi](Engine/Performance/StatCommands) on the device, [stat d3d11rhi](Engine/Performance/StatCommands) in ES2 previewer on PC.
-* The total triangle count of the level should not be greater than 100,000 for any view. This can be seen with [stat openglrhi](Engine/Performance/StatCommands) on the device, [stat d3d11rhi](Engine/Performance/StatCommands) in ES2 previewer on PC.
+* Always make sure that the level lighting has been **[built](Engine/Rendering/LightingAndShadows/Lightmass)** to ensure you get the best performance possible. If you see a red **Lighting needs to be rebuilt** error message on the screen, this means that your lighting needs to be built.
+* Make sure that all the lights in you level are set to **[Static](Engine/Rendering/LightingAndShadows/LightMobility/StaticLights)** as the other lighting options can be very expensive to render.
+* Check to see if your level is making use of **[Precomputed Visibility](Engine/Actors/Volumes)** and if not make sure to add one to the level and re-build lighting to enable it in the level's World Settings. 
+* Draw calls for the entire scene should not be greater than 100 for any view. This can be seen with **[stat openglrhi](Engine/Performance/StatCommands)** on the device, **[stat d3d11rhi](Engine/Performance/StatCommands)** in ES2 previewer on PC.
+* The total triangle count of the level should not be greater than 100,000 for any view. This can be seen with **[stat openglrhi](Engine/Performance/StatCommands)** on the device, **[stat d3d11rhi](Engine/Performance/StatCommands)** in ES2 previewer on PC.
 * Be very careful in how Translucent and Opacity Masked Materials are used in the level as they are very expensive to render.
-* Make sure to [Performance Profile](Engine/Performance) your level often to make sure you catch any performance issues sooner rather than later.
+* Make sure to **[Performance Profile](Engine/Performance)** your level often to make sure you catch any performance issues sooner rather than later.
 
 
 
diff --git a/Engine/Documentation/Source/Platforms/GearVR/Debugging/GearVRDebugging.INT.udn b/Engine/Documentation/Source/Platforms/GearVR/Debugging/GearVRDebugging.INT.udn
--- a/Engine/Documentation/Source/Platforms/GearVR/Debugging/GearVRDebugging.INT.udn
+++ b/Engine/Documentation/Source/Platforms/GearVR/Debugging/GearVRDebugging.INT.udn
@@ -1,8 +1,8 @@
 ﻿Availability:Public
 Title:Samsung Gear VR Debugging
-Crumbs:%ROOT%, Platforms
 Description:Information on how to debug applications for the Gear VR headset.
 Navigation:platform
+Crumbs:
 Platform:Gear VR
 Related: Platforms/Mobile/Performance
 Related: Engine/Performance
diff --git a/Engine/Documentation/Source/Platforms/GearVR/GearVR.INT.udn b/Engine/Documentation/Source/Platforms/GearVR/GearVR.INT.udn
--- a/Engine/Documentation/Source/Platforms/GearVR/GearVR.INT.udn
+++ b/Engine/Documentation/Source/Platforms/GearVR/GearVR.INT.udn
@@ -1,71 +1,9 @@
-Title: Samsung Gear VR Headset Prerequisites
-Crumbs:%ROOT%, Platforms
+Title: Samsung Gear VR Development
 Description:Instructions on how to get the Samsung Gear VR Headset to work with UE4.
 Availability:Public
 Navigation:platform
 Platform:Gear VR
 
-![](Platforms\GearVR\QuickStart\GVR_Headset_00.png)
-
-[REGION:warning]
-Developing for the Gear VR Headset requires **Unreal Engine 4.7.4** or later.
-[/REGION]
-
-## Overview 
-
-With the 4.7 release of Unreal Engine 4 (UE4) you can now build Virtual Reality(VR) projects for the Samsung Gear VR Headset that can transform your Samsung Galaxy Note 4 Smartphone into an immersive VR device. 
-In the following sections we will go over everything you need to know to get your Samsung Galaxy Note 4 Smartphone properly setup so that it can work with the Samsung Gear VR headset.
-
-## Gear VR Headset Pre-requisites 
-
-Before we can get into UE4 Gear VR development we first must make sure that the following prerequisites are fulfilled as if we do not, we run the risk of our Gear VR projects not working with the Gear VR Headset. 
-
-* Make sure you have the USB cord used for charging your Samsung Galaxy Note 4 Smartphone available and plugged into a USB port on the PC you are developing on. 
-The USB cord is of vital importance as this is how you transfer data between your development PC and your Smartphone.
-
-* Only certain types of Samsung Galaxy Note 4's will work with the Gear VR Headset. 
-	The [Oculus website](https://support.oculus.com/hc/en-us/articles/204113096-Samsung-Gear-VR-Innovator-Edition-Samsung-Galaxy-Note-4-Compatibility) has a list of phone models that will work with the Gear VR Headset.
-	Make sure to double check that the Samsung Galaxy Note 4 Smartphone you are buying is the correct model that will work with the Gear VR Headset.
-
-* Make sure to download and install the specific USB drivers for your Galaxy Note 4 from the [Samsung Website](http://www.samsung.com/us/support/downloads). 
-
-* If you do not have it all ready, make sure to download and install [Microsoft Visual Studio 2013](https://www.visualstudio.com/en-us/downloads/download-visual-studio-vs.aspx).
-
-* Download Unreal Engine version **4.7.4** or later from the [UE4 Launcher](https://www.unrealengine.com/) or [Github](https://www.unrealengine.com/ue4-on-github). 
-
-* Install the **Tegra Android Development Pack** which can be found in **Engine\Extras\Android**.  
-
-* Check to see if [Developer Mode](https://docs.unrealengine.com/latest/INT/Platforms/Android/GettingStarted/2/index.html ) has been enabled on your Samsung Galaxy Note 4 Smartphone and if not, enable it. 
-
-* Also make sure to check that USB Debugging has been enabled to ensure that you can load data onto your Smartphone from your PC.
-	
-	![](GVR_Dev_Options_Enabled.png)
-
-* With your Smartphone plugged into your development PC via USB, make sure to accept the RSA Fingerprint key from your PC.
-This way your development PC and Smartphone can send data back and forth.
-Also make sure to check **Always allow this computer** to ensure that your development PC can always interact with this Smartphone when plugged in. 
-
-	![](GVR_USB_Debuging.png)
-
-* Generate and download the **OSIG File** for your Samsung Galaxy Note 4 Smartphone using your Smartphone's Device ID and the [Oculus](https://developer.oculus.com/tools/osig/) website.
-If you are not sure how to get your Smartphone's **Device ID** the [Oculus website](https://developer.oculus.com/tools/osig/) has detailed directions on how to do this.
-	
-	![](GVR_OSIG_Website.png)
-
-* After you download the OSIG file, create a folder called **assets** and place the **OSIG** file inside of it. 
-Then place that folder in **Engine/Build/Android/Java**. 
-
-	[region:warning]
-	**If you do not do this your project will not be able to run on your Galaxy Note 4 Smartphone.**
-	[/region]
-	
-	![](GVR_OSIG_File_Location.png)
-
-It is very important that you make sure all items listed in the Gear VR Prerequisites section are completed before moving on to the next section. 
-Skipping a step could result in your project not working with the Gear VR Headset. 
-Once everything from above has been completed you can now move on to the Gear VR Headset Development Quick Start.
-
-
 [OBJECT:TopicButtonList]
 	[PARAM:icon]
 		![](%ROOT%/vr_icon.png)(convert:false)
@@ -76,6 +14,7 @@ Once everything from above has been completed you can now move on to the Gear VR
 	[PARAM:description]
 	[/PARAM]
 	[PARAM:links]
+		* [](Platforms/GearVR/Prerequisites "%Platforms/GearVR/Prerequisites:description%")
 		* [](Platforms/GearVR/QuickStart "%Platforms/GearVR/QuickStart:description%")
 		* [](Platforms/GearVR/BestPractices "%Platforms/GearVR/BestPractices:description%")
 		* [](Platforms/GearVR/Debugging "%Platforms/GearVR/BestPractices:description%")
diff --git a/Engine/Documentation/Source/Platforms/GearVR/QuickStart/1/GearVRQuickStart_1.INT.udn b/Engine/Documentation/Source/Platforms/GearVR/QuickStart/1/GearVRQuickStart_1.INT.udn
--- a/Engine/Documentation/Source/Platforms/GearVR/QuickStart/1/GearVRQuickStart_1.INT.udn
+++ b/Engine/Documentation/Source/Platforms/GearVR/QuickStart/1/GearVRQuickStart_1.INT.udn
@@ -1,5 +1,5 @@
 Availability:Public
-Title:1. Gear VR C++ Project Creation
+Title:1. Gear VR Project Creation
 Crumbs:%ROOT%, Platforms/GearVR
 Description:How to create a UE4 C++ based project that can be used with the Gear VR Headset.
 Navigation:platform
@@ -23,14 +23,12 @@ Navigation:platform
 
 %Steps%
 
-Setting up a UE4 project that works with the Samsung Gear VR Headset requires a slightly different setup than a normal UE4 project. 
-This is because a Gear VR Headset project needs to be created as a C++ project and then compiled to ensure that the proper files for the Gear VR Headset are compiled into the projects binary files. 
-In the following section we will go over everything you need to know about how to properly setup your projects so they work with the Gear VR Headset.
+With the release of Unreal Engine 4.8, setting up a project that works with the Samsung Gear VR Headset is no different than setup than a normal UE4 project. 
+In the following section we will go over everything you need to know in order to create a new project that will work with the Gear VR headset.
 
 [REGION:warning]
-Currently projects that are created using the Blueprint based templates **WILL NOT** work with the Gear VR Headset. 
-This means that all newly created Gear VR projects **MUST** use one of the C++ based templates. 
-However this does not mean that you can not use Blueprints in your projects even if they were created using the C++ based templates.
+If you are using UE4 version 4.7 you will need to create your project as a C++ project and not a Blueprint project.
+In UE4 version 4.8 and later you do not have to do this and can use either C++ or Blueprints for your projects.
 [/REGION]
 
 1. First make sure that you have UE4 version **4.7.4** or later downloaded and installed on your development PC. 
@@ -45,6 +43,10 @@ You can get UE4 version **4.7.4** or later from the [Unreal Engine Launcher](htt
 1. Once you have a version of UE4 **4.7.4** or later installed on your development PC create a new **C++ Project** from the **New Project Tab**.
 	![](GVR_New_Project_00.png)
 
+	[Region:note]
+	If you are using Unreal Engine 4.8 or later you can use either a C++ based project or a Blueprint based project.
+	[/REGION]
+
 1. Next select the **First Person Template** as the template to be used and make sure to change the following options.
 	* **Mobile / Tablet**
 	* **Scalable 3D / 2D** 
diff --git a/Engine/Documentation/Source/Platforms/GearVR/QuickStart/2/GearVRQuickStart_2.INT.udn b/Engine/Documentation/Source/Platforms/GearVR/QuickStart/2/GearVRQuickStart_2.INT.udn
--- a/Engine/Documentation/Source/Platforms/GearVR/QuickStart/2/GearVRQuickStart_2.INT.udn
+++ b/Engine/Documentation/Source/Platforms/GearVR/QuickStart/2/GearVRQuickStart_2.INT.udn
@@ -27,6 +27,10 @@ Navigation:platform
 
 	![](GVR_Editor_Opened.png)
 
+	[region:note]
+	If you are using Unreal Engine 4.8 or later and a Blueprint based project Visual Studio 2013 will not open up as it is not needed. 
+	[/region]
+
 1. Once the UE4 Editor has loaded, from the **Main Toolbar**, open the **Windows** tab and select the **Plugins** menu item. 
 
 	![](GVR_Plugins_Menu.gif)(convert:false)
diff --git a/Engine/Documentation/Source/Platforms/GearVR/QuickStart/3/GearVRQuickStart_3.INT.udn b/Engine/Documentation/Source/Platforms/GearVR/QuickStart/3/GearVRQuickStart_3.INT.udn
--- a/Engine/Documentation/Source/Platforms/GearVR/QuickStart/3/GearVRQuickStart_3.INT.udn
+++ b/Engine/Documentation/Source/Platforms/GearVR/QuickStart/3/GearVRQuickStart_3.INT.udn
@@ -4,6 +4,7 @@ Crumbs:%ROOT%, Platforms/GearVR
 Description:Instructions on how to deploy projects form UE4 to use with the Samsung Gear VR Headset.
 Navigation:platform
 Platform:Gear VR
+Related:
 
 [VAR:Steps]
 [OBJECT:Navigation]
diff --git a/Engine/Documentation/Source/Platforms/GearVR/QuickStart/GearVRQuickStart.INT.udn b/Engine/Documentation/Source/Platforms/GearVR/QuickStart/GearVRQuickStart.INT.udn
--- a/Engine/Documentation/Source/Platforms/GearVR/QuickStart/GearVRQuickStart.INT.udn
+++ b/Engine/Documentation/Source/Platforms/GearVR/QuickStart/GearVRQuickStart.INT.udn
@@ -5,12 +5,8 @@ Description:Instructions on how to quickly setup and deploy projects for use wit
 Navigation:platform
 Platform:Gear VR
 
-
 The Gear VR Quick Start will walk you through how to setup your development PC and Samsung Galaxy Note 4 Smartphone to work with the Samsung Gear VR headset.
 
-![](GVR_Headset_00.png)
-
-
 [OBJECT:TopicButtonList]
 	[PARAM:icon]
 		![](%ROOT%/build_icon.png)(convert:false)
diff --git a/Engine/Documentation/Source/Platforms/Oculus/BestPractices/OculusRift.INT.udn b/Engine/Documentation/Source/Platforms/Oculus/BestPractices/OculusRift.INT.udn
--- a/Engine/Documentation/Source/Platforms/Oculus/BestPractices/OculusRift.INT.udn
+++ b/Engine/Documentation/Source/Platforms/Oculus/BestPractices/OculusRift.INT.udn
@@ -26,7 +26,7 @@ If that did not fix your issues the next thing to check is to make sure that the
 do this, you need to do the following steps:
 
 * First run a [Standalone Game](Engine/UI/LevelEditor/InEditorTesting/Settings#playinstandalonegame). 
-* Once the Standalone game has been started, bring up the Unreal Console by using the tilde (`) key. 
+* Once the Standalone game has been started, bring up the Unreal Console by using the Backtick (`) key. 
 * Then type `showlog`.
 * Next type `ovrversion`.
 * The current LibOVR version and build date should be printed in the log window like in the following picture:
@@ -53,7 +53,7 @@ to view your work is to run a Standalone Game. To run a Standalone game please d
 
 * Select a Non-PIE game mode like a [Standalone Game](Engine/UI/LevelEditor/InEditorTesting#playineditor). 
 * If the game launches in non-full screen, press **ALT+Enter** to full screen the game window or type **fullscreen** in the [Console](GettingStarted/RunningUnrealEngine#usefulin-gamecommands).
-* Bring up the [Console](GettingStarted/RunningUnrealEngine#usefulin-gamecommands) (`/tilde key) and then type **stereo on**.
+* Bring up the [Console](GettingStarted/RunningUnrealEngine#usefulin-gamecommands) (`/backtick) and then type **stereo on**.
 
 ## Oculus Specific Confg Settings
 
