updated release notes with QA build info

diff --git a/Engine/Documentation/Source/Support/Builds/ReleaseNotes/2015/4_8/ReleaseNotes_4_8.INT.udn b/Engine/Documentation/Source/Support/Builds/ReleaseNotes/2015/4_8/ReleaseNotes_4_8.INT.udn
--- a/Engine/Documentation/Source/Support/Builds/ReleaseNotes/2015/4_8/ReleaseNotes_4_8.INT.udn
+++ b/Engine/Documentation/Source/Support/Builds/ReleaseNotes/2015/4_8/ReleaseNotes_4_8.INT.udn
@@ -5,6 +5,40 @@ Description:4.8 Release Notes
 
 [TOC(start:2 end:2)]
 
+[PUBLISH:Licensee]
+## Build Info
+
+[REGION:simpletable]
+| Item | Info |
+| ---- | ---- |
+| **Labels** | [INCLUDE:#labels] |
+| **Built from Changelist #** | 2582554 (//depot/UE4-QA/Promotable-CL-2582554) |
+| **Branched from Changelist #** | 2577408 (//depot/UE4-Releases/4.8/Promoted-CL-2577408) |
+| **Branched from Changelist #** | 2522625 (//depot/UE4/Promoted-CL-2522625) |
+| [REGION:tablesection]IDEs[/REGION] ||
+| **Visual Studio Version** | Microsoft Visual Studio 2013 12.0.21005.1 REL |
+| **Xcode Version** | 6.3 |
+| [REGION:tablesection]SDKs[/REGION] ||
+| **DirectX** | June 2010 |
+| **HTML5** | 1.25.0 |
+| **Linux Toolchain** | v4_clang-3.5.0_ld-2.24_glibc-2.12.2  _(downloadable from [](http://cdn.unrealengine.com/qfe/v4_clang-3.5.0_ld-2.24_glibc-2.12.2.zip))_ |
+| **Oculus Runtime** | 0.5.0 |
+| **PlayStation 4** | 2.5 (2.0 for Morpheus) |
+| **Steam** | 130  |
+| **Xbox One** | April 2015 XDK (Requires Visual Studio 2012) |
+[/REGION]
+
+<!--
+[EXCERPT:Labels]
+[REGION:plainlist]
+* QA_APPROVED_UE4_BUILD_4.8.0_Release
+* QA_APPROVED_UE4_BUILD_CURRENT  
+* QA_APPROVED_UE4_BUILD_MINIMUM
+[/REGION]
+[/EXCERPT:Labels]
+-->
+[/PUBLISH:Licensee]
+
 This release includes 189 great changes that were submitted from Unreal Engine's amazing community of developers.  Thanks to all of these contributors to Unreal Engine 4.8:
 
 [REDACTED_CONTRIBUTORS]
