#UE4 Docs: Fixing some incorrect wording in the grass how to

diff --git a/Engine/Documentation/Source/Engine/OpenWorldTools/Grass/QuickStart/1/GTQS_1.INT.udn b/Engine/Documentation/Source/Engine/OpenWorldTools/Grass/QuickStart/1/GTQS_1.INT.udn
--- a/Engine/Documentation/Source/Engine/OpenWorldTools/Grass/QuickStart/1/GTQS_1.INT.udn
+++ b/Engine/Documentation/Source/Engine/OpenWorldTools/Grass/QuickStart/1/GTQS_1.INT.udn
@@ -13,7 +13,7 @@ Version: 4.8
 	[PARAM:current]
 	[/PARAM]
 	[PARAM:home]
-		[How-To Home](Engine\OpenWorldTools\Grass\QuickStart "Engine\OpenWorldTools\Grass\QuickStart:Title%")
+		[Quick Start Home](Engine\OpenWorldTools\Grass\QuickStart "Engine\OpenWorldTools\Grass\QuickStart:Title%")
 	[/PARAM]
 	[PARAM:next]
 	[Next Step](Engine\OpenWorldTools\Grass\QuickStart\2 "%Engine\OpenWorldTools\Grass\QuickStart\2:title%") 
diff --git a/Engine/Documentation/Source/Engine/OpenWorldTools/Grass/QuickStart/2/GTQS_2.INT.udn b/Engine/Documentation/Source/Engine/OpenWorldTools/Grass/QuickStart/2/GTQS_2.INT.udn
--- a/Engine/Documentation/Source/Engine/OpenWorldTools/Grass/QuickStart/2/GTQS_2.INT.udn
+++ b/Engine/Documentation/Source/Engine/OpenWorldTools/Grass/QuickStart/2/GTQS_2.INT.udn
@@ -14,7 +14,7 @@ Version: 4.8
 	[PARAM:current]
 	[/PARAM]
 	[PARAM:home]
-		[How-To Home](Engine\OpenWorldTools\Grass\QuickStart "%Engine\OpenWorldTools\Grass\QuickStart:Title%")
+		[Quick Start Home](Engine\OpenWorldTools\Grass\QuickStart "%Engine\OpenWorldTools\Grass\QuickStart:Title%")
 	[/PARAM]
 	[PARAM:next]
 	[Next Step](Engine\OpenWorldTools\Grass\QuickStart\3 "%Engine\OpenWorldTools\Grass\QuickStart\3:title%") 
diff --git a/Engine/Documentation/Source/Engine/OpenWorldTools/Grass/QuickStart/3/GTQS_3.INT.udn b/Engine/Documentation/Source/Engine/OpenWorldTools/Grass/QuickStart/3/GTQS_3.INT.udn
--- a/Engine/Documentation/Source/Engine/OpenWorldTools/Grass/QuickStart/3/GTQS_3.INT.udn
+++ b/Engine/Documentation/Source/Engine/OpenWorldTools/Grass/QuickStart/3/GTQS_3.INT.udn
@@ -14,7 +14,7 @@ Version: 4.8
 	[PARAM:current]
 	[/PARAM]
 	[PARAM:home]
-		[How-To Home](Engine\OpenWorldTools\Grass\QuickStart "%Engine\OpenWorldTools\Grass\QuickStart:Title%")
+		[Quick Start Home](Engine\OpenWorldTools\Grass\QuickStart "%Engine\OpenWorldTools\Grass\QuickStart:Title%")
 	[/PARAM]
 	[PARAM:next]
 	[Next Step](Engine\OpenWorldTools\Grass\QuickStart\4 "%Engine\OpenWorldTools\Grass\QuickStart\4:title%") 
diff --git a/Engine/Documentation/Source/Engine/OpenWorldTools/Grass/QuickStart/4/GTQS_4.INT.udn b/Engine/Documentation/Source/Engine/OpenWorldTools/Grass/QuickStart/4/GTQS_4.INT.udn
--- a/Engine/Documentation/Source/Engine/OpenWorldTools/Grass/QuickStart/4/GTQS_4.INT.udn
+++ b/Engine/Documentation/Source/Engine/OpenWorldTools/Grass/QuickStart/4/GTQS_4.INT.udn
@@ -14,7 +14,7 @@ Version: 4.8
 	[PARAM:current]
 	[/PARAM]
 	[PARAM:home]
-		[How-To Home](Engine\OpenWorldTools\Grass\QuickStart "%Engine\OpenWorldTools\Grass\QuickStart:Title%")
+		[Quick Start Home](Engine\OpenWorldTools\Grass\QuickStart "%Engine\OpenWorldTools\Grass\QuickStart:Title%")
 	[/PARAM]
 	[PARAM:next]
 	[Next Step](Engine\OpenWorldTools\Grass\QuickStart\5 "%Engine\OpenWorldTools\Grass\QuickStart\5:title%") 
diff --git a/Engine/Documentation/Source/Engine/OpenWorldTools/Grass/QuickStart/5/GTQS_5.INT.udn b/Engine/Documentation/Source/Engine/OpenWorldTools/Grass/QuickStart/5/GTQS_5.INT.udn
--- a/Engine/Documentation/Source/Engine/OpenWorldTools/Grass/QuickStart/5/GTQS_5.INT.udn
+++ b/Engine/Documentation/Source/Engine/OpenWorldTools/Grass/QuickStart/5/GTQS_5.INT.udn
@@ -14,7 +14,7 @@ Version: 4.8
 	[PARAM:current]
 	[/PARAM]
 	[PARAM:home]
-		[How-To Home](Engine\OpenWorldTools\Grass\QuickStart "%Engine\OpenWorldTools\Grass\QuickStart:Title%")
+		[Quick Start Home](Engine\OpenWorldTools\Grass\QuickStart "%Engine\OpenWorldTools\Grass\QuickStart:Title%")
 	[/PARAM]
 	[PARAM:next]
 	[/PARAM]
diff --git a/Engine/Documentation/Source/Engine/OpenWorldTools/Grass/QuickStart/GrassToolsQuickStart.INT.udn b/Engine/Documentation/Source/Engine/OpenWorldTools/Grass/QuickStart/GrassToolsQuickStart.INT.udn
--- a/Engine/Documentation/Source/Engine/OpenWorldTools/Grass/QuickStart/GrassToolsQuickStart.INT.udn
+++ b/Engine/Documentation/Source/Engine/OpenWorldTools/Grass/QuickStart/GrassToolsQuickStart.INT.udn
@@ -13,7 +13,7 @@ For example the Procedural Foliage tool is great at simulating huge forests that
 However if you try and use the Procedural Foliage tool to make the Landscape Terrain look like it was densely covered in grass or flowers you will quickly find out that getting the results you want can be very time consuming. 
 To address this issue and give you the ability to densely spawn grass and flowers the Open World Tools have an optional tool called the **Grass Tool**.
 The Grass Tool allows you to define a Static Mesh to spawn wherever a certain Texture is painted on the Landscape terrain.
-In the following How - To we will cover all you need to know about using the Grass Tool in your levels.
+In the following Quick Start we will cover all you need to know about using the Grass Tool in your levels.
 
 ![](T_Grass_Intro.png)
 
