#UE4 Docs: Fixing the home button to say how to instead of quick start.

diff --git a/Engine/Documentation/Source/Engine/OpenWorldTools/ProceduralFoliage/QuickStart/1/PFQS_1.INT.udn b/Engine/Documentation/Source/Engine/OpenWorldTools/ProceduralFoliage/QuickStart/1/PFQS_1.INT.udn
--- a/Engine/Documentation/Source/Engine/OpenWorldTools/ProceduralFoliage/QuickStart/1/PFQS_1.INT.udn
+++ b/Engine/Documentation/Source/Engine/OpenWorldTools/ProceduralFoliage/QuickStart/1/PFQS_1.INT.udn
@@ -13,7 +13,7 @@ Version: 4.8
 	[PARAM:current]
 	[/PARAM]
 	[PARAM:home]
-		[How-To Home](Engine\OpenWorldTools\ProceduralFoliage\QuickStart "Engine\OpenWorldTools\ProceduralFoliage\QuickStart:Title%")
+		[Quick Start Home](Engine\OpenWorldTools\ProceduralFoliage\QuickStart "Engine\OpenWorldTools\ProceduralFoliage\QuickStart:Title%")
 	[/PARAM]
 	[PARAM:next]
 	[Next Step](Engine\OpenWorldTools\ProceduralFoliage\QuickStart\2 "%Engine\OpenWorldTools\ProceduralFoliage\QuickStart\2:title%") 
diff --git a/Engine/Documentation/Source/Engine/OpenWorldTools/ProceduralFoliage/QuickStart/2/PFQS_2.INT.udn b/Engine/Documentation/Source/Engine/OpenWorldTools/ProceduralFoliage/QuickStart/2/PFQS_2.INT.udn
--- a/Engine/Documentation/Source/Engine/OpenWorldTools/ProceduralFoliage/QuickStart/2/PFQS_2.INT.udn
+++ b/Engine/Documentation/Source/Engine/OpenWorldTools/ProceduralFoliage/QuickStart/2/PFQS_2.INT.udn
@@ -14,7 +14,7 @@ Version: 4.8
 	[PARAM:current]
 	[/PARAM]
 	[PARAM:home]
-		[How-To Home](Engine\OpenWorldTools\ProceduralFoliage\QuickStart "Engine\OpenWorldTools\ProceduralFoliage\QuickStart:Title%")
+		[Quick Start Home](Engine\OpenWorldTools\ProceduralFoliage\QuickStart "Engine\OpenWorldTools\ProceduralFoliage\QuickStart:Title%")
 	[/PARAM]
 	[PARAM:next]
 	[Next Step](Engine\OpenWorldTools\ProceduralFoliage\QuickStart\3 "%Engine\OpenWorldTools\ProceduralFoliage\QuickStart\3:title%") 
diff --git a/Engine/Documentation/Source/Engine/OpenWorldTools/ProceduralFoliage/QuickStart/3/PFQS_3.INT.udn b/Engine/Documentation/Source/Engine/OpenWorldTools/ProceduralFoliage/QuickStart/3/PFQS_3.INT.udn
--- a/Engine/Documentation/Source/Engine/OpenWorldTools/ProceduralFoliage/QuickStart/3/PFQS_3.INT.udn
+++ b/Engine/Documentation/Source/Engine/OpenWorldTools/ProceduralFoliage/QuickStart/3/PFQS_3.INT.udn
@@ -14,7 +14,7 @@ Version: 4.8
 	[PARAM:current]
 	[/PARAM]
 	[PARAM:home]
-		[How-To Home](Engine\OpenWorldTools\ProceduralFoliage\QuickStart "Engine\OpenWorldTools\ProceduralFoliage\QuickStart:Title%")
+		[Quick Start Home](Engine\OpenWorldTools\ProceduralFoliage\QuickStart "Engine\OpenWorldTools\ProceduralFoliage\QuickStart:Title%")
 	[/PARAM]
 	[PARAM:next]
 	[Next Step](Engine\OpenWorldTools\ProceduralFoliage\QuickStart\4 "%Engine\OpenWorldTools\ProceduralFoliage\QuickStart\4:title%") 
diff --git a/Engine/Documentation/Source/Engine/OpenWorldTools/ProceduralFoliage/QuickStart/4/PFQS_4.INT.udn b/Engine/Documentation/Source/Engine/OpenWorldTools/ProceduralFoliage/QuickStart/4/PFQS_4.INT.udn
--- a/Engine/Documentation/Source/Engine/OpenWorldTools/ProceduralFoliage/QuickStart/4/PFQS_4.INT.udn
+++ b/Engine/Documentation/Source/Engine/OpenWorldTools/ProceduralFoliage/QuickStart/4/PFQS_4.INT.udn
@@ -14,7 +14,7 @@ Version: 4.8
 	[PARAM:current]
 	[/PARAM]
 	[PARAM:home]
-		[How-To Home](Engine\OpenWorldTools\ProceduralFoliage\QuickStart "Engine\OpenWorldTools\ProceduralFoliage\QuickStart:Title%")
+		[Quick Start Home](Engine\OpenWorldTools\ProceduralFoliage\QuickStart "Engine\OpenWorldTools\ProceduralFoliage\QuickStart:Title%")
 	[/PARAM]
 	[PARAM:next]
 	[Next Step](Engine\OpenWorldTools\ProceduralFoliage\QuickStart\5 "%Engine\OpenWorldTools\ProceduralFoliage\QuickStart\5:title%") 
diff --git a/Engine/Documentation/Source/Engine/OpenWorldTools/ProceduralFoliage/QuickStart/5/PFQS_5.INT.udn b/Engine/Documentation/Source/Engine/OpenWorldTools/ProceduralFoliage/QuickStart/5/PFQS_5.INT.udn
--- a/Engine/Documentation/Source/Engine/OpenWorldTools/ProceduralFoliage/QuickStart/5/PFQS_5.INT.udn
+++ b/Engine/Documentation/Source/Engine/OpenWorldTools/ProceduralFoliage/QuickStart/5/PFQS_5.INT.udn
@@ -14,7 +14,7 @@ Version: 4.8
 	[PARAM:current]
 	[/PARAM]
 	[PARAM:home]
-		[How-To Home](Engine\OpenWorldTools\ProceduralFoliage\QuickStart "Engine\OpenWorldTools\ProceduralFoliage\QuickStart:Title%")
+		[Quick Start Home](Engine\OpenWorldTools\ProceduralFoliage\QuickStart "Engine\OpenWorldTools\ProceduralFoliage\QuickStart:Title%")
 	[/PARAM]
 	[PARAM:next]
 	[Next Step](Engine\OpenWorldTools\ProceduralFoliage\QuickStart\6 "%Engine\OpenWorldTools\ProceduralFoliage\QuickStart\6:title%") 
diff --git a/Engine/Documentation/Source/Engine/OpenWorldTools/ProceduralFoliage/QuickStart/6/PFQS_6.INT.udn b/Engine/Documentation/Source/Engine/OpenWorldTools/ProceduralFoliage/QuickStart/6/PFQS_6.INT.udn
--- a/Engine/Documentation/Source/Engine/OpenWorldTools/ProceduralFoliage/QuickStart/6/PFQS_6.INT.udn
+++ b/Engine/Documentation/Source/Engine/OpenWorldTools/ProceduralFoliage/QuickStart/6/PFQS_6.INT.udn
@@ -14,7 +14,7 @@ Version: 4.8
 	[PARAM:current]
 	[/PARAM]
 	[PARAM:home]
-		[How-To Home](Engine\OpenWorldTools\ProceduralFoliage\QuickStart "Engine\OpenWorldTools\ProceduralFoliage\QuickStart:Title%")
+		[Quick Start Home](Engine\OpenWorldTools\ProceduralFoliage\QuickStart "Engine\OpenWorldTools\ProceduralFoliage\QuickStart:Title%")
 	[/PARAM]
 	[PARAM:next]
 
