#UE4 Docs: Fixing the way Gear VR is used  based on feedback from Samsung.

diff --git a/Engine/Documentation/Source/Platforms/GearVR/QuickStart/3/GearVRQuickStart_3.INT.udn b/Engine/Documentation/Source/Platforms/GearVR/QuickStart/3/GearVRQuickStart_3.INT.udn
--- a/Engine/Documentation/Source/Platforms/GearVR/QuickStart/3/GearVRQuickStart_3.INT.udn
+++ b/Engine/Documentation/Source/Platforms/GearVR/QuickStart/3/GearVRQuickStart_3.INT.udn
@@ -4,7 +4,11 @@ Crumbs:%ROOT%, Platforms/GearVR
 Description:Instructions on how to deploy projects form UE4 to use with the Samsung Gear VR Headset.
 Navigation:platform
 Platform:Gear VR
-Related:
+Related: Platforms\GearVR\Prerequisites
+Related: Platforms\GearVR\BestPractices
+Related: Platforms\GearVR\Debugging
+
+
 
 [VAR:Steps]
 [OBJECT:Navigation]
