updated Oculus version number

diff --git a/Engine/Documentation/Source/Support/Builds/ReleaseNotes/2015/4_8/ReleaseNotes_4_8.INT.udn b/Engine/Documentation/Source/Support/Builds/ReleaseNotes/2015/4_8/ReleaseNotes_4_8.INT.udn
--- a/Engine/Documentation/Source/Support/Builds/ReleaseNotes/2015/4_8/ReleaseNotes_4_8.INT.udn
+++ b/Engine/Documentation/Source/Support/Builds/ReleaseNotes/2015/4_8/ReleaseNotes_4_8.INT.udn
@@ -22,7 +22,7 @@ Description:4.8 Release Notes
 | **DirectX** | June 2010 |
 | **HTML5** | 1.25.0 |
 | **Linux Toolchain** | v4_clang-3.5.0_ld-2.24_glibc-2.12.2  _(downloadable from [](http://cdn.unrealengine.com/qfe/v4_clang-3.5.0_ld-2.24_glibc-2.12.2.zip))_ |
-| **Oculus Runtime** | 0.5.0 |
+| **Oculus Runtime** | 0.6.0 |
 | **PlayStation 4** | 2.5 (2.0 for Morpheus) |
 | **Steam** | 130  |
 | **Xbox One** | April 2015 XDK (Requires Visual Studio 2012) |
