#UE4 Docs: Fixing the intro text to be more inline with the Proc Foliage tool per luis request.

diff --git a/Engine/Documentation/Source/Engine/OpenWorldTools/Grass/QuickStart/GrassToolsQuickStart.INT.udn b/Engine/Documentation/Source/Engine/OpenWorldTools/Grass/QuickStart/GrassToolsQuickStart.INT.udn
--- a/Engine/Documentation/Source/Engine/OpenWorldTools/Grass/QuickStart/GrassToolsQuickStart.INT.udn
+++ b/Engine/Documentation/Source/Engine/OpenWorldTools/Grass/QuickStart/GrassToolsQuickStart.INT.udn
@@ -8,15 +8,15 @@ Version: 4.8
 
 ## Overview
 
-As with anything in life using the right tool for the job is often the determining factor between success and failure and the same can be said when talking about the Open World tools inside of UE4. 
-For example the Procedural Foliage tool is great at simulating huge forests that are composed of many different species of trees and bushes. 
-However if you try and use the Procedural Foliage tool to make the Landscape Terrain look like it was densely covered in grass or flowers you will quickly find out that getting the results you want can be very time consuming. 
-To address this issue and give you the ability to densely spawn grass and flowers the Open World Tools have an optional tool called the **Grass Tool**.
-The Grass Tool allows you to define a Static Mesh to spawn wherever a certain Texture is painted on the Landscape terrain.
-In the following Quick Start we will cover all you need to know about using the Grass Tool in your levels.
+In this Quick Start tutorial we will take a look at how you can setup and use the new **GrassTool** that are now available with the 4.8 release of Unreal Engine 4 (UE4).
+Over the course of this Quick Start tutorial you will learn how to create, setup, and spawn static meshes to make your Landscape look like it is densely covered in thick grass. 
+You will also be introduced to key properties and settings that will help you shape your virtual grasslands to fit your project's needs. 
 
 ![](T_Grass_Intro.png)
 
+You will also be exposed to all of the required Actors and properties the Grass tool requires to function correctly and deliver the results you want.
+When you have completed this Quick Start you will have a new level that should look similar to the image above. 
+
 [region:note]
 Currently the Grass system will only work with the Landscape Terrain Actor. Trying to use the Grass system to spawn grass on any other Actor type will not work.
 [/region]
