#UE4 Docs: Fixing an issue with the image and text being oposite of one another.

diff --git a/Engine/Documentation/Source/Engine/OpenWorldTools/TipsAndTricks/TipsAndTricks.INT.udn b/Engine/Documentation/Source/Engine/OpenWorldTools/TipsAndTricks/TipsAndTricks.INT.udn
--- a/Engine/Documentation/Source/Engine/OpenWorldTools/TipsAndTricks/TipsAndTricks.INT.udn
+++ b/Engine/Documentation/Source/Engine/OpenWorldTools/TipsAndTricks/TipsAndTricks.INT.udn
@@ -79,7 +79,7 @@ When this Material is applied to a Landscape in UE4 you should get results that
 
 ![](T_Masking_In_Action.png)
 
-In the image above, the grass Static Mesh, for the most part, is spawned only in areas that use the green channel of the mask Texture, while the rocks use the red channel of the mask Texture to spawn. 
+In the image above, the grass Static Mesh, for the most part, are spawned only in areas that use the red channel of the mask Texture, while the rocks use the green channel of the mask Texture to spawn. 
 As you might have noticed, the grass and rocks are not entirely contained within the area of the mask Texture that spawns them. 
 This is expected behavior as the placement of the Static Meshes can be jittered and offset to make things look more natural.
 
