#UE4 Docs: Moved the prerec page to it' own folder to give make the Gear VR landing page like the other landing pages.

diff --git a/Engine/Documentation/Source/Platforms/GearVR/Prerequisites/GearVRPrerequisites.INT.udn b/Engine/Documentation/Source/Platforms/GearVR/Prerequisites/GearVRPrerequisites.INT.udn
new file mode 100644
--- /dev/null
+++ b/Engine/Documentation/Source/Platforms/GearVR/Prerequisites/GearVRPrerequisites.INT.udn
@@ -0,0 +1,88 @@
+Title: Samsung Gear VR Headset Prerequisites
+Description:Required initial setup for the Gear VR Headset
+Availability:Public
+Crumbs:
+Navigation:platform
+Platform:Gear VR
+
+![](Platforms\GearVR\QuickStart\GVR_Headset_00.png)
+
+[REGION:warning]
+Developing for the Gear VR Headset requires **Unreal Engine 4.7.4** or later.
+[/REGION]
+
+## Overview 
+
+With the 4.7 release of Unreal Engine 4 (UE4) you can now build Virtual Reality(VR) projects for the Samsung Gear VR Headset that can transform your Samsung Galaxy Note 4, Galaxy S6, and Galaxy S6 Edge Smartphone into an immersive VR device. 
+In the following sections we will go over everything you need to know to get your Samsung Galaxy Note 4 Smartphone properly setup so that it can work with the Samsung Gear VR headset.
+
+## Gear VR Official Documentation
+
+For a more in-depth look at how the Gear VR SDK works or for more information on Oculus Mobile VR development check out the official Oculus documents linked below.
+
+* **[Oculus Mobile Software Development Kit(SDK) Device and Environment Setup Guide](http://static.oculus.com/sdk-downloads/documents/Oculus_Mobile_v0.4.0_Device_and_Environment_Setup_Guide.pdf)**
+* **[Oculus Mobile Development Documentation](http://static.oculus.com/sdk-downloads/documents/Oculus_Mobile_v0.5.0_SDK_Documentation.pdf)**
+
+
+
+## Gear VR Headset Pre-requisites 
+
+Before we can get into UE4 Gear VR development we first must make sure that the following prerequisites are fulfilled as if we do not, we run the risk of our Gear VR projects not working with the Gear VR Headset. 
+
+* Make sure you have the USB cord used for charging your Samsung Galaxy Note 4 Smartphone available and plugged into a USB port on the PC you are developing on. 
+The USB cord is of vital importance as this is how you transfer data between your development PC and your Smartphone.
+
+* Only certain types of Samsung Galaxy Smartphone's will work with the Gear VR Headset. 
+	The **[Oculus website](https://support.oculus.com/hc/en-us/articles/204113096-Samsung-Gear-VR-Innovator-Edition-Samsung-Galaxy-Note-4-Compatibility)** has a list of phone models that will work with the Gear VR Headset.
+	Make sure to double check that the Samsung Galaxy Note 4 Smartphone you are buying is the correct model that will work with the Gear VR Headset.
+
+* Make sure to download and install the specific USB drivers for your Galaxy Note 4 from the **[Samsung Website](http://www.samsung.com/us/support/downloads)**. 
+
+* Download Unreal Engine version **4.7.4** or later from the **[UE4 Launcher](https://www.unrealengine.com/)** or **[Github](https://www.unrealengine.com/ue4-on-github)**. 
+
+* Install the **Tegra Android Development Pack** which can be found in **Engine\Extras\Android**.  
+
+* Check to see if **[Developer Mode](https://docs.unrealengine.com/latest/INT/Platforms/Android/GettingStarted/2/index.html )** has been enabled on your Samsung Galaxy Note 4 Smartphone and if not, enable it. 
+
+* Also make sure to check that USB Debugging has been enabled to ensure that you can load data onto your Smartphone from your PC.
+	
+	![](Platforms\GearVR\GVR_Dev_Options_Enabled.png)
+
+* With your Smartphone plugged into your development PC via USB, make sure to accept the RSA Fingerprint key from your PC.
+This way your development PC and Smartphone can send data back and forth.
+Also make sure to check **Always allow this computer** to ensure that your development PC can always interact with this Smartphone when plugged in. 
+
+	![](Platforms\GearVR\GVR_USB_Debuging.png)
+
+* Generate and download the **OSIG File** for your Samsung Smartphone using your Smartphone's Device ID and the **[Oculus](https://developer.oculus.com/tools/osig/)** website.
+If you are not sure how to get your Smartphone's **Device ID** the **[Oculus website](https://developer.oculus.com/tools/osig/)** has detailed directions on how to do this.
+	
+	![](Platforms\GearVR\GVR_OSIG_Website.png)
+
+* After you download the OSIG file, create a folder called **assets** and place the **OSIG** file inside of it. 
+Then place that folder in **Engine/Build/Android/Java**. 
+
+	[region:warning]
+	**If you do not do this your project will not be able to run on your Galaxy Smartphone.**
+	[/region]
+	
+	![](Platforms\GearVR\GVR_OSIG_File_Location.png)
+
+It is very important that you make sure all items listed in the Gear VR Prerequisites section are completed before moving on to the next section. 
+Skipping a step could result in your project not working with the Gear VR Headset. 
+
+[OBJECT:TopicButtonList]
+	[PARAM:icon]
+		![](%ROOT%/vr_icon.png)(convert:false)
+	[/PARAM]
+	[PARAM:title]
+		Gear VR Development
+	[/PARAM]
+	[PARAM:description]
+	[/PARAM]
+	[PARAM:links]
+		* [](Platforms/GearVR/QuickStart "%Platforms/GearVR/QuickStart:description%")
+		* [](Platforms/GearVR/BestPractices "%Platforms/GearVR/BestPractices:description%")
+		* [](Platforms/GearVR/Debugging "%Platforms/GearVR/BestPractices:description%")
+	[/PARAM]
+[/OBJECT]
\ No newline at end of file
