#UE4 Docs: Adding Mac specific info to the android pages as well as imgaes.

diff --git a/Engine/Documentation/Source/Platforms/Android/GettingStarted/1/AndroidQS_1.INT.udn b/Engine/Documentation/Source/Platforms/Android/GettingStarted/1/AndroidQS_1.INT.udn
--- a/Engine/Documentation/Source/Platforms/Android/GettingStarted/1/AndroidQS_1.INT.udn
+++ b/Engine/Documentation/Source/Platforms/Android/GettingStarted/1/AndroidQS_1.INT.udn
@@ -20,60 +20,106 @@ Description: Install the TADP Android SDK.
 [/OBJECT]
 [/VAR]
 
+%Steps%
 
+%Globals:OSSelect%
 
+[OBJECT:ToggleButtonContent]
+	[PARAMLITERAL:category]
+	OS
+	[/PARAMLITERAL]
+	[PARAMLITERAL:id]
+	windows
+	[/PARAMLITERAL]
+	[PARAMLITERAL:active]
+	active_button_content
+	[/PARAMLITERAL]
+	[PARAM:content]
+	[INCLUDE:#WindowsSetup]
+	[/PARAM]
+[/OBJECT]
 
-%Steps%
+[OBJECT:ToggleButtonContent]
+	[PARAMLITERAL:category]
+	OS
+	[/PARAMLITERAL]
+	[PARAMLITERAL:id]
+	mac
+	[/PARAMLITERAL]
+	[PARAMLITERAL:active]
+	[/PARAMLITERAL]
+	[PARAM:content]
+	[INCLUDE:#MacOSSetup]
+	[/PARAM]
+[/OBJECT]
+
+<!--
+[EXCERPT:WindowsSetup]
 
 [REGION:tip]
 If you already have the Android SDK installed, but your environment variables have been lost, you can skip all of this and simply point the editor to the install locations using the **Edit Menu** -> **Project Settings** -> **Android SDK** properties in the Editor.
 [/REGION]
 
-1.	Navigate to the TADP installation folder:
+1. Navigate to the TADP installation folder:
 | Install Type | Location |
 | ------------ | ---------|
 | Binary | [EngineInstallLocation]\Unreal Engine\[EngineVersion]\Engine\Extras\Android |
 | GitHub | [EngineSyncLocation\UE4\Engine\Extras\Android |
-1.	**Double-click** on `tadp-2.0r8-windows.exe`.
-1.	Click **Next**.
+
+1. **Double-click** on `tadp-2.0r8-windows.exe`.
+
+1. Click **Next**.
     
 	![](tadpinstaller_1.png)(w:450)
+
 1.	Select **I accept the agreement**, then click **Next**.
 
 	![](tadpinstaller_2.png)(w:450)
+
 1.	Click **Next**.
 
 	![](tadpinstaller_3.png)(w:450)
+
 1.	Select the **Express** install, then click **Next**.
 	
 	![](tadpinstaller_4.png)(w:450)
+
 1.	Make sure that the install location is `C:\NVPACK`.
 	
 	![](tadpinstaller_5.png)(w:450) 
+
 1.	Verify that the list of components to install matches this list:
 	
 	![](tadpinstaller_6.png)(w:450)
+
 1.	Enter a proxy if needed:
 	
 	![](tadpinstaller_7.png)(w:450)
+
 1.	Click **Next**. You will need an internet connection and this step may take a few minutes.
 	
 	![](tadpinstaller_8.png)(w:450)
+
 1.	Click **Next**.
 	
 	![](tadpinstaller_9.png)(w:450)
+
 1.	The installer should run, and this will take some time.
 	
 	![](tadpinstaller_10.png)(w:450)
+
 1.	If you are prompted with a warning about the software publisher, select **Install this driver software anyway**.
 	
 	![](tadpinstaller_11.png)
+
 1.	When prompted about saving files after the install, you can select either **Yes** or **No**.
 	
 	![](tadpinstaller_12.png)
+
 1.	Click **Next**.
 	
 	![](tadpinstaller_13.png)(w:450)
+
 1.	Verify that all the components have been installed correctly.
 	
 	![](tadpinstaller_14.png)(w:450)
@@ -81,5 +127,95 @@ If you already have the Android SDK installed, but your environment variables ha
 	
 	![](tadpinstaller_15.png)(w:450)
 
+[/EXCERPT:WindowsSetup]
+
+[EXCERPT:MacOSSetup]
+
+[region:tip]
+It is best to do this step while the editor and launcher are both closed. If you are installing the SDK with them open, you will need to restart both the editor and the launcher before you can work with an Android device.
+[/region]
+
+TADP installer and Android SDK tools require that both Java 6 and Java 7 are installed on your Mac. 
+
+* **Java 6** can be downloaded from [here](http://support.apple.com/kb/DL1572).
+* **Java 7** can be downloaded from [here](http://www.oracle.com/technetwork/java/javase/downloads/jre7-downloads-1880261.html).
+
+1. Navigate to the TADP installation folder:
+| Install Type  | Location  |
+| --- | --- |
+| Binary | [EngineInstallLocation]\Unreal Engine[EngineVersion]\Engine\Extras\Android |
+| GitHub | [EngineSyncLocation\UE4\Engine\Extras\Android |
+
+1. **Double-click** on tadp-3.0r4-osx.dmg. This should mount TADP installer disk image and open its window.
+
+1. **Double-click** TADP installer icon.
+
+	![](Tadp_Mac_Setup-00.png)
+
+1. Click **Next**.
+
+	![](Tadp_Mac_Setup-01.png)
+
+1. Select the **I accept the agreement** option and then click **Next**.
+
+	![](Tadp_Mac_Setup-02.png)
+
+1. Click **Next**.
+
+	![](Tadp_Mac_Setup-03.png)
+
+1. Select the **Express** install, then click **Next**.
+
+	![](Tadp_Mac_Setup-04.png)
+
+1. Select the install location. 
+	[region:note]
+	It is perfectly fine to leave this set to the default location.
+	[/region]
+
+	![](Tadp_Mac_Setup-05.png)
+
+1. Verify that the list of components to install matches what is in the image below.
+
+	![](Tadp_Mac_Setup-06.png)
+
+1. Enter a proxy if needed.
+
+	![](Tadp_Mac_Setup-07.png)
+
+1. Click **Next**. 
+	[region:note]
+	For this part of the setup will take a few minutes to completed and you will need an internet connection to complete it.
+	[/region]
+
+	![](Tadp_Mac_Setup-08.png)
+
+1. Click **Next**.
+
+	![](Tadp_Mac_Setup-09.png)
+
+1. The installer will now run and it will take some time for it complete it's install.
+
+	![](Tadp_Mac_Setup-10.png)
+
+1. Click **Next**.
+
+	![](Tadp_Mac_Setup-11.png)
+
+1. Click **Next**.
+
+	![](Tadp_Mac_Setup-12.png)
+
+1. Verify that all the components have been installed correctly.
+
+	![](Tadp_Mac_Setup-13.png)
+
+1. Uncheck both **View Readme File** and **Browse Documentation**, then click **Finish**.
+
+	![](Tadp_Mac_Setup-14.png)
+
+[/EXCERPT:MacOSSetup]
+
+-->
 %Steps%
 
diff --git a/Engine/Documentation/Source/Platforms/Android/GettingStarted/2/AndroidQS_2.INT.udn b/Engine/Documentation/Source/Platforms/Android/GettingStarted/2/AndroidQS_2.INT.udn
--- a/Engine/Documentation/Source/Platforms/Android/GettingStarted/2/AndroidQS_2.INT.udn
+++ b/Engine/Documentation/Source/Platforms/Android/GettingStarted/2/AndroidQS_2.INT.udn
@@ -21,16 +21,47 @@ Description:Set up your Android device for development.
 [/OBJECT]
 [/VAR]
 
+%Steps%
 
+%Globals:OSSelect%
 
+[OBJECT:ToggleButtonContent]
+	[PARAMLITERAL:category]
+	OS
+	[/PARAMLITERAL]
+	[PARAMLITERAL:id]
+	windows
+	[/PARAMLITERAL]
+	[PARAMLITERAL:active]
+	active_button_content
+	[/PARAMLITERAL]
+	[PARAM:content]
+	[INCLUDE:#WindowsSetup]
+	[/PARAM]
+[/OBJECT]
+
+[OBJECT:ToggleButtonContent]
+	[PARAMLITERAL:category]
+	OS
+	[/PARAMLITERAL]
+	[PARAMLITERAL:id]
+	mac
+	[/PARAMLITERAL]
+	[PARAMLITERAL:active]
+	[/PARAMLITERAL]
+	[PARAM:content]
+	[INCLUDE:#MacOSSetup]
+	[/PARAM]
+[/OBJECT]
 
-%Steps%
 
+<!--
+[EXCERPT:WindowsSetup]
 
-1.	Connect your device.
-1.	Device driver software may automatically install, but if not, visit [Android's OEM USB Drivers](http://developer.android.com/tools/extras/oem-usb.html) page to find driver installation links and additional information.  
+1. Connect your device.
+1.	Device driver software may automatically install, but if not, visit **[Android's OEM USB Drivers](http://developer.android.com/tools/extras/oem-usb.html)** page to find driver installation links and additional information.  
 1.	Open the **Settings** application on your device.
-1.	Look for the **Developer Options** menu. If it is not available, you need to enable **Developer Options** on your device. For more about developer options, visit [Android's Using Hardware Devices](http://developer.android.com/tools/device.html) page.
+1.	Look for the **Developer Options** menu. If it is not available, you need to enable **Developer Options** on your device. For more about developer options, visit **[Android's Using Hardware Devices](http://developer.android.com/tools/device.html)** page.
 	1.	Scroll down and select **About Phone**. Depending on your Android device, you may have an **About Device**, **About Tablet**, or **About Shield** option instead.
 		[REGION: note]
 		Note: This may be in a **More** section on newer Android versions.
@@ -47,7 +78,49 @@ Description:Set up your Android device for development.
 	* It is recommended you check the Always allow checkbox.
 	* In some cases, a correctly configured device may not show up if it is connected as a **Media Device (MTP)** instead of as a **Camera (PTP)**. Select the **Camera (PTP)** option if your device is not visible at all under `adb devices` at this point.
 
+[/EXCERPT:WindowsSetup]
+
+[EXCERPT:MacOSSetup]
+
+1. Connect your device to your Mac.
+
+1. Open the **Settings application** on your device.
+
+1. Look for the **Developer Options** menu. If it is not available, you need to enable Developer Options on your device. 
+	For more about developer options, visit **[Android's Using Hardware Devices](http://developer.android.com/tools/device.html)** page.
+	To enable developer options on your device you will need to do the following.
+
+	* Scroll down and select About Phone. 
+	* Depending on your Android device, you may have an **About Device**, **About Tablet**, or **About Shield** option instead.
+
+	[region:note]
+	In newer versions of Android this may be in the More section.
+	[/region]
+
+	* Tap the Build Number **7** times. 
+	* After the first few presses, you will see a message that says, **You are now # steps away from becoming a developer**. 
+	* Once developer mode has been activated you will get a success message displayed on the screen.
+	* Navigate back to the **Settings menu** and select **Developer Options** which is now available.
+
+1. In the **Developer Options** menu, enable **USB debugging**. 
+
+1. Select **OK** when prompted.
+
+1. To verify that your device is set up properly, open the **Terminal app** from /Applications/Utilties.
+
+1. Type **adb devices** in the terminal window and you should see a device listed. If you do not see your device listed here are some troubleshooting points you can try.
+
+	* If the device is listed but shows unauthorized, that indicates that you need to get the device to ask if you want to authorize the Mac.
+
+	* You may need to unplug and reconnect your device to make it show up.
+
+	* It is recommended you check the Always allow checkbox to always allow connections from the computer that it is pluged into.
+
+	* In some cases, a correctly configured device may not show up if it is connected as a Media Device (MTP) instead of as a Camera (PTP). Select the Camera (PTP) option if your device is not visible at all under adb devices at this point.
+
+[/EXCERPT:MacOSSetup]
 
+-->
 
 
 
diff --git a/Engine/Documentation/Source/Platforms/Android/GettingStarted/5/AndroidQS_5.INT.udn b/Engine/Documentation/Source/Platforms/Android/GettingStarted/5/AndroidQS_5.INT.udn
--- a/Engine/Documentation/Source/Platforms/Android/GettingStarted/5/AndroidQS_5.INT.udn
+++ b/Engine/Documentation/Source/Platforms/Android/GettingStarted/5/AndroidQS_5.INT.udn
@@ -19,12 +19,43 @@ Description:Package your game for installation.
 [/OBJECT]
 [/VAR]
 
+%Steps%
 
+%Globals:OSSelect%
 
+[OBJECT:ToggleButtonContent]
+	[PARAMLITERAL:category]
+	OS
+	[/PARAMLITERAL]
+	[PARAMLITERAL:id]
+	windows
+	[/PARAMLITERAL]
+	[PARAMLITERAL:active]
+	active_button_content
+	[/PARAMLITERAL]
+	[PARAM:content]
+	[INCLUDE:#WindowsSetup]
+	[/PARAM]
+[/OBJECT]
 
-%Steps%
+[OBJECT:ToggleButtonContent]
+	[PARAMLITERAL:category]
+	OS
+	[/PARAMLITERAL]
+	[PARAMLITERAL:id]
+	mac
+	[/PARAMLITERAL]
+	[PARAMLITERAL:active]
+	[/PARAMLITERAL]
+	[PARAM:content]
+	[INCLUDE:#MacOSSetup]
+	[/PARAM]
+[/OBJECT]
 
-1.	Now, to push your full game to your Android device, you will need to package it. In the **File** menu, hover over **Package Project** to expand the full list of options. Here you will again see the list of texture formats, and must select the right one (or unlabeled). If you are unsure what to select, use **ETC1**.
+<!--
+[EXCERPT:WindowsSetup]
+
+1.	To push your full game to your Android device, you will need to package it. In the **File** menu, hover over **Package Project** to expand the full list of options. Here you will again see the list of texture formats, and must select the right one (or unlabeled). If you are unsure what to select, use **ETC1**.
 
 	[REGION:raw]  
 	![](package_menu.png)
@@ -41,6 +72,36 @@ Description:Package your game for installation.
 	![](batch_file.png)
 1.	Your game is now installed and you can open it from your Apps drawer.
 
+[/EXCERPT:WindowsSetup]
+
+[EXCERPT:MacOSSetup]
+1. Now, to push your full game to your Android device, you will need to package it. In the **File** menu, hover over **Package Project** to expand the full list of options. Here you will again see the list of texture formats, and must select the right one (or unlabeled).
+
+	![](AQS_Mac_01.png)
+
+We are currently packaging in **Development** because we are testing this game, but when your game is complete, you can package it in **Shipping** by hovering over the **Build Configurations** option in the **Package Project** menu.
+
+1. You will be prompted for a save location for your packaged game, which should default to your project folder.
+
+	![](AQS_Mac_02.png)
+
+1. A packaging message will appear in the bottom right corner of the editor while the game is being packaged.
+
+	![](AQS_Mac_03.png)
+
+1. Once packaging is successful, a Packaging Complete message will appear in the bottom right of the editor.
+
+1. Navigate to the save location you just set.
+
+1. In this folder, there will be a folder named specifically for your texture type. In this example, the folder is named Android_DXT. Open this folder.
+
+1. **Double-click** on the Install_[ProjectName]_Development-arm*.command script file.
+
+	![](AQS_Mac_04.png)
+
+1. Your game is now installed and you can open it from your Apps drawer.
 
+[/EXCERPT:MacOSSetup]
+-->
 
 %Steps%
\ No newline at end of file
diff --git a/Engine/Documentation/Source/Platforms/Android/GettingStarted/AndroidQuickStart.INT.udn b/Engine/Documentation/Source/Platforms/Android/GettingStarted/AndroidQuickStart.INT.udn
--- a/Engine/Documentation/Source/Platforms/Android/GettingStarted/AndroidQuickStart.INT.udn
+++ b/Engine/Documentation/Source/Platforms/Android/GettingStarted/AndroidQuickStart.INT.udn
@@ -3,43 +3,12 @@ Title: Android Quick Start
 Crumbs:%ROOT%, Platforms, Platforms/Android
 Description:Setting up for development for the Android platform.
 
-[VAR:Topic]
-[OBJECT:Topic]
-	[PARAM:image]
-		![](AndroidTopic.jpg)
-	[/PARAM]
-	[PARAM:title]
-		%Platforms/Android/GettingStarted:title%
-	[/PARAM]
-	[PARAM:path]
-		[RELATIVE:Platforms/Android/GettingStarted]
-	[/PARAM]
-	[PARAM:description]
-	%Platforms/Android/GettingStarted:description%
-	[/PARAM]
-[/OBJECT]
-[/VAR]
 
-[VAR:TopicCompact]
-[OBJECT:TopicCompact]
-	[PARAM:image]
-		![](AndroidTopic.jpg)
-	[/PARAM]
-	[PARAM:title]
-		%Platforms/Android/GettingStarted:title%
-	[/PARAM]
-	[PARAM:path]
-		[RELATIVE:Platforms/Android/GettingStarted]
-	[/PARAM]
-	[PARAM:description]
-	%Platforms/Android/GettingStarted:description%
-	[/PARAM]
-[/OBJECT]
-[/VAR]
 
 The Android Quick Start will walk you through setting up your computer and your Android device for deployment, and putting a template game on your device using both the launching and packaging features.
 This guide uses TADP as the Android SDK, but if you would like to know how to set up another SDK for use with Unreal Engine, see the [](Platforms\Android\Reference) page.
 
+
 [OBJECT:TopicButtonList]
 	[PARAM:icon]
 		![](%ROOT%/build_icon.png)(convert:false)
diff --git a/Engine/Documentation/Source/Platforms/Android/Reference/AndroidGettingStarted.INT.udn b/Engine/Documentation/Source/Platforms/Android/Reference/AndroidGettingStarted.INT.udn
--- a/Engine/Documentation/Source/Platforms/Android/Reference/AndroidGettingStarted.INT.udn
+++ b/Engine/Documentation/Source/Platforms/Android/Reference/AndroidGettingStarted.INT.udn
@@ -5,6 +5,9 @@ Description:How to install different Android SDKs, set environment variables, an
 
 [TOC(start:2 end:3)]
 
+%Globals:OSSelect%
+
+
 
 ## Overview 
 
@@ -45,6 +48,37 @@ Here is what we currently require for Android game development:
 There are several methods for getting the components. If you have already installed one of these, look in the corresponding section for additional requirements:
 
 ### TADP (Tegra Android Development Pack)
+[OBJECT:ToggleButtonContent]
+	[PARAMLITERAL:category]
+	OS
+	[/PARAMLITERAL]
+	[PARAMLITERAL:id]
+	windows
+	[/PARAMLITERAL]
+	[PARAMLITERAL:active]
+	active_button_content
+	[/PARAMLITERAL]
+	[PARAM:content]
+	[INCLUDE:#TADPWindowsSetup]
+	[/PARAM]
+[/OBJECT]
+
+[OBJECT:ToggleButtonContent]
+	[PARAMLITERAL:category]
+	OS
+	[/PARAMLITERAL]
+	[PARAMLITERAL:id]
+	mac
+	[/PARAMLITERAL]
+	[PARAMLITERAL:active]
+	[/PARAMLITERAL]
+	[PARAM:content]
+	[INCLUDE:#TADPMacOSSetup]
+	[/PARAM]
+[/OBJECT]
+
+<!--
+[EXCERPT:TADPWindowsSetup]
 
 TADP is the Tegra Android Development Pack, by Nvidia. This is a one-stop shop for Android development. If you install TADP, all of our requirements will be met.
 
@@ -57,16 +91,41 @@ The important pieces are:
 * NDK
 * JDK
 * Ant 
-* Nsight 
+* Night 
 	
 If you use the default installation set, a pretty minimal set of components will be installed. If you customize the install, you can add more items (like profilers, etc). Note that during installation, 
 one or two of the sub-installers may appear to hang, but they will eventually finish. The install could take as long as 30 to 45 minutes. If you never plan to write code, you may want to 
 select **Customize**, and only install the SDK, Ant, and Java (skipping the NDK and Nsight debugger) to make for the quickest installation. The Nsight portion of TADP is used for debugging an Android game on a device. 
 It is not required to compile C++ code.
 
+![](TADPSetup.png)
+
+[/EXCERPT:TADPWindowsSetup]
 
+[EXCERPT:TADPMacOSSetup]
+TADP is the Tegra Android Development Pack, by Nvidia. This is a one-stop shop for Android development. If you install TADP, all of our requirements will be met.
 
-![](TADPSetup.png)
+The installer for TADP is included with Unreal Engine 4 for Mac. It can be found here: 
+
+	[ENGINE INSTALL LOCATION]/Engine/Extras/Android/tadp-3.0r4-osx.dmg
+
+The important pieces are:
+
+* SDK
+* NDK
+* Ant 
+
+If you use the default installation set, a pretty minimal set of components will be installed. 
+If you customize the install, you can add more items (like profilers, etc). Note that during installation, one or two of the sub-installers may appear to hang, but they will eventually finish. 
+The install could take as long as 30 to 45 minutes on some computers. 
+If you never plan to write code, you may want to select **Customize**, and only install the SDK, Ant, and Java (skipping the NDK and Nsight debugger) to make for the quickest installation.
+The Nsight portion of TADP is used for debugging an Android game on a device so t is not required to compile C++ code.
+
+![](TADP_Mac_OS.png)
+
+[/EXCERPT:TADPMacOSSetup]
+
+-->
 
 ### ADT (Android Development Tools)
 
@@ -99,6 +158,38 @@ For code projects only, you will need to install the NDK:
 
 ## Environment Variables
 
+[OBJECT:ToggleButtonContent]
+	[PARAMLITERAL:category]
+	OS
+	[/PARAMLITERAL]
+	[PARAMLITERAL:id]
+	windows
+	[/PARAMLITERAL]
+	[PARAMLITERAL:active]
+	active_button_content
+	[/PARAMLITERAL]
+	[PARAM:content]
+	[INCLUDE:#ENVWindowsSetup]
+	[/PARAM]
+[/OBJECT]
+
+[OBJECT:ToggleButtonContent]
+	[PARAMLITERAL:category]
+	OS
+	[/PARAMLITERAL]
+	[PARAMLITERAL:id]
+	mac
+	[/PARAMLITERAL]
+	[PARAMLITERAL:active]
+	[/PARAMLITERAL]
+	[PARAM:content]
+	[INCLUDE:#ENVMacOSSetup]
+	[/PARAM]
+[/OBJECT]
+
+<!--
+[EXCERPT:ENVWindowsSetup]
+
 Unless you installed TADP, you will need to verify that some environment variables are set properly. 
 
 [REGION:note]
@@ -139,8 +230,78 @@ Here are the variables that we depend on, and where they should point:
 * Set this to where you unzipped the NDK. It will probably be called something like `android-ndk-r9c`.
 * It must contain a file called `ndk-build.cmd`
 
+[/EXCERPT:ENVWindowsSetup]
+
+[EXCERPT:ENVMacOSSetup]
+Unless you installed TADP, you will need to verify that some environment variables are set properly. 
+
+[region:note]
+Once you have completed the installation (either by editing your .bash_profile file, or by installing TADP), you will need to restart both the Editor and Launcher applications.
+[/region]
+
+To set an environment variable you need to edit your **.bash_profile** file in your home directory. 
+The file is hidden in Finder, so you’ll need to use Terminal to open it for editing. 
+If you’re not familiar with Unix terminal commands please do the following.
+
+1. Open **Terminal app** from /Applications/Utilities.
+1. Type **touch .bash_profile** to create the file if it doesn’t exist yet.
+1. Type open **.bash_profile** to open the file in a text editor.
+1. Copy the following lines to the file (replacing paths to point to where you installed SDK and tools) and save it:
+
+	* export ANDROID_HOME="[REDACTED_USER_HOME]/Devel/NVPACK/android-sdk-macosx"
+	* export NDKROOT="[REDACTED_USER_HOME]/Devel/NVPACK/android-ndk-r10c"
+	* export ANT_HOME="[REDACTED_USER_HOME]/Devel/NVPACK/apache-ant-1.8.2"
+
+Here are the name of the variables that we depend on and where they should point to.
+
+* **ANDROID_HOME**
+	* Set this to the directory where you installed the Android SDK or the SDK directory if you installed ADT. 
+	* It must contain a directory called platform-tools.
+* **ANT_HOME** [Not needed if you installed ADT]
+	* Set this to where you unzipped Ant.
+	* It must contain a directory called bin with a file called ANT in it.
+* **NDKROOT** [Not needed if you do not compile C++ code]
+	* Set this to where you unzipped the NDK and look for something similar to **android-ndk-r9c**.
+	* It must contain a file called **ndk-build**.
+
+[/EXCERPT:ENVMacOSSetup]
+
+-->
+
 ## Android Tips and Tricks 
 
+[OBJECT:ToggleButtonContent]
+	[PARAMLITERAL:category]
+	OS
+	[/PARAMLITERAL]
+	[PARAMLITERAL:id]
+	windows
+	[/PARAMLITERAL]
+	[PARAMLITERAL:active]
+	active_button_content
+	[/PARAMLITERAL]
+	[PARAM:content]
+	[INCLUDE:#ATTWindowsSetup]
+	[/PARAM]
+[/OBJECT]
+
+[OBJECT:ToggleButtonContent]
+	[PARAMLITERAL:category]
+	OS
+	[/PARAMLITERAL]
+	[PARAMLITERAL:id]
+	mac
+	[/PARAMLITERAL]
+	[PARAMLITERAL:active]
+	active_button_content
+	[/PARAMLITERAL]
+	[PARAM:content]
+	[INCLUDE:#ATTMacOSSetup]
+	[/PARAM]
+[/OBJECT]
+
+<!--
+[EXCERPT:ATTWindowsSetup]
 If you do not see your device listed in the **Launch** menu, here are some things to try: 
 - Wait a few seconds and look again, as the background scanning for connected devices can take a short amount of time. 
 - Check your USB debugging connection is all working. 
@@ -157,6 +318,27 @@ in the progress bar, you may see that the last line indicates that it is calling
 - If the device is listed but shows unauthorized, that indicates that you need to get the device to ask if you want to authorize the PC. 
 - You may need to unplug and reconnect your device to make it show up. 
 - It is recommended you check the **Always allow** checkbox. 
+[/EXCERPT:ATTWindowsSetup]
+
+[EXCERPT:ATTMacOSSetup]
+If you do not see your device listed in the Launch menu, here are some things to try: 
+* Wait a few seconds and look again, as the background scanning for connected devices can take a short amount of time to register the device. 
+* Check your USB debugging connection is all working. 
+* Open the Terminal app and then try the following.
+	* Type the following command into the terminal **adb devices**.
+	* If you do not see any devices listed, your USB driver is not installed properly. How to install the driver depends on the model of your Android device. Search online for your device model and **USB driver**, then install the required driver. 
+	* Additionally, you will need to make sure your device has Developer mode enabled, and USB debugging is turned on. See the **[Android Developer Using Hardware Devices](http://developer.android.com/tools/device.html)** guide for more information (specifically items 2 and 3).
+	* Afterwards, make sure the **adb devices** command works; without this, the engine will not be able to do anything with your Android device.
+
+If your device shows up in the **Launch menu**, but gets stuck deploying to the device, it is most likely because the device is not authorized to interact with your Mac. 
+If you click on the **Show Log** line in the progress bar, and see that the last line indicates that it is calling adb uninstall try the following.
+* Open up the Terminal app and run the adb devices command. 
+* If the device is listed but shows unauthorized, that indicates that you need to get the device to ask if you want to authorize the Mac. 
+* It is recommended you check the **Always allow** checkbox to always allow your Mac to talk to this device.
+
+[/EXCERPT:ATTMacOSSetup]
+
+-->
 
 ## Android Texture Formats
 Not all Android devices are made the same. In particular, there are 4 different kinds of rendering hardware. They each support different formats of compressed textures. 
@@ -184,3 +366,4 @@ However, if you do not know, then you can choose **Android** or **Android (ETC1)
 enter a console command (like `stat fps`), but it will also display the formats that your device supports. Once you have that list, you can use that to choose a more ideal type when deploying/packaging.
 
 If you choose a format that your device does not support, it will fail to load. 
+
