#UE4doc #docs: Spelling changes in header. Confg -> Config

diff --git a/Engine/Documentation/Source/Platforms/Oculus/BestPractices/OculusRift.INT.udn b/Engine/Documentation/Source/Platforms/Oculus/BestPractices/OculusRift.INT.udn
--- a/Engine/Documentation/Source/Platforms/Oculus/BestPractices/OculusRift.INT.udn
+++ b/Engine/Documentation/Source/Platforms/Oculus/BestPractices/OculusRift.INT.udn
@@ -55,7 +55,7 @@ to view your work is to run a Standalone Game. To run a Standalone game please d
 * If the game launches in non-full screen, press **ALT+Enter** to full screen the game window or type **fullscreen** in the [Console](GettingStarted/RunningUnrealEngine#usefulin-gamecommands).
 * Bring up the [Console](GettingStarted/RunningUnrealEngine#usefulin-gamecommands) (`/backtick) and then type **stereo on**.
 
-## Oculus Specific Confg Settings
+## Oculus Specific Config Settings
 
 There are some specific settings for the Oculus that need to be placed in your **DefaultEngine.INI** file to permanently 
 set some variables for VR devices. 
@@ -94,7 +94,7 @@ would be. Inside of UE4 this can be done in the following manner.
 If you are using a DK1 headset, you will want to set up a virtual head model. This is not necessary if you are using the DK2 headset.
 [/REGION]
 
-The Virtual Head model is used to offset the players viewpoint. For traditional First Person Shooters the viewport rotates on 
+The Virtual Head model is used to offset the player's viewpoint. For traditional First Person Shooters the viewport rotates on 
 a central axis which is fine. However for a VR game this extra offset is unnecessary and should be zeroed out. 
 
 [REGION:tip]
