#UE4 Docs: Fixing Jira UEDOC-886 by removing the bit about TC_Greyscale and re-wording the sentence to be less obtuse.

diff --git a/Engine/Documentation/Source/Platforms/Mobile/Textures/MobileTextures.INT.udn b/Engine/Documentation/Source/Platforms/Mobile/Textures/MobileTextures.INT.udn
--- a/Engine/Documentation/Source/Platforms/Mobile/Textures/MobileTextures.INT.udn
+++ b/Engine/Documentation/Source/Platforms/Mobile/Textures/MobileTextures.INT.udn
@@ -17,7 +17,7 @@ when creating textures:
 
 In addition, texture settings need to be set up properly to work on mobile platforms:
 
-* Only TC_Default or TC_NormalMap (and TC_Greyscale?) values are used in the **Compression Settings**.
+* It is recommended that you only use  **TC_Default** or **TC_NormalMap** in the **Compression Settings** as any other setting could cause your Texture to take up more memory than expected on your target device.
 * The **sRGB** property should be enabled.
  
 ## Texture Usage in Materials
