#UE4doc #docs: spelling and grammar changes. Bolding reserved words for style guide.

diff --git a/Engine/Documentation/Source/Engine/Rendering/PostProcessEffects/PostProcessMaterials/PostProcessMaterials.INT.udn b/Engine/Documentation/Source/Engine/Rendering/PostProcessEffects/PostProcessMaterials/PostProcessMaterials.INT.udn
--- a/Engine/Documentation/Source/Engine/Rendering/PostProcessEffects/PostProcessMaterials/PostProcessMaterials.INT.udn
+++ b/Engine/Documentation/Source/Engine/Rendering/PostProcessEffects/PostProcessMaterials/PostProcessMaterials.INT.udn
@@ -155,7 +155,7 @@ This is the result:
 
 ![](sceneBeforeTonemapper.png)
 
-We get a stable and anti-aliased image. In motion we might see some artifacts with the temporal anti-aliasing. This feature is using the depth buffer to reproject old images. With the border rendered inside the object
+We get a stable and anti-aliased image. In motion we might see some artifacts with the temporal anti-aliasing. This feature is using the depth buffer to replace old images. With the border rendered inside the object
 it works ok but outside the object we would also need to adjust the depth buffer (not done yet, costs extra performance). Note that TemporalAA is still under development - here it appears the outline is thicker
 but ideally this should not be. 
 
@@ -164,16 +164,16 @@ but ideally this should not be.
 Post process materials you can look up into screen aligned buffers, but you need to know the right UVs. 
 The **ScreenPosition** material expression outputs the UV that you would expect (0,0 at the left top of the viewport and 1,1 at the bottom right).
 Using the **texture coordinate** material expression might give you different results. This is because the actual texture (more correctly it is a render target) is potentially larger than the viewport.
-It can be larger in editor because we share this texture for multiple viewports and the largets extent is use for all viewports.
+It can be larger in editor because we share this texture for multiple viewports and the largest extent is use for all viewports.
 Even in game it can be larger in some cases (e.g. SceneCaptureActors might have a smaller viewport, Matinee black borders, Splitscreen, VR, ...).
 The texture coordinate material expression gives you the UV for this larger texture. If you only need a relative offset (e.g. pixel sized edge detection) you need to scale with the right size.
 The **SceneTexture** material expression has outputs for the size and the inverse of the size (efficient and useful for pixel offsets).
-If you want to get the viewport UV (e.g. to map texture to the viewport) you can use the ScreenPosition material expression or you make the needed computations manually (more control, can be slower).
+If you want to get the viewport UV (e.g. to map texture to the viewport) you can use the **ScreenPosition** material expression or you make the needed computations manually (more control, can be slower).
 For this, you need the **ViewSize** material expression. To test all that, you can use the console variable **r.ViewPortTest** which allows you to test various viewport configurations.
 
 ## Filtered texture lookups
 
-The SceneTexture material expression has a checkbox to get [bilinear] filtered lookups.Using this option can results in a slower rendering so it should be only used if needed.
+The SceneTexture material expression has a checkbox to get [bilinear] filtered lookups. Using this option can results in a slower rendering so it should be only used if needed.
 Many of the screen space textures do not support filtering (e.g. GBuffer). Not exposing this property allows the engine to compress the data if needed (packing prevents filtering).
 
 ## Replacing the Tonemapper
@@ -201,14 +201,14 @@ In order to get the parameters you need to create a material instance from the p
 You still can use your own parameters and get them blended like other post process material settings.
 
 ## Known Issues
-The following issues still needs to be fixed:
+The following issues still need to be fixed:
 
   * **Material expression SceneTexture**
     * SeparateTranslucency does not work.
     * Certain lookups do not work in certain passes (e.g. some we might not fix because they would cost too much performance).
     * MaterialFunction might report an error but it still can be used in a Material with the Post Process domain.
   * **Material**
-    * UV in the PostProcesMaterial might not be in 0..1 range (e.g. in editor when reducing viewport size), this aligns with the lookup but makes it hard to implement something like an vignette effect.
+    * UV in the PostProcesMaterial might not be in 0.1 range (e.g. in editor when reducing viewport size), this aligns with the lookup but makes it hard to implement something like a vignette effect.
     * The asset thumbnail of a post process material does not look right.
     * Outputting alpha is not yet supported (should go through Opacity).
     * The preview material in the material editor does not look right.
@@ -250,7 +250,7 @@ The following issues still needs to be fixed:
 
   * **I need to use "Before Tonemapper" to get proper TemporalAA. When I use a color, it is tone mapped and therefore looks different. How can I prevent that?**
     
-	There is no easy solution. You would need to do an inverse tonemmaping operation (expensive). The color also might appear different
+	There is no easy solution. You would need to do an inverse tonemaping operation (expensive). The color also might appear different
 	depending on eye adaptation. You could expose the EyeAdaptation level to the SceneTextures to compensate for that.
 
   * **How can I get a full dump of the post processing graph?**
diff --git a/Engine/Documentation/Source/Support/Builds/ReleaseNotes/2013/Dec/ReleaseNotesDecember2013.INT.udn b/Engine/Documentation/Source/Support/Builds/ReleaseNotes/2013/Dec/ReleaseNotesDecember2013.INT.udn
--- a/Engine/Documentation/Source/Support/Builds/ReleaseNotes/2013/Dec/ReleaseNotesDecember2013.INT.udn
+++ b/Engine/Documentation/Source/Support/Builds/ReleaseNotes/2013/Dec/ReleaseNotesDecember2013.INT.udn
@@ -90,7 +90,7 @@ QA_APPROVED_UE4_BUILD_MINIMUM
 			![](FolderColors03.png)
 
 * **Package Auto-Save Restoration**
-	* You can now restore packages from the previous editing sessions auto-saves.
+	* You can now restore packages from the previous editing session's auto-saves.
 	* If the editor detects that it did not shut-down cleanly, it will present you will a dialog on start-up to let you choose which packages you'd like to restore.
 
 		![](PackageAutoSaveRestoration.png)
@@ -138,7 +138,7 @@ QA_APPROVED_UE4_BUILD_MINIMUM
 * **Animation Retargeting**
 	* Recently support has been added for Animation Retargeting! 
 	* What is animation retargeting?
-		* Animation retargeting is the ability to use animations made for a specific character, and play them on a character of different proportions. This is important: the characters have to share a common Skeleton for this to work (same bone names and hierarchy), only the bind pose can be different.
+		* Animation retargeting is the ability to use animations made for a specific character, and play them on a character of different proportions. This is important: the characters have to share a common Skeleton for this to work (same bone names and hierarchy); only the bind pose can be different.
 	* We have 3 different settings for bone translation retargeting which can be set on the Skeleton Tree window in Persona:
 		* **Skeleton**: Bone Translation comes from the Target Skeleton's bind pose.
 		* **Animation**: Bone Translation comes from the animation data, unchanged.
@@ -161,11 +161,11 @@ QA_APPROVED_UE4_BUILD_MINIMUM
 #### Editor and Tools
 * Added easier way to check the type of DragDrop operation with DragDropEvent.GetOperationAs<FMyDragDrop>();
 * Added the OnVisualizeTooltip() event that allows for scoped, custom tooltip handling.
-* Added an editor user settings ini option to prefix the project name watermark in the level editor for disambiguation when frequently switching between branches.
+* Added an editor user settings in option to prefix the project name watermark in the level editor for disambiguation when frequently switching between branches.
 * Added custom mesh previews for material thumbnails. Choose your preview mesh in the details panel of the material editor then in the Content Browser, enter thumbnail edit mode and change the material's preview primitive to the teapot.
 * Rocket users now have the ability to send quick feedback from the editor.
 * Added option to preserve non-uniform scaling when snapping.
-* Cooking paths are now in per-game Editor preferences.
+* Cooking paths are now in per-game Editor Preferences.
 * Added folders to Content Browser.
 * Added thin window border to Slate windows.
 * Added interactive actor 'eyedropper' picker for picking actors in the scene.
