#UE4doc #docs: fixed a spelling error: projects's -> project's

diff --git a/Engine/Documentation/Source/Engine/Landscape/UI/InterfaceReference.INT.udn b/Engine/Documentation/Source/Engine/Landscape/UI/InterfaceReference.INT.udn
--- a/Engine/Documentation/Source/Engine/Landscape/UI/InterfaceReference.INT.udn
+++ b/Engine/Documentation/Source/Engine/Landscape/UI/InterfaceReference.INT.udn
@@ -88,7 +88,7 @@ This document contains references to all of the icons for the **Landscape** syst
 [/VAR]
 
 [VAR:RampTool]
-![](Engine/Landscape/icon_Landscape_Tool_Ramp_40x.png)(convert:false) | **Ramp Tool** | This tool will create a ramp between two specified points, adding falloffs according to the settings. Note that this tool can not use any brushes. |
+![](Engine/Landscape/icon_Landscape_Tool_Ramp_40x.png)(convert:false) | **Ramp Tool** | This tool will create a ramp between two specified points, adding falloffs according to the settings. Note that this tool cannot use any brushes. |
 [/VAR]
 
 [VAR:FlattenPTool]
@@ -173,7 +173,7 @@ This document contains references to all of the icons for the **Landscape** syst
 [EXCERPT:LSREF002]
 [REGION:raw]
 | ----- | ----- | ---- |
-| ![](Engine/Landscape\icon_Landscape_Mode_Sculpt_40x.png)(convert:false) | [**Sculpt**](Engine/Landscape/Editing/SculptMode) | This mode allows you to sculpt the look of your Landscape heightmap so that it fits your projects’s needs. 
+| ![](Engine/Landscape\icon_Landscape_Mode_Sculpt_40x.png)(convert:false) | [**Sculpt**](Engine/Landscape/Editing/SculptMode) | This mode allows you to sculpt the look of your Landscape heightmap so that it fits your project’s needs. 
 | ^ |  |  |
 %Ind% %SculptTool%
 %Ind% %SmoothTool%
