#UE4doc #docs: sp: anti-aliasing, largest grammer: need

diff --git a/Engine/Documentation/Source/Engine/Rendering/PostProcessEffects/PostProcessMaterials/PostProcessMaterials.INT.udn b/Engine/Documentation/Source/Engine/Rendering/PostProcessEffects/PostProcessMaterials/PostProcessMaterials.INT.udn
--- a/Engine/Documentation/Source/Engine/Rendering/PostProcessEffects/PostProcessMaterials/PostProcessMaterials.INT.udn
+++ b/Engine/Documentation/Source/Engine/Rendering/PostProcessEffects/PostProcessMaterials/PostProcessMaterials.INT.udn
@@ -29,7 +29,7 @@ Post Process Materials should be used sparingly and only when you really need th
 
 ## Using a Post Process Material
 
-Through post process settings (usually defined with post process volumes or camera settings), it is possible to blend (so called) blendable assets.
+Through post process settings (usually defined with post process volumes or camera settings) it is possible to blend (so called) blendable assets.
 At the moment, only **materials** and **material instances** are blendable assets. The engine provides a few post process materials but with this feature
 you can create your own **custom post processing** without any programmer help.
 
@@ -140,7 +140,7 @@ This is the material we used to visualize it:
 
 Temporal anti-aliasing is a unique UE4 feature that greatly improves image quality with a very moderate performance cost.
 
-By default, a post process material is inserted in the end of the post processing graph (after tone mapper). This means you get the final LDR color after tonemapping, color grading and after the temporal anti-alasing
+By default, a post process material is inserted in the end of the post processing graph (after tone mapper). This means you get the final LDR color after tonemapping, color grading and after the temporal anti-aliasing
 was applied. This is the best spot for many simple post processing effects - for performance and ease of use.
 
 Here you can see how we used the custom depth input to visualize a silhouette around specific objects:
@@ -208,14 +208,14 @@ The following issues still need to be fixed:
     * Certain lookups do not work in certain passes (e.g. some we might not fix because they would cost too much performance).
     * MaterialFunction might report an error but it still can be used in a Material with the Post Process domain.
   * **Material**
-    * UV in the PostProcesMaterial might not be in 0.1 range (e.g. in editor when reducing viewport size), this aligns with the lookup but makes it hard to implement something like a vignette effect.
+    * UV in the PostProcesMaterial might not be in 0-1 range (e.g. in editor when reducing viewport size), this aligns with the lookup but makes it hard to implement something like an vignette effect.
     * The asset thumbnail of a post process material does not look right.
     * Outputting alpha is not yet supported (should go through Opacity).
     * The preview material in the material editor does not look right.
     * There has been a case where a material change was not resulting in a change in the post process. Workaround: restart editor.
 	* The **Content Browser** should allow easy filtering for post process materials.
   * **Blending**
-    * When blending two post process volumes with a blend radius, it is possible to see a non soft transition. This can be prevented by using an unbound volume with a material instance setting that represents the default.
+    * When blending two post process volumes with a blend radius, it is possible to see a non-soft transition. This can be prevented by using an unbound volume with a material instance setting that represents the default.
 
 
 ## FAQ
