#UE4doc #docs:  minor spelling and grammar changes.

diff --git a/Engine/Documentation/Source/Platforms/Android/GettingStarted/3/AndroidQS_3.INT.udn b/Engine/Documentation/Source/Platforms/Android/GettingStarted/3/AndroidQS_3.INT.udn
--- a/Engine/Documentation/Source/Platforms/Android/GettingStarted/3/AndroidQS_3.INT.udn
+++ b/Engine/Documentation/Source/Platforms/Android/GettingStarted/3/AndroidQS_3.INT.udn
@@ -25,7 +25,7 @@ Description:Start a template project
 %Steps%
 
 
-1.	Create a new project. Here, we are testing a Blueprints-based project, but if you start with a C++ project and then compile it so it can be opened in the editor, the next steps will be the same.  
+1.	Create a new project. Here we are testing a Blueprint-based project. If you start with a C++ project and then compile it, the following steps will be the same.  
 
 ![](new_android_project.png)
 
