#UE4doc #docs Updated Description field.

diff --git a/Engine/Documentation/Source/Gameplay/HowTo/ActorInput/Code/ActorInput_Code.INT.udn b/Engine/Documentation/Source/Gameplay/HowTo/ActorInput/Code/ActorInput_Code.INT.udn
--- a/Engine/Documentation/Source/Gameplay/HowTo/ActorInput/Code/ActorInput_Code.INT.udn
+++ b/Engine/Documentation/Source/Gameplay/HowTo/ActorInput/Code/ActorInput_Code.INT.udn
@@ -1,7 +1,7 @@
 Availability: Docs
 Title:How to Setup Input on an Actor in Code
 Crumbs: %ROOT%, Gameplay, Gameplay\HowTo\ActorInput
-Description:A How To Guide for settings up input on an Actor in Code.
+Description:A How To Guide for setting up an Actor in code.
 
 [TOC (start:2 end:2)]
 
