#UE4 Docs: Fixing Jira 1283 by removing the windows env setup section as it was out of date.

diff --git a/Engine/Documentation/Source/Platforms/iOS/GettingStarted/iOSGettingStarted.INT.udn b/Engine/Documentation/Source/Platforms/iOS/GettingStarted/iOSGettingStarted.INT.udn
--- a/Engine/Documentation/Source/Platforms/iOS/GettingStarted/iOSGettingStarted.INT.udn
+++ b/Engine/Documentation/Source/Platforms/iOS/GettingStarted/iOSGettingStarted.INT.udn
@@ -217,15 +217,6 @@ Read this section if you are going to do your main development and coding on Win
         * for (( ;; )) ; do ./UnrealRemoteTool; done
     * Note that in Mountain Lion, Terminal will remember the working directory and command history, so if you reboot, an up-arrow should recall the command.
 
-## Windows system setup
-
-* Set up environment variables (or modify the code in RemoteToolchain.cs so you can share a single Mac for everyone).
-    * ue.RemoteCompileServerName = The DNS name or IP address of the Mac used for compiling (e.g. a1012).
-    * ue.IOSSigningServer = This is optional (when compiling, the package will force the same as ue.RemoteCompileServerName, but if you run IPhonePackager.exe on its own, this will determine the Mac to use if -mac is not specified).
-
-The reason for two sets of environment variables is to enable completely separate compiling and signing servers. You may need to restart your computer for these changes to take effect.
-
-
 ### Compiling on Windows
 Once you have completed the above steps for setting up your Windows and Mac systems, compiling should be straightforward, as it uses UnrealBuildTool to perform compilation:
 * Run //depot/UE4/GenerateProjectFiles_AllPlatforms.bat
