#UE4docs #docs Minor grammatical update.

diff --git a/Engine/Documentation/Source/Engine/Rendering/Materials/Functions/Reference/Math/Math.INT.udn b/Engine/Documentation/Source/Engine/Rendering/Materials/Functions/Reference/Math/Math.INT.udn
--- a/Engine/Documentation/Source/Engine/Rendering/Materials/Functions/Reference/Math/Math.INT.udn
+++ b/Engine/Documentation/Source/Engine/Rendering/Materials/Functions/Reference/Math/Math.INT.udn
@@ -11,7 +11,7 @@ Math functions perform basic math equations on the values of pixels within a tex
 
 ## Math Functions
 
-The following is a list of all the functions underneath the Math category.
+The following is a list of all the functions found in the Math category.
 
 ### Add Components
 
@@ -67,7 +67,7 @@ The LinearSine function takes in a scalar value and outputs the linear sine (or
 
 ### VectorToRadialValue
 
-This function transforms the vector of a Vector2 into an angle, or transforms UV coordinate data into radial coordinates. In the case of a vector, the angle will be output in one channel and the length of the vector in the other.
+This function transforms the vector of a Vector2 into an angle, or transforms UV coordinate data into radial coordinates. In the case of a vector, the angle will output in one channel and the length of the vector in the other.
 
 | Item | Description |
 | -------- | -------- |
@@ -77,7 +77,7 @@ This function transforms the vector of a Vector2 into an angle, or transforms UV
 |[REGION:tablesection]Outputs[/REGION] ||
 | **Radial Coordinates** | Returns the radial coordinates of the input. In the case of a vector, the angle is on one channel and the distance is on the other.|
 | **Vector Converted to Angle** | Returns the angle of the input vector or a radial gradient in the case of UVs.|
-| **Linear Distance** | Returns the linear length of the input vector, or in the case of UVs outputs a radial gradient of distances.|
+| **Linear Distance** | Returns the linear length of the input vector, or in the case of UVs, outputs a radial gradient of distances.|
 
 
 ![Vector To Radial Value](VectorToRadialValue_Demo.png)
\ No newline at end of file
