#UE4doc #docs Minor grammatical update.

diff --git a/Engine/Documentation/Source/Engine/Rendering/Materials/ExpressionReference/Functions/FunctionsExpressions.INT.udn b/Engine/Documentation/Source/Engine/Rendering/Materials/ExpressionReference/Functions/FunctionsExpressions.INT.udn
--- a/Engine/Documentation/Source/Engine/Rendering/Materials/ExpressionReference/Functions/FunctionsExpressions.INT.udn
+++ b/Engine/Documentation/Source/Engine/Rendering/Materials/ExpressionReference/Functions/FunctionsExpressions.INT.udn
@@ -15,10 +15,10 @@ The FunctionInput expression can only be placed in a material function, where it
 | -------- | ----------- |
 | [REGION:tablesection]Properties[/REGION] ||
 | **Input Name** | The input's name, which will be displayed on MaterialFunctionCall expressions that use the material function containing the input. |
-| **Description** | A description of the input, which will be displayed as a tooltip when the connector for this input on a MaterialFunction Call expression is hovered over with the mouse. |
+| **Description** | A description of the input, which will be displayed as a 'tooltip' when the connector for this input on a MaterialFunction Call expression is hovered over with the mouse. |
 | **Input Type** | The type of data this input expects. Data passed to this input will be cast to this type, throwing a compiler error if the cast fails because the data is not compatible. |
 | **Preview Value** | The value to use as a preview for this input when editing the material function containing it. |
-| **Use Preview Value As Default** | If enabled, the **Preview Value** will be used as the default value for this input if no data is passed in. |
+| **Use Preview Value As Default** | If enabled, **Preview Value** will be used as the default value for this input if no data is passed in. |
 | **Sort Priority** | Specifies the priority for this input to use when determining the order of the inputs to be displayed on a MaterialFunctionCall expression. |
 
 
