#UE4doc #docs Minor grammatical update.

diff --git a/Engine/Documentation/Source/Engine/Rendering/Materials/Functions/Reference/VectorOps/VectorOps.INT.udn b/Engine/Documentation/Source/Engine/Rendering/Materials/Functions/Reference/VectorOps/VectorOps.INT.udn
--- a/Engine/Documentation/Source/Engine/Rendering/Materials/Functions/Reference/VectorOps/VectorOps.INT.udn
+++ b/Engine/Documentation/Source/Engine/Rendering/Materials/Functions/Reference/VectorOps/VectorOps.INT.udn
@@ -11,7 +11,7 @@ The VectorOps category contains special functions for applying various vector-ba
 
 ## VectorOps Functions
 
-The following is a list of all the functions underneath the VectorOps category.
+The following is a list of all the functions found in the VectorOps category.
 
 ### Fresnel
 
@@ -22,8 +22,8 @@ Unlike the regular Fresnel expression node, the Fresnel function allows you to d
 |[REGION:tablesection]Inputs[/REGION] ||
 | **Normal Vector (Vector3)** | The first vector used in the Fresnel operation. Typically, this is the surface normal.|
 | **Camera Vector (Vector3)** | The vector of the camera's direction.|
-| **Invert Fresnel (StaticBool)** | This inverts the operation, reversing the way the normals are calculated to give the results.|
-| **Power (Scalar)** | This value controls how quickly the color falls off between the core and the edges.|
+| **Invert Fresnel (StaticBool)** | This inverts the operation, reversing the way normals are calculated to yield results.|
+| **Power (Scalar)** | This value controls how quickly color falls off between the core and the edges.|
 | **Use Cheap Contrast (StaticBool)** | This activates an internal CheapContrast function to boost the contrast of the Fresnel's result.|
 | **Cheap contrast dark (Scalar)** | When using cheap contrast, this controls how much the dark value shows in the result. Has no effect when not using cheap contrast.|
 | **Cheap contrast bright (Scalar)** | When using cheap contrast, this controls how much the bright value shows in the result. Has no effect when not using cheap contrast.|
