#UE4 Docs: Fixing Jira 1455 by adding a note about how 4.8 supports ES31 and AEP rendering on Tegra K1 devices.

diff --git a/Engine/Documentation/Source/Platforms/Android/DeviceCompatibility/AndroidDeviceCompatibility.INT.udn b/Engine/Documentation/Source/Platforms/Android/DeviceCompatibility/AndroidDeviceCompatibility.INT.udn
--- a/Engine/Documentation/Source/Platforms/Android/DeviceCompatibility/AndroidDeviceCompatibility.INT.udn
+++ b/Engine/Documentation/Source/Platforms/Android/DeviceCompatibility/AndroidDeviceCompatibility.INT.udn
@@ -56,4 +56,8 @@ The following table lists individual devices we have tested here:
 **1:** The Nexus 5 with the latest publicly available driver performs poorly when using features from the Full HDR tier. We have been working closely with Qualcomm in this area and they have developed faster drivers that remove the bottlenecks we were running in to. Our HDR features are fully supported on their latest internal drivers which we hope will be available to the public soon!
 (#KindleFireHDX)
 **2:** Similar to the Nexus 5, the Kindle Fire HDX runs in to some bottlenecks in the Adreno 330 driver and we expect it to perform well in the future with an updated driver.
+
+**3:** Unreal Engine version 4.8 and later will now support the use of **ES31** and **AEP** rendering on Tegra K1 devices. 
+Please note that due to the limited number of devices that currently support ES31 and AEP rendering very little performance and compatibility testing has been completed at this time.
+If you are developing a project that is aimed at using one of theses devices it is highly recommended that you perform extensive performance and compatibility testing to ensure that the device delivers your project in the manner you intended. 
 [/REGION]
\ No newline at end of file
diff --git a/Engine/Documentation/Source/Platforms/Android/Reference/AndroidGettingStarted.INT.udn b/Engine/Documentation/Source/Platforms/Android/Reference/AndroidGettingStarted.INT.udn
--- a/Engine/Documentation/Source/Platforms/Android/Reference/AndroidGettingStarted.INT.udn
+++ b/Engine/Documentation/Source/Platforms/Android/Reference/AndroidGettingStarted.INT.udn
@@ -359,6 +359,12 @@ compressed formats, you have a choice of what format you want to be cooked. In t
 The first one (**Android**) actually takes all formats and puts them into the cooked packages. This will take longer, and will make larger packages, but will guarantee the minimal memory usage by choosing 
 the best format at runtime. It will work on any device (as will **Android (ETC1)**).
 
+[region:note]
+Unreal Engine version 4.8 and later will now support the use of **ES31** and **AEP** rendering on Tegra K1 devices. 
+Please note that due to the limited number of devices that currently support ES31 and AEP rendering very little performance and compatibility testing has been completed at this time.
+If you are developing a project that is aimed at using one of theses devices it is highly recommended that you perform extensive performance and compatibility testing to ensure that the device delivers your project in the manner you intended. 
+[/region]
+
 ### Choosing a Format
 You most likely will want to choose an optimal format for your device. If you know the type of GPU in your device, you can look at the table above to choose a format.
 
