#UE4doc #docs: Minor grammar changes: reason -> reasons, its -> it is

diff --git a/Engine/Documentation/Source/Engine/Landscape/QuickStart/3/LSQS_3.int.udn b/Engine/Documentation/Source/Engine/Landscape/QuickStart/3/LSQS_3.int.udn
--- a/Engine/Documentation/Source/Engine/Landscape/QuickStart/3/LSQS_3.int.udn
+++ b/Engine/Documentation/Source/Engine/Landscape/QuickStart/3/LSQS_3.int.udn
@@ -37,7 +37,7 @@ Sculpting the Landscape is a rather easy but time consuming process. All of the
 For this part of the Landscape tutorial, we are going to start with a completely flat section of the Landscape and then build up the details as we go along. The goal here is not to exactly mimic what was created in the tutorial but to get you familiar and comfortable with using the various Landscape tools. 
 
 [REGION:note]
-There could be a lot of various reason as to why what you do in this tutorial does not come out exactly the same as what you see in the following screen shots. Working with the Landscape tools requires a lot of trial and error so your results will vary, sometimes greatly, from what you are seeing in the following set of images. The most important thing to get out of this tutorial is to understand how each of the Landscape tools work and how all the tools work together to give you the final product.
+There could be a lot of various reasons as to why what you do in this tutorial does not come out exactly the same as what you see in the following screen shots. Working with the Landscape tools requires a lot of trial and error so your results will vary, sometimes greatly, from what you are seeing in the following set of images. The most important thing to get out of this tutorial is to understand how each of the Landscape tools work and how all the tools work together to give you the final product.
 [/REGION]
 ---
 
@@ -71,7 +71,7 @@ You can find the brush size and strength settings that were used for this step l
 	![](T_Landscape_Smooth.png)
 	[/REGION]
 
-1. With the Landscape now smoothed out, its time to add some flat mesa like sections using the **Flatten Tool**. 
+1. With the Landscape now smoothed out, it is time to add some flat mesa like sections using the **Flatten Tool**. 
 You can find the brush size and strength settings that were used for this step listed below and when completed, you should have something that looks like the following.
 	|Tool Used| Brush Size | Strength Setting |
 	|---------| --- | -------------- |
