#UE4doc #docs: add numbered list to four items: Landscape, Materials, Resources, Textures

diff --git a/Engine/Documentation/Source/Engine/Landscape/QuickStart/4/LSQS_4.INT.udn b/Engine/Documentation/Source/Engine/Landscape/QuickStart/4/LSQS_4.INT.udn
--- a/Engine/Documentation/Source/Engine/Landscape/QuickStart/4/LSQS_4.INT.udn
+++ b/Engine/Documentation/Source/Engine/Landscape/QuickStart/4/LSQS_4.INT.udn
@@ -31,13 +31,10 @@ For this section of the tutorial, we are going to be making the following four f
 If you would like to know more about how to setup folders in Unreal Engine 4, please refer to the following page about [**Folders**](Engine/Content/Browser/UserGuide/WorkingWithSources).
 [/REGION]
 
-**Landscape** 
-
-**Materials**
-
-**Resources**
-
-**Textures**
+1. **Landscape** 
+2. **Materials**
+3. **Resources**
+4. **Textures**
 
 When completed, you should have something that looks like the following.
 
