#UE4doc #docs:  made minor changes with spacing and "language & tone" changes of the style guide.

diff --git a/Engine/Documentation/Source/Engine/Landscape/QuickStart/6/LSQS_6.INT.udn b/Engine/Documentation/Source/Engine/Landscape/QuickStart/6/LSQS_6.INT.udn
--- a/Engine/Documentation/Source/Engine/Landscape/QuickStart/6/LSQS_6.INT.udn
+++ b/Engine/Documentation/Source/Engine/Landscape/QuickStart/6/LSQS_6.INT.udn
@@ -36,7 +36,7 @@ While the previous quick start tutorial for Landscape will get you up and runnin
 	![](T_Landscape_Tips_Tricks_00.png)
 
 * Landscape compiles shaders separately for each component based on which layers are painted on them. So, for example, if you have a component with a dirt layer on it but no trace of the grass layer has been painted on it, the textures for the grass layer are left out of the material for that component, making it cheaper to render. So when do an optimization pass, it can be worthwhile to go over a Landscape and look for components that only have a tiny trace of a given layer and erase them to reduce material complexity.
-* Another issue to watch out for when painting layers is to avoid having too many textures on one component. The material editor stats show the limit of how many texture samples you are allowed to use, but for Landscape materials the masks for each layer count as texture samples too and do not show in the stats. If a component starts showing the default texture(Grey Squares) when you paint a new layer onto it, it is likely that it is gone over the texture sample limit and either needs to have a layer erased or the material needs to be optimized to use less textures.
+* Another issue to watch out for when painting layers is to avoid having too many textures on one component. The material editor stats show the limit of how many texture samples you are allowed to use, but for Landscape materials the masks for each layer count as texture samples too and do not show in the stats. If a component starts showing the default texture (Grey Squares) when you paint a new layer onto it, it is likely that it is gone over the texture sample limit and either needs to have a layer erased or the material needs to be optimized to use less textures.
 * You can change the LOD Distance Factor for individual Landscape components so they will simplify at closer or further distance thresholds. Things like mountain peaks or anything with a distinct silhouette will LOD the most noticeably as you move further away, so you can reduce the LOD bias for those components to preserve their shape. You can also raise the LOD bias for low-detail areas like flat plains that will not look noticeably different with less tessellation.
 
 ##World Composition 
@@ -55,7 +55,7 @@ World Machine is a powerful procedural terrain creation software. It can be used
 
 **Terragen**
 
-Terragen is another very powerful fully procedural terrain creation software. Much like World Machine, it can be used to build, texture, and export both heightmaps and textures for your Landscape. You can read more about how Terragen might be able to help you out with your Landscape by checking out their website.  
+Terragen is another powerful fully procedural terrain creation software. Much like World Machine, it can be used to build, texture, and export both heightmaps and textures for your Landscape. You can read more about how Terragen might be able to help you out with your Landscape by checking out their website.  
 
 [**Terragen**](http://planetside.co.uk/)
 
