Grammatical and spelling revisions.

diff --git a/Engine/Documentation/Source/Engine/Rendering/Materials/ExpressionReference/Constant/ConstantExpressions.INT.udn b/Engine/Documentation/Source/Engine/Rendering/Materials/ExpressionReference/Constant/ConstantExpressions.INT.udn
--- a/Engine/Documentation/Source/Engine/Rendering/Materials/ExpressionReference/Constant/ConstantExpressions.INT.udn
+++ b/Engine/Documentation/Source/Engine/Rendering/Materials/ExpressionReference/Constant/ConstantExpressions.INT.udn
@@ -30,13 +30,13 @@ You can quickly create a Constant node by holding the **1** key and clicking in
 [EXCERPT:Constant2vector]
 ## Constant2Vector
 
-The Constant2Vector expression outputs a two-channel vector value, in other words, two constants numbers.
+The Constant2Vector expression outputs a two-channel vector value, in other words, two constant numbers.
 
 | Item | Description |
 | -------- | ----------- |
 | [REGION:tablesection]Properties[/REGION] ||
-| **R** | Specifies the float value of the red, or first, channel of the vector the expression outputs. |
-| **G** | Specifies the float value of the green, or second, channel of the vector the expression outputs. |
+| **R** | Specifies the float value of the red (first) channel of the vector the expression outputs. |
+| **G** | Specifies the float value of the green (second) channel of the vector the expression outputs. |
 
 **Examples:** (0.4, 0.6), (1.05, -0.3)
 
@@ -60,9 +60,9 @@ The Constant3Vector expression outputs a three-channel vector value, in other wo
 | Item | Description |
 | -------- | ----------- |
 | [REGION:tablesection]Properties[/REGION] ||
-| **R** | Specifies the float value of the red, or first, channel of the vector the expression outputs. |
-| **G** | Specifies the float value of the green, or second, channel of the vector the expression outputs. |
-| **B** | Specifies the float value of the blue, or third, channel of the vector the expression outputs. |
+| **R** | Specifies the float value of the red (first) channel of the vector the expression outputs. |
+| **G** | Specifies the float value of the green (second) channel of the vector the expression outputs. |
+| **B** | Specifies the float value of the blue (third) channel of the vector the expression outputs. |
 
 
 **Examples:** (0.4, 0.6, 0.0), (1.05, -0.3, 0.3)
@@ -85,10 +85,10 @@ The Constant4Vector expression outputs a four-channel vector value, in other wor
 | Item | Description |
 | -------- | ----------- |
 | [REGION:tablesection]Properties[/REGION] ||
-| **R** | Specifies the float value of the red, or first, channel of the vector the expression outputs. |
-| **G** | Specifies the float value of the green, or second, channel of the vector the expression outputs. |
-| **B** | Specifies the float value of the blue, or third, channel of the vector the expression outputs. |
-| **A** | Specifies the float value of the alpha, or fourth, channel of the vector the expression outputs. |
+| **R** | Specifies the float value of the red (first) channel of the vector the expression outputs. |
+| **G** | Specifies the float value of the green (second) channel of the vector the expression outputs. |
+| **B** | Specifies the float value of the blue (third) channel of the vector the expression outputs. |
+| **A** | Specifies the float value of the alpha (fourth) channel of the vector the expression outputs. |
 
 **Examples:** (0.4, 0.6, 0.0, 1.0), (1.05, -0.3, 0.3, 0.5)
 
diff --git a/Engine/Documentation/Source/Engine/Rendering/Materials/ExpressionReference/Particles/ParticleExpressions.INT.udn b/Engine/Documentation/Source/Engine/Rendering/Materials/ExpressionReference/Particles/ParticleExpressions.INT.udn
--- a/Engine/Documentation/Source/Engine/Rendering/Materials/ExpressionReference/Particles/ParticleExpressions.INT.udn
+++ b/Engine/Documentation/Source/Engine/Rendering/Materials/ExpressionReference/Particles/ParticleExpressions.INT.udn
@@ -38,7 +38,7 @@ In this example, you can see where the ParticleColor expression is providing col
 [EXCERPT:ParticleDirection]
 ## ParticleDirection
 
-The ParticleDirection expression outputs 3vector (RGB) data on a per-particle basis, representing the direction a given particle is currently traveling. 
+The ParticleDirection expression outputs Vector3 (RGB) data on a per-particle basis, representing the direction a given particle is currently traveling. 
 
 
 ![ParticleDirectionExample.png](ParticleDirectionExample.png)(w:900)
@@ -76,7 +76,7 @@ The ParticleMotionBlurFade expression outputs a value representing the amount of
 [EXCERPT:ParticlePositionWS]
 ## ParticlePositionWS
 
-The ParticlePositionWS expression outputs 3vector (RGB) data representing each individual particle's position in world space.
+The ParticlePositionWS expression outputs Vector3 (RGB) data representing each individual particle's position in world space.
 
 
 ![ParticlePositionWS.png](ParticlePositionWSExample.png)(w:900)
diff --git a/Engine/Documentation/Source/Engine/Rendering/Materials/ExpressionReference/Vector/VectorExpressions.INT.udn b/Engine/Documentation/Source/Engine/Rendering/Materials/ExpressionReference/Vector/VectorExpressions.INT.udn
--- a/Engine/Documentation/Source/Engine/Rendering/Materials/ExpressionReference/Vector/VectorExpressions.INT.udn
+++ b/Engine/Documentation/Source/Engine/Rendering/Materials/ExpressionReference/Vector/VectorExpressions.INT.udn
@@ -10,11 +10,11 @@ Description:Expressions that output a vector value, such as a position or normal
 [EXCERPT:ActorPositionWS]
 ## ActorPositionWS
 
-ActorPositionWS outputs 3vector (RGB) data representing the location of the object with this material on it in world space. 
+ActorPositionWS outputs Vector3 (RGB) data representing the location of the object with this material on it in world-space. 
 
 ![](ActorPositionColors.png)(w:434) ![](ActorPositionNetwork.png)(w:490)
 
-In this example, you can see that ActorPositionWS is being fed directly into the Base Color of the material. As a result, each of the spheres with the material applied to them show a different color as they are moved to different locations in 3D space. Note that the result of the ActorPositionWS node is being divided by 1600 to create a nice blend in color, rather than a pop.
+In this example, you can see that ActorPositionWS is being fed directly into the Base Color of the material. As a result, each of the spheres with the material applied to them show a different color as they are moved to different locations in 3D space. Note that the result of the ActorPositionWS node is being divided by 1600 to create a nice blend-in color, rather than a pop.
 
 [/EXCERPT:ActorPositionWS]
 
@@ -52,9 +52,9 @@ The CameraVector expression outputs a three-channel vector value representing th
 
 
 ## LightVector
-
+[REGION:warning]
 **This expression is deprecated in Unreal Engine 4 as lighting calculations are now deferred.**
-
+[/REGION:warning]
 <!--
 The LightVector expression is a three-channel vector value representing the direction of the light with respect to the surface.
 
@@ -81,7 +81,7 @@ In the example above, you can see that the colors of the objects correspond to t
 ## ObjectOrientation
 
 
-The ObjectOrientation expression outputs the world-space up vector of the object. In other words, this is the direction the local positive Z-axis of the object the material is applied to is pointing.
+The ObjectOrientation expression outputs the world-space up vector of the object. In other words, the object's local positive z-axis is pointing in this direction.
 
 ![ObjectOrientationExample.png](ObjectOrientationExample.png)
 [/EXCERPT:ObjectOrientation]
@@ -91,7 +91,7 @@ The ObjectOrientation expression outputs the world-space up vector of the object
 ## ObjectPositionWS
 
 
-The ObjectPositionWS expression outputs the world-space center position of the object's bounds.  This is useful for creating spherical lighting for foliage, for example.
+The ObjectPositionWS expression outputs the world-space center position of the object's bounds.  For example, this is useful for creating spherical lighting for foliage.
 
 
 ![](WorldPositionWSScene.png)(w:460) ![](WorldPositionNetwork.png)(w:470)
@@ -128,7 +128,7 @@ The ReflectionVectorWS expression is similar in spirit to [CameraVectorWS](#Came
 ## VertexNormalWS
 
 
-The VertexNormalWS expression outputs the world space vertex normal. It can only be used in material inputs that are executed in the vertex shader, like WorldPositionOffset.  This is useful for making a mesh grow or shrink.  Note that offsetting position along the normal will cause the geometry to split apart along UV seams.
+The VertexNormalWS expression outputs the world-space vertex normal. It can only be used in material inputs that are executed in the vertex shader, like WorldPositionOffset.  This is useful for making a mesh grow or shrink.  Note that offsetting position along the normal will cause the geometry to split apart along UV seams.
 
 [REGION:fullwidth]
 ![VertexNormalWSExample.png](VertexNormalWSExample.png)
