#UE4 Docs: added Topic/TopicCompact vars back

diff --git a/Engine/Documentation/Source/Platforms/Android/GettingStarted/AndroidQuickStart.INT.udn b/Engine/Documentation/Source/Platforms/Android/GettingStarted/AndroidQuickStart.INT.udn
--- a/Engine/Documentation/Source/Platforms/Android/GettingStarted/AndroidQuickStart.INT.udn
+++ b/Engine/Documentation/Source/Platforms/Android/GettingStarted/AndroidQuickStart.INT.udn
@@ -3,7 +3,39 @@ Title: Android Quick Start
 Crumbs:%ROOT%, Platforms, Platforms/Android
 Description:Setting up for development for the Android platform.
 
+[VAR:Topic]
+[OBJECT:Topic]
+	[PARAM:image]
+		![](AndroidTopic.jpg)
+	[/PARAM]
+	[PARAM:title]
+		%Platforms/Android/GettingStarted:title%
+	[/PARAM]
+	[PARAM:path]
+		[RELATIVE:Platforms/Android/GettingStarted]
+	[/PARAM]
+	[PARAM:description]
+	%Platforms/Android/GettingStarted:description%
+	[/PARAM]
+[/OBJECT]
+[/VAR]
 
+[VAR:TopicCompact]
+[OBJECT:TopicCompact]
+	[PARAM:image]
+		![](AndroidTopic.jpg)
+	[/PARAM]
+	[PARAM:title]
+		%Platforms/Android/GettingStarted:title%
+	[/PARAM]
+	[PARAM:path]
+		[RELATIVE:Platforms/Android/GettingStarted]
+	[/PARAM]
+	[PARAM:description]
+	%Platforms/Android/GettingStarted:description%
+	[/PARAM]
+[/OBJECT]
+[/VAR]
 
 The Android Quick Start will walk you through setting up your computer and your Android device for deployment, and putting a template game on your device using both the launching and packaging features.
 This guide uses TADP as the Android SDK, but if you would like to know how to set up another SDK for use with Unreal Engine, see the [](Platforms\Android\Reference) page.
