#UE4doc #docs: fixed typo: Priorites -> Priorities

diff --git a/Engine/Documentation/Source/Resources/ContentExamples/PostProcessing/PostProcessing.INT.udn b/Engine/Documentation/Source/Resources/ContentExamples/PostProcessing/PostProcessing.INT.udn
--- a/Engine/Documentation/Source/Resources/ContentExamples/PostProcessing/PostProcessing.INT.udn
+++ b/Engine/Documentation/Source/Resources/ContentExamples/PostProcessing/PostProcessing.INT.udn
@@ -79,7 +79,7 @@ Listed below are the examples provided inside the **PostProcessing** map:
 | **[1.14 AA Method](Resources\ContentExamples\PostProcessing\1_14)** | Demonstrates **Anti-Aliasing** (or AA) which improves image quality by smoothing out the edges that appear jagged on screen. |
 | **[1.15 Blur GBuffer intensity](Resources\ContentExamples\PostProcessing\1_15)** | Although hard to see, this effect helps with anti-aliasing by increasing the **Material Roughness** on the hard edges of surfaces, smoothing out the reflections on those hard edges. |
 | **[1.16 Blendables](Resources\ContentExamples\PostProcessing\1_16)** | **Blendables** are, in short, overlays that can be applied to the screen (e.g. screen warping or coloring). |
-| **[1.17 Priorites](Resources\ContentExamples\PostProcessing\1_17)** | Allows you to control which **Post Process Volumes** will override one another if they are overlapping. |
+| **[1.17 Priorities](Resources\ContentExamples\PostProcessing\1_17)** | Allows you to control which **Post Process Volumes** will override one another if they are overlapping. |
 | **[1.18 Blend Radius and Weight](Resources\ContentExamples\PostProcessing\1_18)** | Used to change the transition from one **Post Process Volume** to another over a distance (e.g. blend scene tint or color based on distance to an object). |
 | **[1.19 Unbound](Resources\ContentExamples\PostProcessing\1_19)** | This allows you to create a global **Post Process** effect without being constrained to being inside an actual volume (e.g. the effects are applied where ever you are in the world). |
 
