#UE4doc #docs Grammatical and content updates.

diff --git a/Engine/Documentation/Source/Engine/Rendering/Materials/ExpressionReference/Math/MathExpressions.INT.udn b/Engine/Documentation/Source/Engine/Rendering/Materials/ExpressionReference/Math/MathExpressions.INT.udn
--- a/Engine/Documentation/Source/Engine/Rendering/Materials/ExpressionReference/Math/MathExpressions.INT.udn
+++ b/Engine/Documentation/Source/Engine/Rendering/Materials/ExpressionReference/Math/MathExpressions.INT.udn
@@ -66,7 +66,7 @@ The Clamp expression takes in value(s) and constrains them to a specified range,
 | **Max** | Takes in the value to use as the maximum when clamping. |
 
 
-**Examples:** Clamping 0.3 with Min 0.0 and Max 1.0 yields 0.3; Clamping 1.3 with Min 0.0 and Max 1.0 yields 1.0;
+**Examples:** Clamping 0.3 with Min 0.0 and Max 1.0 yields 0.3; Clamping 1.3 with Min 0.0 and Max 1.0 yields 1.0.
 
 ![ClampExample.png](ClampExample.png)(w:900)
 
@@ -96,7 +96,7 @@ In the example above, the color would oscillate with co-sinusoidal frequency.
 [EXCERPT:CrossProduct]
 ## CrossProduct
 
-The CrossProduct expression computes the cross product of two three-channel vector value inputs and outputs the resulting three-channel vector value. Given two lines (or vectors) in space, the cross product is a line (or vector) which is perpendicular to both of the inputs.
+The CrossProduct expression computes the cross product of two three-channel vector value inputs and outputs the resulting three-channel vector value. Given two vectors in space, the cross product is a vector that is perpendicular to both of the inputs.
 
 | Item | Description |
 | -------- | ----------- |
@@ -114,7 +114,7 @@ The CrossProduct expression computes the cross product of two three-channel vect
 
 ## Divide
 
-The Divide expression takes two inputs and outputs the result of the first divided by the second.  The division happens per channel, meaning that the R channel of the first is divided by the second, the G channel of the first is divided by the second, and so on. Both inputs must have the same number of values unless the divisor is a single float value.
+The Divide expression takes two inputs and outputs the result of the first divided by the second.  The division happens per channel, meaning that the R channel of the first is divided by the second, the G channel of the first is divided by the second, and so on. Both inputs must have the same number of values unless the divisor is a single float value. _Never divide by zero._
 
 | Item | Description |
 | -------- | ----------- |
@@ -168,25 +168,25 @@ In this example, FMod is returning a repeating value that ramps up from 0 to 1 e
 
 The Frac expression takes in value(s) and outputs the decimal part of those values.  See also [Ceil](#Ceil) and [Floor](#Floor).
 
-**Examples:** Frac of 0.2 is 0.2; Frac of (0.0,1.6) is (0.0, 0.6)
+**Examples:** Frac of 0.2 is 0.2; Frac of (0.0,1.6) is (0.0, 0.6).
 
 ![FracExample.png](FracExample.png)(w:900)
 
-In this example, the Frac node is converting time into a series of repeating 0-1 progressions, causing the color to fade (via the Lerp) from green to red, then snap back to green, and repeat indefinitely.
+In this example, the Frac node is converting time into a series of repeating 0 - 1 progressions, causing the color to fade (via the Lerp) from green to red, then snap back to green, repeating indefinitely.
 
 [EXCERPT:If]
 ## If
 
-The If expression compares two inputs and then passes through one of three other input values based on the result of the comparison. The two inputs to be compared must be single float values.
+The If expression compares two inputs and then passes through one of three other input values based on the result of the comparison. Both compared inputs must be single float values.
 
 | Item | Description |
 | -------- | ----------- |
 | [REGION:tablesection]Inputs[/REGION] ||
 | **A** | Takes in a single float value. |
 | **B** | Takes in a single float value. |
-| **A&lt;B** | Takes in the value(s) to output if the value of A is greater than the value of B. |
-| **A=B** | Takes in the value(s) to output if the value of A is equal to the value of B. |
-| **A&gt;B** | Takes in the value(s) to output if the value of A is less than the value of B. |
+| **A &lt; B** | Takes in the value(s) to output if the value of A is greater than the value of B. |
+| **A = B** | Takes in the value(s) to output if the value of A is equal to the value of B. |
+| **A &gt; B** | Takes in the value(s) to output if the value of A is less than the value of B. |
 
 
 ![IfExample.png](IfExample.png)(w:900)
@@ -201,13 +201,13 @@ In this example, the If expression takes in a heightmap and outputs either red o
 
 ## Max
 
-The Max expression takes in two inputs then takes the maximum of the two. 
+The Max expression takes in two inputs and then outputs the maximum of the two. 
 
 This node is similar to Photoshop's Lighten. 
 
 ![](MaxExample.png)(w:720)
 
-In the example above, A is "0" and B is "1" therefore "1" (white) is the resulting base color. 
+In the example above, A is "0" and B is "1"; therefore, "1" (white) is the resulting base color. 
 
 | Item | Description |
 | -------- | ----------- |
@@ -221,13 +221,13 @@ In the example above, A is "0" and B is "1" therefore "1" (white) is the resulti
 
 ## Min
 
-The Min expression takes in two inputs then takes the minimum of the two. 
+The Min expression takes in two inputs then outputs the minimum of the two. 
 
 This node is similar to Photoshop's Darken. 
 
 ![](MinExample.png)(w:720)
 
-In the example above, A is "0" and B is "1" therefore "0" (black) is the resulting base color.
+In the example above, A is "0" and B is "1"; therefore, "0" (black) is the resulting base color.
 
 | Item | Description |
 | -------- | ----------- |
@@ -243,7 +243,7 @@ In the example above, A is "0" and B is "1" therefore "0" (black) is the resulti
 
 ## Multiply
 
-The Multiply expression takes two inputs, multiplies them together, and outputs the result. Similar to Photoshop's multiply layer blend.  The multiplication happens per channel, meaning that the R channel of the first is multiplied by the second; the G channel of the first is multiplied by the second, and so on. Both inputs must have the same number of values unless one of the values is a single float value.
+The Multiply expression takes two inputs, multiplies them together, and outputs the result. Similar to Photoshop's multiply layer blend.  Multiplication happens per channel, meaning that the R channel of the first is multiplied by the second; the G channel of the first is multiplied by the second, and so on. Both inputs must have the same number of values unless one of the values is a single float value.
 
 | Item | Description |
 | -------- | ----------- |
@@ -255,9 +255,9 @@ The Multiply expression takes two inputs, multiplies them together, and outputs
 | **B** | Takes in the second value(s) to multiply. |
 
 
-Do not forget that materials in UE4 are not limited to [0,1]!  If colors/values are greater than 1, Multiply will actually brighten colors.
+Do not forget that materials in UE4 are not limited to [0,1].  If colors/values are greater than 1, Multiply will actually brighten colors.
 
-**Examples:** Multiply of 0.4 and 0.5 is 0.2; Multiply of (0.2,-0.4,0.6) and (0.0,2.0,1.0) is (0.0,-0.8,0.6); Multiply of (0.2,-0.4,0.6) and 0.5 is (0.1,-0.2,0.3)
+**Examples:** Multiply of 0.4 and 0.5 is 0.2; Multiply of (0.2,-0.4,0.6) and (0.0,2.0,1.0) is (0.0,-0.8,0.6); Multiply of (0.2,-0.4,0.6) and 0.5 is (0.1,-0.2,0.3).
 
 **Example Usage:** Multiply is often used to brighten or darken colors/textures.
 
@@ -272,18 +272,18 @@ The Normalize expression calculates and outputs the normalized value of its inpu
 ![NormalizeExample.png](NormalizeExample.png)(w:900)
 
 [REGION:note]
-It is not necessary to normalize an expression which plugs into the Normal material output.
+It is not necessary to normalize an expression that plugs into the Normal material output.
 [/REGION]
 
 [/EXCERPT:Normalize]
 
 ## OneMinus
 
-The OneMinus expression takes an input value and outputs one minus that value. This operation is performed per channel.
+The OneMinus expression takes an input value and outputs 'one-minus-that-value'. This operation is performed per channel.
 
-**Examples:** OneMinus of 0.4 is 0.6; OneMinus of (0.2,0.5,1.0) is (0.8,0.5,0.0); OneMinus of (0.0,-0.4,1.6) is (1.0,1.4,-0.6)
+**Examples:** OneMinus of 0.4 is 0.6; OneMinus of (0.2,0.5,1.0) is (0.8,0.5,0.0); OneMinus of (0.0,-0.4,1.6) is (1.0,1.4,-0.6).
 
-**Example Usage:** When the input colors are in the range [0,1], OneMinus has the same effect as what is commonly called "invert" -- that is, OneMinus returns the complimentary color that when added to the input will produce white.
+**Example Usage:** When the input colors are in the range [0,1], OneMinus has the same effect as what is commonly called "invert" -- that is, OneMinus returns the complementary color that when added to the input will produce white.
 
 ![OneMinusExample.png](OneMinusExample.png)
 
@@ -301,9 +301,9 @@ The Power expression takes two inputs, raises Base to the Exp power, and outputs
 | **Exp** | Takes in the exponent value. |
 
 
-**Examples:** Power of Base 0.5, Exp 2.0 is 0.25;
+**Examples:** Power of Base 0.5, Exp 2.0 is 0.25.
 
-**Example Usage:** If the colors you pass to Power are in [0,1], Power can behave as a contrast adjustment, where only the brighter values remain.
+**Example Usage:** If the colors you pass to Power are in [0,1], Power can behave as a contrast adjustment, where only brighter values remain.
 
 ![PowerExample.png](PowerExample.png)
 
@@ -334,7 +334,7 @@ The SquareRoot expression outputs the square root of the input value.  SquareRoo
 
 ## Subtract
 
-The Subtract node takes in two inputs, subtracts the second input from the first, and outputs the difference.  The subtraction happens per channel, meaning that the R channel of the second input gets subtracted from the first; the G channel of the second input gets subtracted from the first, and so on. Both inputs must have the same number of channels unless the second input is a single Constant value. Constants can be subtracted from a vector with any number of inputs.
+The Subtract node takes in two inputs, subtracts the second input from the first, and outputs the difference.  Subtraction happens per channel, meaning that the R channel of the second input gets subtracted from the first; the G channel of the second input gets subtracted from the first, and so on. Both inputs must have the same number of channels unless the second input is a single Constant value. Constants can be subtracted from a vector with any number of inputs.
 
 | Item | Description |
 | -------- | ----------- |
@@ -346,7 +346,7 @@ The Subtract node takes in two inputs, subtracts the second input from the first
 | **B** | Takes in the value(s) of the subtrahend. |
 
 
-**Examples:** Subtract of 0.5 and 0.2 is 0.3; Subtract of (0.2,-0.4,0.6) and (0.1,0.1,1.0) is (0.1,-0.5,-0.4); Subtract of (0.2,0.4,1.0) and 0.2 is (0.0,0.2,0.8)
+**Examples:** Subtract of 0.5 and 0.2 is 0.3; Subtract of (0.2,-0.4,0.6) and (0.1,0.1,1.0) is (0.1,-0.5,-0.4); Subtract of (0.2,0.4,1.0) and 0.2 is (0.0,0.2,0.8).
 
 **Example Usage:** Subtract can be used to darken colors and offset UVs.
 
diff --git a/Engine/Documentation/Source/Engine/Rendering/Materials/ExpressionReference/VectorOps/VectorOpExpressions.INT.udn b/Engine/Documentation/Source/Engine/Rendering/Materials/ExpressionReference/VectorOps/VectorOpExpressions.INT.udn
--- a/Engine/Documentation/Source/Engine/Rendering/Materials/ExpressionReference/VectorOps/VectorOpExpressions.INT.udn
+++ b/Engine/Documentation/Source/Engine/Rendering/Materials/ExpressionReference/VectorOps/VectorOpExpressions.INT.udn
@@ -34,10 +34,10 @@ The ComponentMask expression allows you to select a specific subset of channels
 | Item | Description |
 | -------- | ----------- |
 | [REGION:tablesection]Properties[/REGION] ||
-| **R** | If checked, the red, or first, channel of the input value will be passed through to the output. |
-| **G** | If checked, the green, or second, channel of the input value will be passed through to the output. |
-| **B** | If checked, the blue, or third, channel of the input value will be passed through to the output. |
-| **A** | If checked, the alpha, or fourth, channel of the input value will be passed through to the output. |
+| **R** | If checked, the red (first) channel of the input value will be passed through to the output. |
+| **G** | If checked, the green (second) channel of the input value will be passed through to the output. |
+| **B** | If checked, the blue (third) channel of the input value will be passed through to the output. |
+| **A** | If checked, the alpha (fourth) channel of the input value will be passed through to the output. |
 
 
 **Examples:** ComponentMask with an input of (0.2,0.8,1.4) and the R and B channels checked will output (0.2,1.4).
