#UE4doc #docs Grammatical update, including emphasis added to the Event Tick key phrase.

diff --git a/Engine/Documentation/Source/Resources/ContentExamples/Blueprints/1_6/1_6.INT.udn b/Engine/Documentation/Source/Resources/ContentExamples/Blueprints/1_6/1_6.INT.udn
--- a/Engine/Documentation/Source/Resources/ContentExamples/Blueprints/1_6/1_6.INT.udn
+++ b/Engine/Documentation/Source/Resources/ContentExamples/Blueprints/1_6/1_6.INT.udn
@@ -42,20 +42,20 @@ The Components for this example consist of two Text Render Components. One simpl
 ![](1.6Components.png)(w:500)
 
 ## Construction Script
-The Construction Script for this example simply sets the text of the Text Render Component to the value of the private variable IntegerVariable. This data is cast to a string in order to be compatible with the _Text_ property of the Text Render Component.
+The Construction Script for this example simply sets the text of the Text Render Component to the value of IntegerVariable (private variable). This data is cast to a string in order to be compatible with the _Text_ property of the Text Render Component.
 
 ![](1.6ConstructionScript.png)
 
 ## EventGraph
 
-The event graph for this Blueprint handles the work of the mathematical addition.
+The event graph for this Blueprint handles the work of mathematical addition.
 
 ![](1.6EventGraph.png)
 
 
-1. On every tick (Event Tick)
+* On every tick (**Event Tick**):
 	1. A 1-second delay starts.
-	1. At the end of the delay, we add 1 to the current value of IntegerVariable and sets the sum as the new value of IntegerVariable.
+	1. At the end of the delay, we add 1 to the current value of IntegerVariable and set the sum as the new value of IntegerVariable.
 	1. Sets the text of the TextRenderComponent to the value of the  IntegerVariable. This value gets cast to a string.
 
 [/REGION]
\ No newline at end of file
