#UE4doc #docs Minor punctuation repair.

diff --git a/Engine/Documentation/Source/Engine/Blueprints/UserGuide/MathNode/MathNode.INT.udn b/Engine/Documentation/Source/Engine/Blueprints/UserGuide/MathNode/MathNode.INT.udn
--- a/Engine/Documentation/Source/Engine/Blueprints/UserGuide/MathNode/MathNode.INT.udn
+++ b/Engine/Documentation/Source/Engine/Blueprints/UserGuide/MathNode/MathNode.INT.udn
@@ -163,4 +163,4 @@ Currently, the following operators are unsupported.
 * Unary prefix: 
 		+, -, ++, --, ~, !, etc.
 * Postfix: 
-		++, --, [], etc
+		++, --, [], etc.
