#UE4doc #docs Minor grammatical update.

diff --git a/Engine/Documentation/Source/Resources/ContentExamples/MathHall/MathHall.INT.udn b/Engine/Documentation/Source/Resources/ContentExamples/MathHall/MathHall.INT.udn
--- a/Engine/Documentation/Source/Resources/ContentExamples/MathHall/MathHall.INT.udn
+++ b/Engine/Documentation/Source/Resources/ContentExamples/MathHall/MathHall.INT.udn
@@ -47,7 +47,7 @@ Related: Engine/Rendering/Materials/ExpressionReference
 ## Overview
 The **Math_Hall** map inside the **Content Examples** project aims to give you more of an understanding of how math is used within Unreal Engine 4 inside of **Blueprints** and **Materials**. 
 
-The first section entitled **Functions**, covers Math and Material Functions (left side of the wall) and example usage cases (right side of the wall). This section is excellent for learning more about blending Materials as well as how logic can be applied to Materials to create events like rippling water or a bouncing ball entirely inside the **Material Editor**.  
+The first section entitled **Functions**, covers Math and Material Functions (left side of the wall) and example usage cases (right side of the wall). This section is excellent for learning more about blending Materials as well as how logic can be applied to Materials to create events like rippling water or bouncing balls entirely inside the **Material Editor**.  
 
 The second section covers **Vector Math** and describes the various types of Vectors and how they are used in Materials as well as their importance as it relates to Blueprints (e.g. making objects face the player, checking if players are looking towards something or finding out how far the player is from an object). 
 
@@ -66,7 +66,7 @@ Listed below are the examples provided in the **Functions** section, example use
 | **1.5 Ceil** | **Ceil** works just like floor but rounds up instead of down. The result of ceil is always 1 greater than floor (e.g. 0.1 becomes 1, 1.1 becomes 2, etc. An example of creating a Mosaic is also provided). |
 | **1.6 FMod** | **FMod** stands for Floating Modulus and has both A and B inputs. It divides A by B and returns the remainder (e.g. an example is provided for generating a stripe pattern on a material). |
 | **1.7 Power** | **Power** takes a base input and applies an exponent. Negative values become positive just like **ABS** and a power of 1 will give the same result as ABS (e.g. an example is given to adjust the contrast of a texture).  |
-| **1.8 Lerp** | Linear Interpolation or **Lerp** returns a blend between the A and B inputs based on the alpha input. In the example given, when the alpha is 0, 100% of A will be returned. When the alpha is 1, 100% of B is returned (e.g. an example of using Lerp to blend between two textures is given). |
+| **1.8 Lerp** | Linear Interpolation (**Lerp**) returns a blend between the A and B inputs based on the alpha input. In the given example, when the alpha is 0, 100% of A will be returned. When the alpha is 1, 100% of B is returned (e.g. an example of using Lerp to blend between two textures is given). |
 | **1.9 HeightLerp** | **HeightLerp** is a **Lerp** that modulates the blend using a heightmap (e.g. an example is given where the heightmap of a texture is used to blend between two textures where a grass texture blends into a cobblestone texture). |
 | **1.10 Clamp** | **Clamp** is used to clamp the input value to the specified Min and Max Range (e.g. using clamp to specify a min/max view angle on a security camera). |
 | **1.11 Sine** | **Sine** is a function of an angle and represents a ratio between the sides of a right triangle (e.g. Sine has a variety of applications and can create swaying motions inside Blueprints or Materials; it can also be used to make "waves" in water). |
