fixed typo

diff --git a/Engine/Documentation/Source/Engine/Blueprints/Editor/UIBreakdowns/InterfaceUI/InterfaceUI.INT.udn b/Engine/Documentation/Source/Engine/Blueprints/Editor/UIBreakdowns/InterfaceUI/InterfaceUI.INT.udn
--- a/Engine/Documentation/Source/Engine/Blueprints/Editor/UIBreakdowns/InterfaceUI/InterfaceUI.INT.udn
+++ b/Engine/Documentation/Source/Engine/Blueprints/Editor/UIBreakdowns/InterfaceUI/InterfaceUI.INT.udn
@@ -48,7 +48,7 @@ When you first open a [](Engine/Blueprints/UserGuide\Types\Interface), the UI wi
 1. [Toolbar](Engine/Blueprints/Editor\UIComponents\Toolbar\ToolbarOptions)
 1. [My Blueprint](Engine/Blueprints/Editor\UIComponents\MyBlueprint)
 1. [Details](Engine/Blueprints/Editor\UIComponents\Details)
-1. [Graph Editor](Engine/Blueprints/Editor\UIComponents\GraphEditor)v
+1. [Graph Editor](Engine/Blueprints/Editor\UIComponents\GraphEditor)
 [/EXCERPT:Visible]
 
 [EXCERPT:Invisible]
