#UE4doc #docs: Made changes to match the instructional guide v1.0

diff --git a/Engine/Documentation/Source/Engine/Animation/AnimHowTo/AnimHowToOverview.INT.udn b/Engine/Documentation/Source/Engine/Animation/AnimHowTo/AnimHowToOverview.INT.udn
--- a/Engine/Documentation/Source/Engine/Animation/AnimHowTo/AnimHowToOverview.INT.udn
+++ b/Engine/Documentation/Source/Engine/Animation/AnimHowTo/AnimHowToOverview.INT.udn
@@ -1,13 +1,13 @@
 Availability: Public
-Title:Animation - How To...
+Title:Animation How To's
 Crumbs: %ROOT%, Engine, Engine\Animation, Engine\Animation\AnimHowTo
 Description:Examples of different aspects of the Animation system in Unreal Engine 4 are demonstrated through several step-by-step How To guides.
 
-The **Animation - How To...** pages provide several step-by-step examples of how to work with and use different aspects of the Animation system in Unreal Engine 4. Whether you are new to Animation in UE4 or want to learn how to work with certain aspects of the Animation system, the How To pages below cover a broad range of topics and should give you a place to start or expand upon your current knowledge. 
+The **Animation How To's** pages provide several step-by-step examples of how to work with and use different aspects of the Animation system in Unreal Engine 4. Whether you are new to Animation in UE4 or want to learn how to work with certain aspects of the Animation system, the How To pages below cover a broad range of topics and should give you a place to start or expand upon your current knowledge. 
 
 Refer to the table below for more information and click each link to view its contents. 
 
-|How To... | Description |
+|How To | Description |
 | ------- | ------------ |
 | **[](Engine\Animation\AnimHowTo\PreviewingAnimations)** | %Engine\Animation\AnimHowTo\PreviewingAnimations:description% |
 | **[](Engine\Animation\AnimHowTo\Retargeting)** | %Engine\Animation\AnimHowTo\Retargeting:description% |
