#UE4 Docs: FIxing Jira UEDOC-1631 by adding a new section about setting the texture LOD's on a per device basis.

diff --git a/Engine/Documentation/Source/Platforms/Mobile/Textures/MobileTextures.INT.udn b/Engine/Documentation/Source/Platforms/Mobile/Textures/MobileTextures.INT.udn
--- a/Engine/Documentation/Source/Platforms/Mobile/Textures/MobileTextures.INT.udn
+++ b/Engine/Documentation/Source/Platforms/Mobile/Textures/MobileTextures.INT.udn
@@ -35,3 +35,26 @@ the following setup which only requires two textures:
 	* TC_NormalMap compression
 
 Additional layers will need to use more textures, with the option of only layering in BaseColor + Roughness to avoid multiple normal map texture fetches.
+
+## Per-Device Texture LOD
+
+With the release of Unreal Engine 4.8 you can now specify Texture LOD settings on a per-device basis as opposed to per-platform. 
+To set the Texture LOD settings per-device you will need to do the following.
+
+1. From the main Tool Bar open the **Devices Profiles** manager by going to **Windows** > **Developer Tools** and then clicking on **Device Profiles**.
+
+	![](T_Device_Profiles_Manager.png)
+
+1. After clicking on the **Device Profiles** option the following menu will be displayed.
+
+	![](T_DPM.png)
+
+1. From this menu select the device you want to edit and then click on the three dots that are to the left of the device name to open that devices profile.
+
+1. To make a change just click on the property you want to adjust and then input what you want the setting to be.
+
+	![](T_Selecting_Editing_Profile.gif)(convert:false)
+
+
+
+
