#UE4 Docs: Fixing Jira UEDOC-999 by adding the mobile previewer page to the mobile landing page and also added some images to better show off what you should do.

diff --git a/Engine/Documentation/Source/Platforms/Mobile/Mobile.INT.udn b/Engine/Documentation/Source/Platforms/Mobile/Mobile.INT.udn
--- a/Engine/Documentation/Source/Platforms/Mobile/Mobile.INT.udn
+++ b/Engine/Documentation/Source/Platforms/Mobile/Mobile.INT.udn
@@ -26,6 +26,7 @@ Navigation:platform
 	* [](Platforms/Mobile/Content "%Platforms/Mobile/Content:description%")
 	* [](Platforms/Mobile/PostProcessEffects "%Platforms/Mobile/PostProcessEffects:description%")
 	* [](Platforms/Mobile/Performance "%Platforms/Mobile/Performance:description%")
+	* [](Platforms/Mobile/Previewer "%Platforms/Mobile/Previewer:description%")
 	[/PARAM]
 [/OBJECT]
 
diff --git a/Engine/Documentation/Source/Platforms/Mobile/Previewer/MobilePreviewer.INT.udn b/Engine/Documentation/Source/Platforms/Mobile/Previewer/MobilePreviewer.INT.udn
--- a/Engine/Documentation/Source/Platforms/Mobile/Previewer/MobilePreviewer.INT.udn
+++ b/Engine/Documentation/Source/Platforms/Mobile/Previewer/MobilePreviewer.INT.udn
@@ -4,5 +4,14 @@ Crumbs:%ROOT%, Platforms, Platforms/Mobile
 Description:Previewing games on mobile platforms.
 Platform:Mobile
 
-There is currently no in-editor preview. You have to launch a standalone game with the 'Enable Mobile Preview' 
-option set. The mobile previewer is intended to match devices as closely as possible.
+There is currently no in-editor preview to simulate mobile rendering.
+To see what your project could look like on mobile, you have to launch a standalone game with the **Enable Mobile Preview** option set. 
+To enable the Mobile preview go to the **Main Tool Bar** and click on the arrow that is to the right of the **Play** button to expose the different play modes.
+Click on the one that says **Mobile Preview** and your project will then be launched in a new window that simulates what your project should look like on a mobile device.
+
+![](T_Activate_Mobile_Preview.png)
+
+[region:note]
+Please note that the Mobile previewer is intended to match mobile devices as closely as possible but is not indicative of what your project is going to look like on the target device.
+You should always make sure to fully test your project on your target device and only use the mobile previewer to see if your work is headed in the right direction.
+[/region]
