#UE4doc #docs: changed spelling error. Additve -> Additive

diff --git a/Engine/Documentation/Source/Engine/Animation/AimOffset/AimOffset.INT.udn b/Engine/Documentation/Source/Engine/Animation/AimOffset/AimOffset.INT.udn
--- a/Engine/Documentation/Source/Engine/Animation/AimOffset/AimOffset.INT.udn
+++ b/Engine/Documentation/Source/Engine/Animation/AimOffset/AimOffset.INT.udn
@@ -72,7 +72,7 @@ Once you import each of the poses you need into Unreal Engine, you will need to
 * **Additive Aim Type** - AAT_MeshSpace
 
 [REGION:note]
-Aim Offsets only accept animations with an Additve Aim Type of Mesh Space!
+Aim Offsets only accept animations with an Additive Aim Type of Mesh Space!
 [/REGION]
 
 
