#UE4doc #docs: Changed capitalization on a word

diff --git a/Engine/Documentation/Source/Engine/Animation/AnimationComposite/AnimationComposite.INT.udn b/Engine/Documentation/Source/Engine/Animation/AnimationComposite/AnimationComposite.INT.udn
--- a/Engine/Documentation/Source/Engine/Animation/AnimationComposite/AnimationComposite.INT.udn
+++ b/Engine/Documentation/Source/Engine/Animation/AnimationComposite/AnimationComposite.INT.udn
@@ -54,7 +54,7 @@ To add Animation Sequences to the Composite, simply drag and drop them into the
 ![](AddingAnimationsToComposite.png)(w:700)
 
 [REGION:note]
-As you drag an asset out of the **Asset Browser**, you may notice the Asset Editor gray out and warn that you are **Currently previewing another asset**. This is expected behavior. Simply continue dragging the Animation Sequence onto the Composite track.
+As you drag an asset out of the **Asset Browser**, you may notice the Asset Editor gray out and warn that you are **currently previewing another asset**. This is expected behavior. Simply continue dragging the Animation Sequence onto the Composite track.
 [/REGION]
 
 You can move Sequences along the Composite track by dragging them, and remove them from the track by **right-clicking** them and choosing **Delete Segment** from the context menu.
