#UE4doc #docs: changed multiple spelling errors

diff --git a/Engine/Documentation/Source/Engine/Animation/AnimBlueprints/Interface/AnimationBlueprintUI.INT.udn b/Engine/Documentation/Source/Engine/Animation/AnimBlueprints/Interface/AnimationBlueprintUI.INT.udn
--- a/Engine/Documentation/Source/Engine/Animation/AnimBlueprints/Interface/AnimationBlueprintUI.INT.udn
+++ b/Engine/Documentation/Source/Engine/Animation/AnimBlueprints/Interface/AnimationBlueprintUI.INT.udn
@@ -58,7 +58,7 @@ The following sections describe the areas of the Animation Blueprint UI.
 | --- | --- |
 | Toolbar | Shows the Toolbar if it is currently not visible. |
 | Details | Shows the [Details](Engine/Blueprints\Editor/UIComponents\Details) panel if it is currently not visible. |
-| Debug | Shows the **Debug** panel if if it currently not visible. |
+| Debug | Shows the **Debug** panel if it is currently not visible. |
 | Palette | Shows the [](Engine/Blueprints\Editor/UIComponents\Palette) panel if it is not currently visible. |
 | My Blueprint | Shows the [](Engine/Blueprints\Editor/UIComponents\MyBlueprint) panel if it is not currently visible. |
 | Compiler Results | Shows the **Compiler Results** panel if it is not currently visible.  |
@@ -67,10 +67,10 @@ The following sections describe the areas of the Animation Blueprint UI.
 | Asset Browser | Shows the **Asset Browser** panel if it is currently closed, and changes focus to it if it is open. |
 | Viewport | Shows the **Preview Viewport** panel if it is not currently visible.  |
 | Animation Notifies | Opens the **Animation Notifies** window, which lists all Animation Notifies currently associated with the Skeletal Mesh's animation assets. |
-| Animation Blueprint Editor | Opens the **Animation Blueprint Editor** window if it is currenlty not visible. |
-| Asset Override Editor | Opens the **Asset Override Editor** window if it is currenlty not visible. |
-| Animation SlotName | Opens the **Animation SlotName** window if it is currenlty not visible. |
-| Animation Slot Group Name | Opens the **Animation Slot Group Name** window if it is currenlty not visible. |
+| Animation Blueprint Editor | Opens the **Animation Blueprint Editor** window if it is currently not visible. |
+| Asset Override Editor | Opens the **Asset Override Editor** window if it is currently not visible. |
+| Animation SlotName | Opens the **Animation SlotName** window if it is currently not visible. |
+| Animation Slot Group Name | Opens the **Animation Slot Group Name** window if it is currently not visible. |
 | Content Browser | Opens a sub-menu with access to all four **Content Browsers**. |
 | Developer Tools | Opens a sub-menu with access to **Camera Debugger**, **Collision Analyzer**, **Debug Tools**, **Modules**, and **Widget Reflector**. |
 | Blueprint Debugger | Opens the **Blueprint Debugger**. |
@@ -84,7 +84,7 @@ The following sections describe the areas of the Animation Blueprint UI.
 | [REGION:tablesection]**Layout**[/REGION] ||
 | Reset Layout | Resets the layout to the default setting (requires automatic restart of the editor). |
 | Save Layout | Saves the current panel layout. |
-| Enable Fullscreen | Enables/Disables the display of the edtior in Fullscreen mode. |
+| Enable Fullscreen | Enables/Disables the display of the editor in Fullscreen mode. |
 
 
 
