#UE4doc #docs: fixed spelling error: permanancy  -> permanency

diff --git a/Engine/Documentation/Source/Engine/Animation/Persona/BlueprintDefaults/BlueprintDefaults.INT.udn b/Engine/Documentation/Source/Engine/Animation/Persona/BlueprintDefaults/BlueprintDefaults.INT.udn
--- a/Engine/Documentation/Source/Engine/Animation/Persona/BlueprintDefaults/BlueprintDefaults.INT.udn
+++ b/Engine/Documentation/Source/Engine/Animation/Persona/BlueprintDefaults/BlueprintDefaults.INT.udn
@@ -46,7 +46,7 @@ You can preview your various animation assets by manually modifying the values o
 
 ## Edit Modes
 
-Along the top of the **Anim Blueprint Editor** panel are two modes that can be switched between and affect the permanancy of changes made to variable values in the panel.
+Along the top of the **Anim Blueprint Editor** panel are two modes that can be switched between and affect the permanency of changes made to variable values in the panel.
 
 | Edit Mode | Description | |
 | --- | ----------- | --- |
