#UE4doc #docs: Fixed spelling error: digram -> diagram

diff --git a/Engine/Documentation/Source/Engine/Animation/IKSetups/IKSetups.INT.udn b/Engine/Documentation/Source/Engine/Animation/IKSetups/IKSetups.INT.udn
--- a/Engine/Documentation/Source/Engine/Animation/IKSetups/IKSetups.INT.udn
+++ b/Engine/Documentation/Source/Engine/Animation/IKSetups/IKSetups.INT.udn
@@ -34,7 +34,7 @@ For examples of Hand and Foot **IK**, you can also reference the [Animation Cont
 
 If you are already familiar with what IK is and how it works, you can skip this section!
 
-Most animated skeletons in Unreal are driven by direct rotational data fed straight into the bones of the character or Skeletal Mesh. This can be thought of as _forward kinematics_, or direct application of rotation to joints or bones. Below is a digram illustrating the concept:
+Most animated skeletons in Unreal are driven by direct rotational data fed straight into the bones of the character or Skeletal Mesh. This can be thought of as _forward kinematics_, or direct application of rotation to joints or bones. Below is a diagram illustrating the concept:
 
 [REGION:raw]
 ![](diagram_FK.png)(w:600)
