#UE4doc #docs: fixed spelling error: aniamtion -> animation

diff --git a/Engine/Documentation/Source/Engine/Animation/CharacterSetupOverview/CharacterSetupOverview.INT.udn b/Engine/Documentation/Source/Engine/Animation/CharacterSetupOverview/CharacterSetupOverview.INT.udn
--- a/Engine/Documentation/Source/Engine/Animation/CharacterSetupOverview/CharacterSetupOverview.INT.udn
+++ b/Engine/Documentation/Source/Engine/Animation/CharacterSetupOverview/CharacterSetupOverview.INT.udn
@@ -168,7 +168,7 @@ So, to summarize the flow of setup back up the chain:
 * The GameMode specifies which Pawn (Character) Class and which PlayerController Class you will need to play the game.
 * The Character Class: 
 	* Contains the Skeletal Mesh that you imported via FBX.
-	* Takes in data from the PlayerController Class and converts it to movement (not aniamtion).
+	* Takes in data from the PlayerController Class and converts it to movement (not animation).
 	* Stores which Animation Blueprint will be used to drive skeletal animation within its Mesh Component.
 * The Animation Blueprint:
 	* Takes data from the Character Class into its Event Graph.
