#UE4doc #docs: fixed spelling error: asssets -> assets

diff --git a/Engine/Documentation/Source/Engine/Animation/Persona/MenuBar/PersonaMenuBar.INT.udn b/Engine/Documentation/Source/Engine/Animation/Persona/MenuBar/PersonaMenuBar.INT.udn
--- a/Engine/Documentation/Source/Engine/Animation/Persona/MenuBar/PersonaMenuBar.INT.udn
+++ b/Engine/Documentation/Source/Engine/Animation/Persona/MenuBar/PersonaMenuBar.INT.udn
@@ -54,7 +54,7 @@ If you are editing animation assets such as _Skeletal Meshes_ or _Animation Sequ
 | **Set Preview Mesh as Default** | Changes the skeletons default preview mesh to the current open preview mesh. The skeleton will require saving after this action. |
 | **Remove Unused Bones from Skeleton** | Removes any bones from the skeleton that are not used by any of its meshes. The skeleton and associated animations will require saving after this action.  |
 | **Record to new Animation** | Create a new animation from the currently playing animation. |
-| [REGION:tablesection]**The following section applies only if Persona was opend via an _Animation_** asset[/REGION] ||
+| [REGION:tablesection]**The following section applies only if Persona was opened via an _Animation_** asset[/REGION] ||
 | **Reimport** (Animation) | Reimports the AnimationSequence. |
 | **Apply Compression** | Apply compression to the current animation. |
 | **Export to FBX** | Export the current animation to an FBX. |
@@ -111,7 +111,7 @@ The panels available in the **Window** menu, as well as their order of appearanc
 | **Mesh Details**  _(Mesh mode)_ | Opens the **Mesh Details** panel, where you can change aspects of the mesh including Material applications, LODs, Physics Assets, and more. |
 | **Anim Asset Details** _(Animation mode)_ | Opens the **Anim Asset Details** panel, where you perform most of the necessary edits to Animation Sequences, Blend Spaces, Animation Montages, and Aim Offsets. |
 | **Anim Blueprint Editor** _(Animation Blueprint)_| Opens the **Anim Blueprint Editor** panel if it is not already open, where you can set the default values for all variables in your Blueprint or change the values of the preview instance, depending on mode. |
-| **Asset Override Editor** _(Animation Blueprint)_| Opens the **Asset Override Editor** panel if it is not already open. This editor is for overriding the animation asssets referenced by the parent animation graph. |
+| **Asset Override Editor** _(Animation Blueprint)_| Opens the **Asset Override Editor** panel if it is not already open. This editor is for overriding the animation assets referenced by the parent animation graph. |
 | **Animation SlotName** | Opens the **Animation SlotName** panel and allows you to modify custom animation SlotNames. |
 | **Animation Slot Group Name** | Opens the **Animation Slot Group Name** panel which lets you modify custom animation Slot Group Names. |
 | **Debug** | Opens the **Debug** window, which shows all breakpoints throughout the Blueprint. |
