#UE4doc #docs: Minor grammar adjustment: a -> an.

diff --git a/Engine/Documentation/Source/Engine/Animation/Persona/Modes/Modes.INT.udn b/Engine/Documentation/Source/Engine/Animation/Persona/Modes/Modes.INT.udn
--- a/Engine/Documentation/Source/Engine/Animation/Persona/Modes/Modes.INT.udn
+++ b/Engine/Documentation/Source/Engine/Animation/Persona/Modes/Modes.INT.udn
@@ -33,7 +33,7 @@ In each mode, Persona's UI updates to show a different layout of tools and edito
 [/EXCERPT:Overview]
 
 [REGION:note]
-	Not all of the modes are available at all times. For instance, Animation Blueprint mode is _only_ available when a Animation Blueprint is opened from the **Content Browser**. 
+	Not all of the modes are available at all times. For instance, Animation Blueprint mode is _only_ available when an Animation Blueprint is opened from the **Content Browser**. 
 
 	Also, be aware that changes made to editing layouts will be maintained. As an example, if you move a panel while in Skeleton editing mode, then switch away to Animation, when you switch back to Skeleton, the moved panel will keep the position you gave it.
 [/REGION]
