#UE4doc #docs: Spelling and spacing changes: Heirachy -> Hierarchy

diff --git a/Engine/Documentation/Source/Engine/Animation/Persona/Modes/Mesh/Mesh.INT.udn b/Engine/Documentation/Source/Engine/Animation/Persona/Modes/Mesh/Mesh.INT.udn
--- a/Engine/Documentation/Source/Engine/Animation/Persona/Modes/Mesh/Mesh.INT.udn
+++ b/Engine/Documentation/Source/Engine/Animation/Persona/Modes/Mesh/Mesh.INT.udn
@@ -113,10 +113,10 @@ This option allows you to reimport the asset being used and adjust the import se
 | **Update Skeleton Reference Pose** | If enabled, update the Skeleton (of the mesh being imported)'s reference pose. |
 | **Use T0As Ref Pose** | Enable this option to use frame 0 as the reference pose.  |
 | **Preserve Smoothing Groups** | If checked, triangles with non-matching smoothing groups will be physically split. |
-| **Import Meshes in Bone Hierachy** | If checked, meshes nested in bone hierarchies will be imported instead of being converted to bones.  |
+| **Import Meshes in Bone Hierarchy** | If checked, meshes nested in bone hierarchies will be imported instead of being converted to bones.  |
 | **Import Morph Targets** | If enabled, creates Unreal morph objects for the imported meshes. |
 | **Import Mesh LODs** | If enabled, creates LOD models for Unreal meshes from LODs in the imported file; if not enabled, only the base mesh from the LOD group is imported. |
-| **Normal Import Method** | Enabling this option will read the tangents(tangent,binormal,normal) from FBX files instead of generating them automatically.  |
+| **Normal Import Method** | Enabling this option will read the tangents (tangent, binormal, normal) from FBX files instead of generating them automatically.  |
 | **Source File Path** | This is the path to the resource used to construct the Static Mesh. |
 | **Source File Timestamp** | Date/Time-stamp of the file from the last import. |
 
