#UE4doc #docs: changed multiple spelling errors: aniamtion -> animation and changed a derogatory term according to the instructional guide: 'simply try'

diff --git a/Engine/Documentation/Source/Engine/Animation/Persona/SkeletonTree/SkeletonTree.INT.udn b/Engine/Documentation/Source/Engine/Animation/Persona/SkeletonTree/SkeletonTree.INT.udn
--- a/Engine/Documentation/Source/Engine/Animation/Persona/SkeletonTree/SkeletonTree.INT.udn
+++ b/Engine/Documentation/Source/Engine/Animation/Persona/SkeletonTree/SkeletonTree.INT.udn
@@ -49,7 +49,7 @@ The **Skeleton Tree** panel shows the skeletal hierarchy of the current skeletal
 
 ![](SearchLine.png)
 
-The Search line allows for quick and easy searching of an skeletal hierarchy. Simply type in some part of the name of the bone, Socket, or Preview  Object. and the list will filter to what you want to see.
+The Search line allows for quick and easy searching of a skeletal hierarchy. Type in some part of the name of the Bone, Socket, or Preview  Object. and the list will filter to what you want to see.
 
 ![](SeachFiltering.png)
 
@@ -59,7 +59,7 @@ The Search line allows for quick and easy searching of an skeletal hierarchy. Si
 
 ![](button_BoneAndSocketButtons.png)
 
-The Bone and Socket buttons are used to control visibility over various bone and Socket types. Each button opens a menu allowing you to select the types of bones or Sockets you want to see.
+The Bone and Socket buttons are used to control visibility over various bone and Socket types. Each button opens a menu allowing you to select the types of Bones or Sockets you want to see.
 
 [REGION:raw]
 ![](BonesMenu.png)
@@ -105,7 +105,7 @@ When **right-clicking** on a bone, you will see the following context menu:
 | [REGION:tablesection]Bone Translation Retargeting[/REGION]||
 |**Recursively Set Translation Retargeting Skeleton**|This recursively sets the bone and all of its children to use Translation Retargeting based off the Skeleton. See the [Translation Retargeting Controls](#translationretargetingcontrols)
 |**Recursively Set Translation Retargeting Animation**|This recursively sets the bone and all of its children to use Translation Retargeting based off the Animation. See the [Translation Retargeting Controls](#translationretargetingcontrols)
-|**Recursively Set Translation Retargeting AnimationScaled**|This recursively sets the bone and all of its children to use Translation Retargeting based off the scaled Aniamtion. See the [Translation Retargeting Controls](#translationretargetingcontrols)
+|**Recursively Set Translation Retargeting AnimationScaled**|This recursively sets the bone and all of its children to use Translation Retargeting based off the scaled Animation. See the [Translation Retargeting Controls](#translationretargetingcontrols)
 | [REGION:tablesection]LOD Bone Reduction[/REGION]||
 |**Remove Selected Bone from...**| Places the selected bone and its children on the _Bones to Remove_ list for a certain LOD level within the LOD settings. See LOD Settings for more information.|
 |**Add Remove Selected Bone to...**| Removes the selected bone from the _Bones to Remove_ list for the specified LOD level. See LOD Settings for more information.|
@@ -128,8 +128,8 @@ When **right-clicking** on a [Socket](Engine/Content\Types\SkeletalMeshes\Socket
 |**Copy Selected Sockets**| Copies the selected socket to the clipboard for pasting on to other bones.|
 | [REGION:tablesection]Bone Translation Retargeting[/REGION]||
 |**Recursively Set Translation Retargeting Skeleton**|This recursively sets the bone and all of its children to use Translation Retargeting based off the Skeleton. See the [Translation Retargeting Controls](#translationretargetingcontrols)
-|**Recursively Set Translation Retargeting Animation**|This recursively sets the bone and all of its children to use Translation Retargeting based off the Aniamtion. See the [Translation Retargeting Controls](#translationretargetingcontrols)
-|**Recursively Set Translation Retargeting AnimationScaled**|This recursively sets the bone and all of its children to use Translation Retargeting based off the scaled Aniamtion. See the [Translation Retargeting Controls](#translationretargetingcontrols)
+|**Recursively Set Translation Retargeting Animation**|This recursively sets the bone and all of its children to use Translation Retargeting based off the Animation. See the [Translation Retargeting Controls](#translationretargetingcontrols)
+|**Recursively Set Translation Retargeting AnimationScaled**|This recursively sets the bone and all of its children to use Translation Retargeting based off the scaled Animation. See the [Translation Retargeting Controls](#translationretargetingcontrols)
 | [REGION:tablesection]Attached Assets Actions[/REGION]||
 |**Add Preview Asset**| Opens a menu of all eligible assets that can be temporarily attached to a bone for previewing purposes. For instance, if you want to see what a certain piece of armor would look like attached a bone, this is the option you would use. Note, however, that this is not a permanent attachment and will only be visible in Persona. |
 
