#UE4doc #docs: Fixed a few spelling errors: depicated -> depicted and Mixmano -> Mixamo

diff --git a/Engine/Documentation/Source/Engine/Animation/RetargetingDifferentSkeletons/RetargetingDifferentSkeletons.INT.udn b/Engine/Documentation/Source/Engine/Animation/RetargetingDifferentSkeletons/RetargetingDifferentSkeletons.INT.udn
--- a/Engine/Documentation/Source/Engine/Animation/RetargetingDifferentSkeletons/RetargetingDifferentSkeletons.INT.udn
+++ b/Engine/Documentation/Source/Engine/Animation/RetargetingDifferentSkeletons/RetargetingDifferentSkeletons.INT.udn
@@ -41,7 +41,7 @@ Once you have installed the packs above, you can add them to a project via the *
 
 The Animation Starter Pack includes several animations that we will retarget for a character from the Mixamo Character Pack.
 
-Also, your version of Unreal Engine 4 may vary from the one depicated above. 
+Also, your version of Unreal Engine 4 may vary from the one depicted above. 
 [/REGION]
 
 The steps below will show you how to access the default Rig and retarget animations: 
@@ -192,7 +192,7 @@ The steps below will show you how to access the default Rig and retarget animati
 
 ## Retargeting Adjustments
 
-When setting up a Rig to share between Skeletons, you may encounter situations when one Skeleton has additional Bones that the other does not have (in the Getting Started section, if you used Skeletons specified in the example, you may have noticed extra Bones in the Mixmamo Skeleton's hands for example). This can sometimes cause issues but can be fixed by recursively setting translation retargeting to the Skeleton.
+When setting up a Rig to share between Skeletons, you may encounter situations when one Skeleton has additional Bones that the other does not have (in the Getting Started section, if you used Skeletons specified in the example, you may have noticed extra Bones in the Mixamo Skeleton's hands for example). This can sometimes cause issues but can be fixed by recursively setting translation retargeting to the Skeleton.
 
 For example, here is a shot before adjustments: 
 
