#UE4doc #docs: Updated characters that were missing from the page. Ex: '

diff --git a/Engine/Documentation/Source/Engine/Animation/RootMotion/RootMotion.INT.udn b/Engine/Documentation/Source/Engine/Animation/RootMotion/RootMotion.INT.udn
--- a/Engine/Documentation/Source/Engine/Animation/RootMotion/RootMotion.INT.udn
+++ b/Engine/Documentation/Source/Engine/Animation/RootMotion/RootMotion.INT.udn
@@ -24,7 +24,7 @@ For those new to the concept, it may not be readily apparent why proper root mot
 [REGION:warning]
 Early versions of the engine (Versions previous to 4.5) required the use of a AnimMontage (Please see the [AnimMontage documentation](https://docs.unrealengine.com/latest/INT/Engine/Animation/AnimMontage/index.html) for more information). 
 
-Since version 4.6, Root Motion is now handled per animation and can be toggled on/off in the animation�s properties in Persona. To fully utilize Root Motion, your character will need an Animation Blueprint (More information can be found in the [Animation Blueprint documentation](https://docs.unrealengine.com/latest/INT/Engine/Animation/AnimBlueprints/index.html) ). Within the characters Animation Blueprint properties, the �Root Motion� section allows you to set which type of animations will have Root Motion applied to them. Please note that at the time of writing that you still need to use an AnimMontage for your Root Motion to be applied successfully. 
+Since version 4.6, Root Motion is now handled per animation and can be toggled on/off in the animation's properties in Persona. To fully utilize Root Motion, your character will need an Animation Blueprint (More information can be found in the [Animation Blueprint documentation](https://docs.unrealengine.com/latest/INT/Engine/Animation/AnimBlueprints/index.html) ). Within the characters Animation Blueprint properties, the **Root Motion** section allows you to set which type of animations will have Root Motion applied to them. Please note that at the time of writing that you still need to use an AnimMontage for your Root Motion to be applied successfully. 
 [/REGION]
 
 [REGION:lightbox]
@@ -47,7 +47,7 @@ In order to properly use Root Motion, it is important to note that the Root Bone
 [/REGION]
 
 
-The root motion of an animation can be visualized in Persona. Simply open up any animation in which the root bone is moving, and in the viewport choose **Show > Bones**. If you have not selected �Enable Root Motion� in the animation properties, you will notice a red line that appears as the character's root moves. This illustrates the root motion of the animation.
+The root motion of an animation can be visualized in Persona. Simply open up any animation in which the root bone is moving, and in the viewport choose **Show > Bones**. If you have not selected **Enable Root Motion** in the animation properties, you will notice a red line that appears as the character's root moves. This illustrates the root motion of the animation.
 
 [REGION:fullwidth]
 ![](RootMotionBone.png)
