#UE4doc #docs: adjusted some spelling errors: addtion -> addition and grammar fixesl

diff --git a/Engine/Documentation/Source/Engine/Animation/Sequences/Curves/Curves.INT.udn b/Engine/Documentation/Source/Engine/Animation/Sequences/Curves/Curves.INT.udn
--- a/Engine/Documentation/Source/Engine/Animation/Sequences/Curves/Curves.INT.udn
+++ b/Engine/Documentation/Source/Engine/Animation/Sequences/Curves/Curves.INT.udn
@@ -8,7 +8,7 @@ Curves provide a way to change the value of a Material parameter or a Morph Targ
 [/EXCERPT:Intro]
 
 
-When you import Morph Target animations into UE4 via FBX, Morph Target Curves are automatically generated for that AnimSequence. However, they are hidden by default. This is by design; hiding away what could by many different Morph Target Curves will prevent the editor from slowing down.
+When you import Morph Target animations into UE4 via FBX, Morph Target Curves are automatically generated for that AnimSequence. However, they are hidden by default. This is by design; hiding away what could be many different Morph Target Curves will prevent the editor from slowing down.
 
 It should be noted that the ability to adjust curves for Morph Targets, while useful, is more of a convenience feature than anything else. If you are handling a very complex Morph Target animation, you will probably have a much easier time if you perform the animation in the 3D animation package from which you exported out the original Morph Target.
 
diff --git a/Engine/Documentation/Source/Engine/Animation/Sequences/Editor/AnimSequenceEditor.INT.udn b/Engine/Documentation/Source/Engine/Animation/Sequences/Editor/AnimSequenceEditor.INT.udn
--- a/Engine/Documentation/Source/Engine/Animation/Sequences/Editor/AnimSequenceEditor.INT.udn
+++ b/Engine/Documentation/Source/Engine/Animation/Sequences/Editor/AnimSequenceEditor.INT.udn
@@ -115,13 +115,13 @@ When multiple notifies are copied, by **right-clicking** in the timeline to brin
 
 ![](CopyPasteMenu.png)
 
-**Paste Multiple Relative** : Paste at the mouse cursor retaining relative spacing between notifies regardless of animation length. This means that the notify group should look like it is spaced the same as before the copy; even if the animation is shorter or longer. 
+**Paste Multiple Relative:** Paste at the mouse cursor retaining relative spacing between notifies regardless of animation length. This means that the notify group should look like it is spaced the same as before the copy; even if the animation is shorter or longer. 
 
-**Paste Multiple Absolute** : Paste at the mouse cursor retaining absolute spacing. The notifies will retain their exact time spacing meaning that they will fire at exactly the same intervals but the group is not scaled to shorter or longer animations. 
+**Paste Multiple Absolute:** Paste at the mouse cursor retaining absolute spacing. The notifies will retain their exact time spacing meaning that they will fire at exactly the same intervals but the group is not scaled to shorter or longer animations. 
 
-**Paste at Original Time** : Puts all of the nodes exactly where they were when they were copied. 
+**Paste at Original Time:** Puts all of the nodes exactly where they were when they were copied. 
 
-Addtional Notes: 
+Additional Notes: 
 * Any paste operation that would put any of the notifies off the right edge of the track will still succeed but the notifies will be clamped into the track. 
 * Pasting at the mouse cursor uses the most top-left point in the selection as a reference point for placing the nodes. 
 * Pasting a group that would pass below the bottom of the notify panel will add as many tracks as necessary to contain the group automatically. 
