#UE4doc #docs: made edits in spacing for a picture, spelling centre -> center and restructed the sentence: So for practical purposes, if you are modifying something every frame, that should go in the Tick Blueprint.

diff --git a/Engine/Documentation/Source/Engine/Animation/Sequences/Notifies/AnimNotifies.INT.udn b/Engine/Documentation/Source/Engine/Animation/Sequences/Notifies/AnimNotifies.INT.udn
--- a/Engine/Documentation/Source/Engine/Animation/Sequences/Notifies/AnimNotifies.INT.udn
+++ b/Engine/Documentation/Source/Engine/Animation/Sequences/Notifies/AnimNotifies.INT.udn
@@ -165,7 +165,7 @@ There is currently no way to retrieve the length of the AnimNotifyState within a
 
 When the **AnimNotifyState** is triggered, first the **Start** Blueprint will be called, and then the **Tick** Blueprint is called on the same frame update. Every subsequent `tick`, only the **Tick** Blueprint will be called. On the final `tick` update, the **Tick** Blueprint will be called followed by the **End** Blueprint.
 
-So for practical purposes, if you are modifying something every frame, that should go in the **Tick** Blueprint. If you need to set variables, flags, or properties, to some value before the `ticks` hit, it should go in the **Start** Blueprint. And finally, the **End** Blueprint can be used to change variables, flags, or properties after the final `tick` hits.
+If you are modifying something every frame, it should go in the **Tick** Blueprint. If you need to set variables, flags, or properties, to some value before the `ticks` hit, it should go in the **Start** Blueprint. And finally, the **End** Blueprint can be used to change variables, flags, or properties after the final `tick` hits.
 
 
 ## Native Notify States 
@@ -247,7 +247,7 @@ This method is based on providing a width scale value. A width scale of 1.0 will
 
 ![](WidthScaleMode.png)
 
-_The left image depicts **From Centre** which scales outwards from centre, the middle image shows **From First Socket** which scales away from the first socket and the right image demonstrates **From Second Socket** which scales away from the second socket._ 
+_The left image depicts **From Centre** which scales outwards from center, the middle image shows **From First Socket** which scales away from the first socket and the right image demonstrates **From Second Socket** which scales away from the second socket._ 
 
 There are also some debug rendering options available:
 
@@ -264,7 +264,8 @@ The only restriction with this is that the **Particle System Component** must be
 
 1. Create an **Emitter** component (or use an existing one). 
 1. Call the **BeginTrails** function on the component. This will start all trail emitters in the component. You also need to provide the trail's other properties here. 
- ![](TrailBlueprint2.png)
+
+    ![](TrailBlueprint2.png)
 1. Call the **EndTrails** function on the component. 
  ![](TrailBlueprint3.png)
 1. You can also alter the trail's properties while it is running (the example below alters width each tick using the game time).  
