made landing page match sidebar organization

 - updated titles on pages to match sidebar

diff --git a/Engine/Documentation/Source/Engine/Engine.INT.udn b/Engine/Documentation/Source/Engine/Engine.INT.udn
--- a/Engine/Documentation/Source/Engine/Engine.INT.udn
+++ b/Engine/Documentation/Source/Engine/Engine.INT.udn
@@ -1,5 +1,5 @@
 Availability:Public
-Title: Unreal Editor Manual
+Title: Engine Features
 Crumbs: %ROOT%
 Description:The complete guide to creating games with Unreal Engine.
 Navigation:topic
diff --git a/Engine/Documentation/Source/Engine/UI/Editor.INT.udn b/Engine/Documentation/Source/Engine/UI/Editor.INT.udn
--- a/Engine/Documentation/Source/Engine/UI/Editor.INT.udn
+++ b/Engine/Documentation/Source/Engine/UI/Editor.INT.udn
@@ -1,7 +1,7 @@
 Availability:Public
-Title: Unreal Editor Interface
+Title: Unreal Editor Manual
 Crumbs:%ROOT%, Engine
-Description:A breakdown of the various parts of the Unreal Editor user interface, along with key principles on how the interface works.
+Description:A breakdown of the Unreal Editor user interface, along with key principles on how the interface works.
 Navigation:topic
 
 [EXCERPT:LevelEditor]
@@ -32,7 +32,7 @@ Navigation:topic
 	![](%ROOT%/editor_icon.png)(convert:false)
 [/PARAM]
 [PARAM:title]
-	Editor Reference
+	%Engine/UI:title%
 [/PARAM]
 [PARAM:description]
 	%Engine/UI:description%
diff --git a/Engine/Documentation/Source/UE4Home.INT.udn b/Engine/Documentation/Source/UE4Home.INT.udn
--- a/Engine/Documentation/Source/UE4Home.INT.udn
+++ b/Engine/Documentation/Source/UE4Home.INT.udn
@@ -66,9 +66,11 @@ Navigation:home
 [REGION:frontpage]
 	[REGION:features third]
 	%GettingStarted:role%
+	%Engine/UI:role%
 	%Engine:role%
-	%Programming:topicrole%
+	%Gameplay:role%
 	%Engine/Blueprints:role%
+	%Programming:topicrole%
 	%Platforms:role%
 	%Resources:role%
 	[/REGION]
