#UE4Doc #docs Added SkillLevel and Version to Level Designer Quick Start.

diff --git a/Engine/Documentation/Source/Engine/QuickStart/EditorQuickStart.INT.udn b/Engine/Documentation/Source/Engine/QuickStart/EditorQuickStart.INT.udn
--- a/Engine/Documentation/Source/Engine/QuickStart/EditorQuickStart.INT.udn
+++ b/Engine/Documentation/Source/Engine/QuickStart/EditorQuickStart.INT.udn
@@ -2,6 +2,8 @@ Availability:Public
 Title:Level Designer Quick Start
 Crumbs: %ROOT%, GettingStarted
 Description: Getting up and running with the basics of the Unreal Editor.
+SkillLevel: Beginner
+Version: 4.8
 
 [VAR:Topic]
 [OBJECT:Topic]
