#UE4Doc #docs Added SkillLevel and Version to the Programming Quick Start.

diff --git a/Engine/Documentation/Source/Programming/QuickStart/QuickStart.INT.udn b/Engine/Documentation/Source/Programming/QuickStart/QuickStart.INT.udn
--- a/Engine/Documentation/Source/Programming/QuickStart/QuickStart.INT.udn
+++ b/Engine/Documentation/Source/Programming/QuickStart/QuickStart.INT.udn
@@ -1,7 +1,9 @@
 Availability: Public
 Title:Programming Quick Start
 Crumbs:%ROOT%, Programming
-Description:Create your first code project, and add a new C++ class.
+Description:Create your first code project and add a new C++ class.
+SkillLevel: Beginner
+Version: 4.8
 
 [VAR:Topic]
 [OBJECT:Topic]
