#UE4doc #docs: fixed minor typos

diff --git a/Engine/Documentation/Source/Engine/Animation/SyncGroups/SyncGroups.INT.udn b/Engine/Documentation/Source/Engine/Animation/SyncGroups/SyncGroups.INT.udn
--- a/Engine/Documentation/Source/Engine/Animation/SyncGroups/SyncGroups.INT.udn
+++ b/Engine/Documentation/Source/Engine/Animation/SyncGroups/SyncGroups.INT.udn
@@ -28,7 +28,7 @@ Leaders and followers are designated simply by which animation has the greater b
 |---|---|
 |[INCLUDE:#runanimation]|Run animation, original size, set to a Sync Group named _WalkRun_.  |
 |[INCLUDE:#walkanimation]| Walk animation, original size, also set to a Sync Group named _WalkRun_. |
-|[INCLUDE:#walkanimationscaled] |In this instance, the blend weight for the run animation would be greater than the walk animation. Since they are in the same Sync Group, the walk scales down to match the length of the run aniamtion.|
+|[INCLUDE:#walkanimationscaled] |In this instance, the blend weight for the run animation would be greater than the walk animation. Since they are in the same Sync Group, the walk scales down to match the length of the run animation.|
 |[INCLUDE:#runanimationscaled]|Here, the weighting has switched such that the walking animation is receiving a greater blend weight. This causes the run animation to scale up to match the length of the walk. | 
 
 <!---
