#UE4doc #docs: grammar change: affect -> effect

diff --git a/Engine/Documentation/Source/Engine/Audio/SoundActors/SoundActors.INT.udn b/Engine/Documentation/Source/Engine/Audio/SoundActors/SoundActors.INT.udn
--- a/Engine/Documentation/Source/Engine/Audio/SoundActors/SoundActors.INT.udn
+++ b/Engine/Documentation/Source/Engine/Audio/SoundActors/SoundActors.INT.udn
@@ -149,7 +149,7 @@ The **Audio** component does have a unique event which can be used. This **OnAud
 For an added level of control, you can include **Audio Volumes** in conjunction with **Ambient Sound Actors**. For more information on how to use an **Audio Volume**, refer to the Audio Volumes section of the [Volumes](Engine/Actors/Volumes) page. 
 
 [REGION:tip]
-Making changes to an **Audio Volume** will not work in realtime inside the editor. You will need to rebuild the geometry of the level that contains the edited volume in order for your changes to take affect. 
+Making changes to an **Audio Volume** will not work in realtime inside the editor. You will need to rebuild the geometry of the level that contains the edited volume in order for your changes to take effect. 
 [/REGION]
 
 
