#UE4doc #docs: fixed minor spelling errors: envolope and indefinietly

diff --git a/Engine/Documentation/Source/Engine/Audio/SoundCues/NodeReference/SoundCueReference.INT.udn b/Engine/Documentation/Source/Engine/Audio/SoundCues/NodeReference/SoundCueReference.INT.udn
--- a/Engine/Documentation/Source/Engine/Audio/SoundCues/NodeReference/SoundCueReference.INT.udn
+++ b/Engine/Documentation/Source/Engine/Audio/SoundCues/NodeReference/SoundCueReference.INT.udn
@@ -205,9 +205,9 @@ You can use the properties in the **Modulation** category to randomize predefine
 | [REGION:tablesection]Looping[/REGION] ||
 | **Loop Start** | The time in seconds where the envelope's loop begins. |
 | **Loop End** | The time in seconds where the envelope's loop ends. |
-| **Duration After Loop** | The time in seconds it takes the evelope to fade out after the last loop is completed. |
+| **Duration After Loop** | The time in seconds it takes the envelope to fade out after the last loop is completed. |
 | **Loop Count** | The number of times the envelope should loop if looping is enabled and the envelope is not set to loop indefinitely. |
-| **Loop Indefinitely** | If enabled, the envelope will continue to loop indefenitely regardless of the **Loop Count** value. |
+| **Loop Indefinitely** | If enabled, the envelope will continue to loop indefinitely regardless of the **Loop Count** value. |
 | **Loop** | If enabled, the envelope will loop using the loop settings. |
 | [REGION:tablesection]Envelope[/REGION] ||
 | **Volume Curve** | The distribution defining the volume envelope.  |
