#UE4doc #docs: spelling error fix: edtior -> editor

diff --git a/Engine/Documentation/Source/Engine/Audio/WAV/WAV.INT.udn b/Engine/Documentation/Source/Engine/Audio/WAV/WAV.INT.udn
--- a/Engine/Documentation/Source/Engine/Audio/WAV/WAV.INT.udn
+++ b/Engine/Documentation/Source/Engine/Audio/WAV/WAV.INT.udn
@@ -55,7 +55,7 @@ Description:The basic requirements for getting sound files into Unreal Engine 4.
 		![](WAV_TitleIcon.png)(w:200 convert:false) 
 	[/PARAM]
 	[PARAM:description]
-	Unreal Engine 4's audio system provides tools and features to mold the sounds in the game to give them the desired feel. This all starts with your audio source files. Either as mono or stereo files, they can be quickly imported into Unreal Edtior 4 and begin breathing more life into your creations. Below are the basic requirements for your audio content and a quick import guide.
+	Unreal Engine 4's audio system provides tools and features to mold the sounds in the game to give them the desired feel. This all starts with your audio source files. Either as mono or stereo files, they can be quickly imported into Unreal Editor 4 and begin breathing more life into your creations. Below are the basic requirements for your audio content and a quick import guide.
 	[/PARAM]
 [/OBJECT]
 
