corrected documentation on postprocess volumes

diff --git a/Engine/Documentation/Source/Shared/Types/APostProcessVolume/APostProcessVolume.INT.udn b/Engine/Documentation/Source/Shared/Types/APostProcessVolume/APostProcessVolume.INT.udn
--- a/Engine/Documentation/Source/Shared/Types/APostProcessVolume/APostProcessVolume.INT.udn
+++ b/Engine/Documentation/Source/Shared/Types/APostProcessVolume/APostProcessVolume.INT.udn
@@ -5,7 +5,9 @@ Description:Overview of the PostProcessVolume class
 
 [EXCERPT:APostProcessVolume]
 
-The PostProcessVolume is a special type of volume that can be added to a level. Currently Unreal Engine 4 does not make use of post processing chains, making PostProcessVolumes the only way to manipulate post processing parameters.
+The PostProcessVolume is a special type of volume that can be added to a level. Unreal Engine 3 had the post processing chains, but in Unreal Engine 4 all this got replaced by
+using the PostProcessSettings which can be conveniently adjusted using one or multiple PostProcessVolume. By using multiple volumes (or other objects) with priority and weight
+teh settings can be softly blended depending on the camera position and manipulating weights.
 
 [VAR:ToolTipFullLink]
 Engine/Rendering/PostProcessEffects
