improved volume blending docs on postprocess materials

diff --git a/Engine/Documentation/Source/Engine/Rendering/PostProcessEffects/PostProcessMaterials/PostProcessMaterials.INT.udn b/Engine/Documentation/Source/Engine/Rendering/PostProcessEffects/PostProcessMaterials/PostProcessMaterials.INT.udn
--- a/Engine/Documentation/Source/Engine/Rendering/PostProcessEffects/PostProcessMaterials/PostProcessMaterials.INT.udn
+++ b/Engine/Documentation/Source/Engine/Rendering/PostProcessEffects/PostProcessMaterials/PostProcessMaterials.INT.udn
@@ -83,6 +83,28 @@ Depending on the camera position the setting of one volume are used:
 
 With camera movement, you can perceive soft linear transition between the two effect setting.
 
+The following shows a top down view of a level that has how two volumes. The large volume has a red material instance and the small volume has a green material instance specific as blendable.
+The smaller volume has a higher priority. The material params are getting blended depending on the camera position.
+The fuzzy borders are defined by the BlendRadius property specified in the volumes and extend the volume shape.
+
+With the correct setup all blending happens as expected.
+
+[REGION:imagetable]
+| ![](VolumeBlendBad.png) | ![](VolumeBlendGood.png) |
+| ------ | ------ |
+| Bad setup | Good setup |
+[/REGION]
+
+The difference between the both setup is the default value specified on the material (scalar or vector) parameters.
+The good setup has values that make the pass appear to not have any effect (e.g. multiply with white or lerp with 0).
+
+**In both setup we see this:** When the camera is outside the influence of either volume the postprocess pass is not getting rendered (visualize by the grey grid).
+If we are fully inside either of the volumes we also see the correct blending.
+
+**The bad setup :** When the camera enters the influence radius we see a hard transition but because is uses the wrongly specified default parameters.
+
+**The good setup :** The transition of entering the camera the influence radius is well hidden and we see a smooth transition to the volume colors.
+
 [REGION:note]
 All material instances properties are blended, no matter if the property checkbox is checked or not (in that case it blends the
 property from the parent). This is different from the post processing settings where an unchecked property is not having any effect.
