#UE4doc #docs: made minor grammar changes: any where -> anywhere, example word . -> example word.

diff --git a/Engine/Documentation/Source/Engine/Basics/Screenshots/Screenshots.INT.udn b/Engine/Documentation/Source/Engine/Basics/Screenshots/Screenshots.INT.udn
--- a/Engine/Documentation/Source/Engine/Basics/Screenshots/Screenshots.INT.udn
+++ b/Engine/Documentation/Source/Engine/Basics/Screenshots/Screenshots.INT.udn
@@ -17,7 +17,7 @@ By default any screenshots will be saved to the following location in your proje
 
 **Saved\Screenshots\Windows**
 
-When ever you do take a screenshot, the following message will be displayed in the lower Right - Hand corner of the screen. 
+When you do take a screenshot, the following message will be displayed in the lower Right - Hand corner of the screen. 
 
 ![](HRSS_Save_Location.png)
 
@@ -56,7 +56,7 @@ Similarly, the command below will take a screenshot at 3840 wide by 2160 high:
 
 ![](HRSS_The_Tool.png)
 
-The High Resolution Screenshot Tool(HRSST) allows you to specify a specific region of your viewport to take a screenshot with. It also gives you the ability to adjust the size of the Screenshot you are going to take as well as output the results of the GBuffer. 
+The High Resolution Screenshot Tool (HRSST) allows you to specify a specific region of your viewport to take a screenshot with. It also gives you the ability to adjust the size of the Screenshot you are going to take as well as output the results of the GBuffer. 
 
 ### Accessing the High Resolution Screenshot Tool 
 You can access the HRSST by first clicking on the **View Port Options** and then selecting the **High Resolution Screenshot** option. 
@@ -68,7 +68,7 @@ Once you have selected the **High Resolution Screenshot**, the HRSST will be dis
 ![](HRSS_HRSST.png)
 
 [region:tip]
-You can move the HRSST any where you want to in side of UE4. However be careful as sometime the HRSST can become hidden by other floating windows like the **Content Browser**. 
+You can move the HRSST anywhere you want to in side of UE4. However be careful as sometime the HRSST can become hidden by other floating windows like the **Content Browser**. 
 [/region]
 
 
@@ -81,7 +81,7 @@ You can move the HRSST any where you want to in side of UE4. However be careful
 | 5: | This allows you to specify a region within the viewport that you wish to capture. You define a region by first clicking on the tool and draw a region to capture by holding down the **Left Mouse Button** and dragging in the Editor viewport.|
 | 6: | Clicking on this will clear any previous capture regions that you might have created.|
 | 7: | Clicking on this will make the capture region the entire editor window.|
-| 8: | Clicking on this this will take the screenshot .|
+| 8: | Clicking on this this will take the screenshot.|
 
 [region:note]
 Due to the high system requirements of a high resolution screenshot, very large multipliers might cause the graphics driver to be unresponsive and possibly crash. In these circumstances, please try using a lower number. 
@@ -133,7 +133,7 @@ Using the HRSST tool is something that can be done in the few following steps.
 
 
 ### Buffer Visualizations
-The High Resolution Screenshot tool(HRSST) supports the ability to render out each on of the buffers in the G-Buffer.  After clicking the capture button, the HRSST  writes out the following buffers.
+The High Resolution Screenshot tool (HRSST) supports the ability to render out each one of the buffers in the G-Buffer.  After clicking the capture button, the HRSST  writes out the following buffers.
 
 * **DiffuseColor**
 * **LightingModel**
