#UE4doc #docs: changed spacing at end of sentences: delete . -> delete. Fixed missing characters as well.

diff --git a/Engine/Documentation/Source/Engine/Basics/SourceControl/InEditor/SourceControlInsideUE4.INT.udn b/Engine/Documentation/Source/Engine/Basics/SourceControl/InEditor/SourceControlInsideUE4.INT.udn
--- a/Engine/Documentation/Source/Engine/Basics/SourceControl/InEditor/SourceControlInsideUE4.INT.udn
+++ b/Engine/Documentation/Source/Engine/Basics/SourceControl/InEditor/SourceControlInsideUE4.INT.udn
@@ -98,7 +98,7 @@ If you select a file in the **Content Browser** and press **Delete** (or select
 | Number | Description|
 | -------| ---------- |
 | 1:     | This is the asset(s) that you are going to delete.|
-| 2:     | This shows asset(s) that are referencing the assets you are about to delete .|
+| 2:     | This shows asset(s) that are referencing the assets you are about to delete.|
 | 3:     | Here you specify what assets you want other assets that reference this asset(s) to reference when this asset is deleted.|
 | 4:     | Here you can force delete an asset. However be very careful doing this as it can have a very negative impact on your project. You should only use this option as a last resort if nothing else works.|
 | 5:     | This will cancel the delete operation.|
@@ -108,7 +108,7 @@ that is created when deleting assets, you can **Right-Click** on a **folder** in
 in your project and permanently delete the redirectors, marking the files for delete in source control. Presently, to commit these deletions, you need to use an external tool such as P4V or TortoiseSVN.
 
 ### Level Files
-Levels can be managed in source control just like any other asset via the **Content Browser**. However levels can also be worked with via the **Level Browser�s** menu, 
+Levels can be managed in source control just like any other asset via the **Content Browser**. However levels can also be worked with via the **Level Browser's** menu, 
 which contains options for source control in the **Source Control** sub-menu.
 ![](SC_UE4_Level_Broswer.png)
 
