#UE4doc #docs: fixed small grammar errors and changed missing characters for the markup language .

diff --git a/Engine/Documentation/Source/Engine/Basics/SourceControl/Perforce/P4SourceControl.INT.udn b/Engine/Documentation/Source/Engine/Basics/SourceControl/Perforce/P4SourceControl.INT.udn
--- a/Engine/Documentation/Source/Engine/Basics/SourceControl/Perforce/P4SourceControl.INT.udn
+++ b/Engine/Documentation/Source/Engine/Basics/SourceControl/Perforce/P4SourceControl.INT.udn
@@ -12,7 +12,7 @@ Related: Programming/Basics/ConfigurationFiles
 [TOC(start:2)]
 
 ## Overview 
-Unreal Engine 4 comes with an integrated method that allow people to work together on their projects using version control software called Perforce. 
+Unreal Engine 4 comes with an integrated method that allows people to work together on their projects using version control software called Perforce. 
 Version control allows users on your team to share assets and code with one another, as well as providing a backup and history of changes so any files can 
 be rolled back to an earlier version if something went wrong with the file.
 
@@ -67,7 +67,7 @@ Open the typemap by entering the following in the command prompt.
 
 	p4 -P YourPassword typemap
 
-You will then see a text window open � this is your server�s typemap. The typemap below is an example that we have used in the past:
+You will then see a text window open; this is your server's typemap. The typemap below is an example that we have used in the past:
 
 	# Perforce File Type Mapping Specifications.
 	#
@@ -107,7 +107,7 @@ Once edited, save the file by pressing **Ctrl+S** on your keyboard and then exit
 
 ### Server on your Local Network
 You can set up a server on your local network for people to use to collaborate with one another. 
-To do this, follow the steps above on the server and connect using the server�s name or IP from your clients.
+To do this, follow the steps above on the server and connect using the server's name or IP from your clients.
 
 ### Cloud Provider
 There are various cloud providers that allow you to host Perforce servers off-site and to collaborate remotely. 
@@ -132,7 +132,7 @@ You should initially add only the following directories of your project:
 
 Folders and files highlighted in **Green** should be added to your depot. The Binaries folder in **Yellow** is optional, depending on how your team wants to work together. It is not required if every team member is going to compile their 
 editor/game themselves, but should be included in the **initial add** if any team members will just be grabbing a version from the depot to create content. 
-Note that the server typemap setup above should mean that coders can compile their binaries without them being set to be **�Read-Only�** by Perforce.
+Note that the server typemap setup above should mean that coders can compile their binaries without them being set to be **Read-Only** by Perforce.
 
 ### Adding Files to Perforce
 
