#UE4doc #docs: fixed missing characters from the markup language. Also made minor grammar fixes

diff --git a/Engine/Documentation/Source/Engine/Basics/SourceControl/SVN/SVNSourceControl.INT.udn b/Engine/Documentation/Source/Engine/Basics/SourceControl/SVN/SVNSourceControl.INT.udn
--- a/Engine/Documentation/Source/Engine/Basics/SourceControl/SVN/SVNSourceControl.INT.udn
+++ b/Engine/Documentation/Source/Engine/Basics/SourceControl/SVN/SVNSourceControl.INT.udn
@@ -12,13 +12,13 @@ Related: Programming/Basics/ConfigurationFiles
 [TOC(start:2)]
 
 ## Overview 
-Unreal Engine 4 comes with an integrated method that allow people to work together on their projects using version control software called Subversion or SVN for short. 
+Unreal Engine 4 comes with an integrated method that allows people to work together on their projects using version control software called Subversion or SVN for short. 
 Version control allows users on your team to share assets and code with one another, as well as providing a backup and history of changes so any files can 
 be rolled back to an earlier version if something went wrong with the file.
 
 
 ## Using Subversion (SVN) with Unreal Engine 4
-If Perforce is not an option for you, then you can use Subversion(SVN). SVN is similar to Perforce in how users handle data. 
+If Perforce is not an option for you, then you can use Subversion (SVN). SVN is similar to Perforce in how users handle data. 
 The following guide will walk you through how to setup a SVN source control server for your UE4 projects. 
 
 ### Setting Up a SVN Server
@@ -53,7 +53,7 @@ Once created, you should see your repository in the tree view.
 
 ### Server on your Local Network
 You can set up a server on your local network for people to use to collaborate with one another. 
-To do this, follow the steps above on the server and connect using the server�s URL from your clients.
+To do this, follow the steps above on the server and connect using the server's URL from your clients.
 
 ### Secure Access
 Many users like to use a secure connection to access their repositories. These connections have URLs starting with **https://** and have a number of things to be aware of when setting them up.
@@ -94,7 +94,7 @@ If you are trying to edit a project that someone else has already created, you c
 [/region]
 
 
-To create your working copy, you should **Check Out** the (assumed empty) root of your repository. To do this, **Right-Click** on your project�s folder in Windows Explorer and select **SVN Checkout**.
+To create your working copy, you should **Check Out** the (assumed empty) root of your repository. To do this, **Right-Click** on your project's folder in Windows Explorer and select **SVN Checkout**.
 
 ![](SC_SVN_Checkout.png)
 
