#UE4doc #docs: changed a missing character error.

diff --git a/Engine/Documentation/Source/Engine/Basics/SourceControl/SVN/SVNSourceControl.INT.udn b/Engine/Documentation/Source/Engine/Basics/SourceControl/SVN/SVNSourceControl.INT.udn
--- a/Engine/Documentation/Source/Engine/Basics/SourceControl/SVN/SVNSourceControl.INT.udn
+++ b/Engine/Documentation/Source/Engine/Basics/SourceControl/SVN/SVNSourceControl.INT.udn
@@ -138,7 +138,7 @@ Because of how the Mac stores security information such as passwords, you may ne
 You may need to perform these steps if you see errors in your logs like:
 
 	E175002: Unexpected HTTP status 400 'Bad Request' on '/projectname/'
-	E175002: Unexpected HTTP status 405 'Method Not Allowed' on '/projectname/�
+	E175002: Unexpected HTTP status 405 'Method Not Allowed' on '/projectname/'
 	E215004: Authentication failed and interactive prompting is disabled; see the --force-interactive option
 	E215004: No more credentials or we tried too many times.
 
