#UE4doc #docs: fixed minor spelling and grammar changes: aviable -> available, there is -> there are

diff --git a/Engine/Documentation/Source/Engine/Blueprints/BestPractices/BestPractices.INT.udn b/Engine/Documentation/Source/Engine/Blueprints/BestPractices/BestPractices.INT.udn
--- a/Engine/Documentation/Source/Engine/Blueprints/BestPractices/BestPractices.INT.udn
+++ b/Engine/Documentation/Source/Engine/Blueprints/BestPractices/BestPractices.INT.udn
@@ -68,7 +68,7 @@ between Blueprints.
 -	Macros take the nodes from the macro graph, and actually replace the macro node with a copy of all those nodes.  Basically, when the Blueprint is compiled, the macro copies all the
  graph nodes, and pastes them in where the macro node is.
 
-Because of how they work under the hood, that means there's a few differences in what you can do:
+Because of how they work under the hood, that means there are a few differences in what you can do:
 
 -	Macros can use any node for the class they're scoped to (the class you pick when you create a new Macro Blueprint, or if it's a local macro, the class the macro is defined in).  
 This means that you have a little more versatility in the nodes you can put in a macro.  
@@ -105,7 +105,7 @@ all Blueprints, even unloaded ones, if you uncheck the **Find in Current Bluepri
 	![](SearchResults.png)
 
 * The context menu only shows nodes that you can connect to the pin you dragged from (e.g. when you drag from a light, it shows you functions relevant to the light).  
-When you're trying to discover all the possible functionality available to act on a variable or component, just browse through the context menu and see what's avaiable.
+When you're trying to discover all the possible functionality available to act on a variable or component, just browse through the context menu and see what is available.
 
 	![](ContextMenu.png)
 
