#UE4doc #docs: Changed title to match new style guide setup for How To's

diff --git a/Engine/Documentation/Source/Engine/Blueprints/BP_HowTo/BlueprintHowTo.INT.udn b/Engine/Documentation/Source/Engine/Blueprints/BP_HowTo/BlueprintHowTo.INT.udn
--- a/Engine/Documentation/Source/Engine/Blueprints/BP_HowTo/BlueprintHowTo.INT.udn
+++ b/Engine/Documentation/Source/Engine/Blueprints/BP_HowTo/BlueprintHowTo.INT.udn
@@ -1,5 +1,5 @@
 Availability: Public
-Title:Blueprint - How To...
+Title:Blueprints - How To's
 Crumbs: %ROOT%, Engine, Engine/Blueprints
 Description:The Blueprint How To page provides several short step-by-step guides for working with Blueprints. 
 Related: Engine/Blueprints
