#UE4doc #docs: Made a minor grammar fix: a -> an

diff --git a/Engine/Documentation/Source/Engine/Blueprints/BP_HowTo/ActorReference/ActorReference.INT.udn b/Engine/Documentation/Source/Engine/Blueprints/BP_HowTo/ActorReference/ActorReference.INT.udn
--- a/Engine/Documentation/Source/Engine/Blueprints/BP_HowTo/ActorReference/ActorReference.INT.udn
+++ b/Engine/Documentation/Source/Engine/Blueprints/BP_HowTo/ActorReference/ActorReference.INT.udn
@@ -116,7 +116,7 @@ _For this example, we are using a new **Blueprint Third Person Template** with *
 
 	This creates a **Get** node which will get whatever is assigned to the variable which we have set to an Actor in our level. 
 
-1. Drag off the **TargetActor** node, search for and add a **Add Actor Local Offset** node. 
+1. Drag off the **TargetActor** node, search for and add an **Add Actor Local Offset** node. 
 
 	![](16AddOffSet.png)
 
