#UE4 Docs: Fixing jira UEDOC-731 by adding new info and pics to the Mobile Previewer page.

diff --git a/Engine/Documentation/Source/Platforms/Mobile/Previewer/MobilePreviewer.INT.udn b/Engine/Documentation/Source/Platforms/Mobile/Previewer/MobilePreviewer.INT.udn
--- a/Engine/Documentation/Source/Platforms/Mobile/Previewer/MobilePreviewer.INT.udn
+++ b/Engine/Documentation/Source/Platforms/Mobile/Previewer/MobilePreviewer.INT.udn
@@ -1,17 +1,47 @@
 Availability:Public
 Title:Mobile Previewer
 Crumbs:%ROOT%, Platforms, Platforms/Mobile
-Description:Previewing games on mobile platforms.
+Description: Previewing games on mobile platforms will look in the UE4 editor.
 Platform:Mobile
 
-There is currently no in-editor preview to simulate mobile rendering.
-To see what your project could look like on mobile, you have to launch a standalone game with the **Enable Mobile Preview** option set. 
-To enable the Mobile preview go to the **Main Tool Bar** and click on the arrow that is to the right of the **Play** button to expose the different play modes.
-Click on the one that says **Mobile Preview** and your project will then be launched in a new window that simulates what your project should look like on a mobile device.
+![](T_Rendering_Level_Examples.png)
 
-![](T_Activate_Mobile_Preview.png)
+[region:caption]
+1:Mobile / HTML5 - Open GL ES2,
+2:Shader Model 4.0 - DX10/ OpenGL 3.3+,
+3:Shader Model 5.0 - DX11/ OpenGL 4.3+.
+[/region]
+
+With the release of Unreal Engine (UE4) 4.8 you can now preview what your scene will look like on mobile devices directly in the UE4 editor viewport.
+When enabling different preview levels the Materials in your scene will be re-compiled to emulate the look and feature set of the render you selected. 
+You can also switch back and forth between rendering levels seamlessly inside of UE4 without having to re-start the editor.
+To change the which rendering level is used in the current UE4 Editor session, you will need to do the following. 
+
+[region:tip]
+While you do not need to re-start the Editor for the changes to take effect the first time you change the Preview Rendering Level the editor will need time to re-compile shaders.
+[/region]
+
+1. From the  Main Toolbar use the mouse and click on the **Settings** icon.
+
+1. Under the **Scalability** section move the mouse over the **Preview Rendering Level** section to expose the different rendering levels options. 
+
+1. Selecting the rendering option you would like to preview in editor by clicking on it.
+
+1. Once a rendering level has been selected the Materials in the viewport will automatically update to reflect the new rendering method. 
+![](T_PRL.png)
 
 [region:note]
 Please note that the Mobile previewer is intended to match mobile devices as closely as possible but is not indicative of what your project is going to look like on the target device.
 You should always make sure to fully test your project on your target device and only use the mobile previewer to see if your work is headed in the right direction.
 [/region]
+
+Versions of UE4 prior to 4.8 do not have the ability to change the editor viewport rendering method to simulate mobile rendering.
+To see what your project could look like on mobile, you have to launch a standalone game using the **Enable Mobile Preview** option. 
+To enable the Mobile Preview mode you will need to do the following.
+
+1. On the **Main Tool Bar** click on the triangle that is to the right of the **Play** button to expose the different play modes.
+
+1. Click on the one that says **Mobile Preview** and your project will then be launched in a new window that simulates what your project should look like on a mobile device.
+
+	![](T_Activate_Mobile_Preview.png)
+
