#UE4doc #docs: Fixed Multiple Spelling Errors: Carater-> Character, Bluperint -> Blueprint, etc. Altered color of image indicators to match instructional guide: green -> yellow

diff --git a/Engine/Documentation/Source/Engine/Blueprints/BP_HowTo/DirectBlueprintComs/DirectBlueprintComs.INT.udn b/Engine/Documentation/Source/Engine/Blueprints/BP_HowTo/DirectBlueprintComs/DirectBlueprintComs.INT.udn
--- a/Engine/Documentation/Source/Engine/Blueprints/BP_HowTo/DirectBlueprintComs/DirectBlueprintComs.INT.udn
+++ b/Engine/Documentation/Source/Engine/Blueprints/BP_HowTo/DirectBlueprintComs/DirectBlueprintComs.INT.udn
@@ -28,7 +28,7 @@ For this guide we are using the **Blueprint Third Person** template with **Start
 
 	![](DirectCom2.png)
 
-	This will ensure that only the **Pawns** and by extension **Charaters** can pass through the box. 
+	This will ensure that only the **Pawns** and by extension **Characters** can pass through the box. 
 
 1. In the **MyBlueprint** window, click the **Add Variable** button and name the variable **TargetBlueprint**. 
 
@@ -72,15 +72,15 @@ For this guide we are using the **Blueprint Third Person** template with **Start
 
 	![](DirectCom7.png)
 
-	Here we are specifying the instance of the Blueprint_Effect_Fire Bluperint in the level that we want to communicate with. If we have placed several instances of the Blueprint_Effect_Fire Blueprint in the level, each of the instances would be displayed inside the drop-down menu allowing us to indicate which instance is the **Target Blueprint** to communicate with. 
+	Here we are specifying the instance of the Blueprint_Effect_Fire Blueprint in the level that we want to communicate with. If we have placed several instances of the Blueprint_Effect_Fire Blueprint in the level, each of the instances would be displayed inside the drop-down menu allowing us to indicate which instance is the **Target Blueprint** to communicate with. 
 
 	[REGION:tip]
 	Clicking the **Eye Dropper** icon allows you to select the instance of your **Target Bluperint** inside the level instead of using the drop-down menu. Remember, you can only select Blueprints to target based on the variable type you have specified. 
 	[/REGION]
 
-1. Click the **Play** button to play in the editor and test running through the box with the charcter. 
+1. Click the **Play** button to play in the editor and test running through the box with the character. 
 
-You will see that once the character enters the box, the fire effect will be activated inside the level. You can use **Direct Bluperint Communcation*** to not only change the properties of a Blueprint from inside another Blueprint, but you can also pass the values of Variables, call Functions or Events, or anything else that the Blueprint would otherwise normally be able to do. 
+You will see that once the character enters the box, the fire effect will be activated inside the level. You can use **Direct Blueprint Communication** to not only change the properties of a Blueprint from inside another Blueprint, but you can also pass the values of Variables, call Functions or Events, or anything else that the Blueprint would otherwise normally be able to do. 
 
 
 
