#UE4 Docs: Working on Jira UEDOC-691 - Adding a section about mobile content scale factor.

diff --git a/Engine/Documentation/Source/Platforms/Mobile/Performance/MobilePerformance.INT.udn b/Engine/Documentation/Source/Platforms/Mobile/Performance/MobilePerformance.INT.udn
--- a/Engine/Documentation/Source/Platforms/Mobile/Performance/MobilePerformance.INT.udn
+++ b/Engine/Documentation/Source/Platforms/Mobile/Performance/MobilePerformance.INT.udn
@@ -16,7 +16,7 @@ on mobile platforms. These are listed from most severe to least.
 * Make sure precomputed visibility is setup properly.  To do this, you need to place PrecomputedVisibilityVolumes around anywhere that the player can walk, and build lighting.  You need to make sure the same P level is used when building lighting and when running the game (aka do not build the sublevel by itself).  You can verify that it is working by entering `stat initviews` on device or in previewer and make sure the `Statically occluded primitives` is > 0.  Use `r.ShowPrecomputedVisibilityCells 1` to visualize the cells in the editor.
 * Use masked and translucent Materials sparingly, and only in places where they cover a small part of the screen.  iOS devices are very optimal in shading opaque surfaces, as they will only shade each pixel once, most of the time.  But for masked and translucency, every single layer will have to be shaded.  The total GPU time of a frame can easily be doubled or more by having overdraw.  Use the Shader complexity viewmode in Mobile previewer to investigate hot spots.
 * Draw calls of the entire scene should be &lt;= 700 for any view.  The areas with poor occlusion, like looking over the landscape, are going to be the biggest challenges for this.  This can be seen with `stat openglrhi` on device, `stat d3d11rhi` in ES2 previewer on PC.
-* Triangle count of the entire scene should be &lt;= 500k for any view.  This has been determined to be the polycount needed to hit 30fps on both iPad4 and iPad Air.  This can be seen with `stat openglrhi` on device, `stat d3d11rhi` in ES2 previewer on PC.
+* Triangle count of the entire scene should be &lt;= 500k for any view.  This has been determined to be the poly count needed to hit 30fps on both iPad4 and iPad Air.  This can be seen with `stat openglrhi` on device, `stat d3d11rhi` in ES2 previewer on PC.
 * Materials should use as few texture lookups and instructions as possible.  You have 5 texture samplers available, but using all of those makes for a fairly expensive Material.  For now, use the PC instruction count to optimize, and ES2 previewer shader complexity to visualize the total cost.
 * Only use independent texture fetches in Materials.  This means UVs in the pixel shader (BaseColor, Roughness, etc) must not be manipulated in any way, such as scaling.  Use the new CustomizedUVs feature instead to do the scaling in the vertex shader.  Some special features like environment mapping require math on the UVs and that is ok for special cases. 
 * Prefer square textures as they have less memory waste, and always use power of 2 dimensions (512, 1024, 2048).  Use `listtextures` in the ES2 previewer to see where all the texture memory is going. 
@@ -27,7 +27,7 @@ UE4 supports a variety of lighting features on mobile devices. Using these featu
 
 ###LDR (Low Dynamic Range)
 
-Low dynamic range or LDR mode is the highest performance tier supported in UE4 and is recommended for games that do not require lighting or postprocessing features. In order to use this mode, you must disable `Mobile HDR` for your project in the Rendering section of the [Project Settings Editor](Engine/UI/ProjectSettings). 
+Low dynamic range or LDR mode is the highest performance tier supported in UE4 and is recommended for games that do not require lighting or post processing features. In order to use this mode, you must disable `Mobile HDR` for your project in the Rendering section of the [Project Settings Editor](Engine/UI/ProjectSettings). 
 
 **Advantages**
 
@@ -39,7 +39,7 @@ Low dynamic range or LDR mode is the highest performance tier supported in UE4 a
 
 * The scene's color is written out in gamma space with each color channel clamped to the range [0,1].
 * Translucent primitives are blended in gamma space. In most cases, this requires that you author your translucent textures and Materials differently than you would in HDR or for a normal PC game.
-* Postprocessing features are unavailable in this mode.
+* Post processing features are unavailable in this mode.
 
 
 **Recommendations**
@@ -56,7 +56,7 @@ In this tier, you leverage static lighting and fully rough Materials to create l
 **Advantages**
 
 * Access to static lighting and global illumination features.
-* Full HDR pipeline with access to some postprocessing features like tonemapping.
+* Full HDR pipeline with access to some post processing features like tone mapping.
 * Translucency is blended in linear space allowing you to author content as you normally would for any other UE4 game.
 
 
@@ -69,19 +69,19 @@ In this tier, you leverage static lighting and fully rough Materials to create l
 **Recommendations**
 
 * Author all of your Materials with the fully rough flag set.
-* Consider disabling lightmap directionality in your Materials for additional performance.
+* Consider disabling light map directionality in your Materials for additional performance.
 * Use only static lights in your maps.
 * Disable bloom and anti-aliasing. Also stick to the basic set of film and color controls. See [mobile post process](Platforms/Mobile/PostProcessEffects) for more details.
 
 
 ### Full HDR Lighting
 
-In this tier, you take advantage of most of the HDR lighting features available for mobile in UE4 as well as some of the postprocessing features. Using these features requires quite a bit of performance in exchange for high quality lighting features.
+In this tier, you take advantage of most of the HDR lighting features available for mobile in UE4 as well as some of the post processing features. Using these features requires quite a bit of performance in exchange for high quality lighting features.
 
 **Advantages**
 
 * Access to static lighting and global illumination features.
-* Full HDR pipeline with access to some postprocessing features.
+* Full HDR pipeline with access to some post processing features.
 * Translucency is blended in linear space allowing you to author content as you normally would for any other UE4 game.
 * Realistic specular reflections on surfaces with support for varying roughness.
 * Full support for normal mapped surfaces allows interesting surface detail without adding a lot of polygons to your meshes.
@@ -90,7 +90,7 @@ In this tier, you take advantage of most of the HDR lighting features available
 **Recommendations**
 
 * Consider enabling bloom to take full advantage of the HDR lighting pipeline.
-* Realistic specular reflections in combination with HDR lighting can lead to specular aliasing. You should take advantage of the normal map to roughness feature to help reduce specular aliasing due to high frequency information in your normal map. You should also consider enabling anti-aliasing in the postprocess volume to help reduce this artifact.
+* Realistic specular reflections in combination with HDR lighting can lead to specular aliasing. You should take advantage of the normal map to roughness feature to help reduce specular aliasing due to high frequency information in your normal map. You should also consider enabling anti-aliasing in the post process volume to help reduce this artifact.
 * Place reflection captures intelligently to get the best results.
 * Use only static lights in your scene.
 
@@ -114,7 +114,7 @@ In this tier, you take advantage of all of the HDR lighting features available f
  
 ## Shader Complexity View Mode
 
-You can use the shader complexity viewmode in the PC ES2 previewer to get a sense of Material cost.  To use it, 
+You can use the shader complexity view mode in the PC ES2 previewer to get a sense of Material cost.  To use it, 
 just enter `viewmode shadercomplexity` on the console in the PC ES2 previewer. Just like the SM5 shader complexity, 
 green means good performance, bright red means very expensive, white or pink means EXTREMELY expensive.  Here are some 
 examples from the mobile showcase:
@@ -136,3 +136,47 @@ Here the pillars are quite expensive as they used all 5 texture lookups and did
 Here the tree is causing pixels to be EXTREMELY expensive.  If the player could ever get up close or cover his screen with that Material, the cost would be enormous.
 [/EXCERPT:Image3]
 -->
+
+## Mobile Content Scale Factor 
+Mobile Content Scale Factor is a way to scale the resolution of your project to best fit the screen resolution of the mobile device being used to view your project. 
+You can create and store multiple device profiles by creating a new .INI file called **DefaultDeviceProfiles.INI** in your project's config folder. 
+Inside of this file you can specify what will happen to the resolution of your project by entering the following command. 
+
+	 r.MobileContentScaleFactor = 
+
+In the following image you can see an example of the Mobile Content Scale Factor in action using the command listed above.
+
+![](T_MCF_INI_File.png)
+
+The image above was taken from the UE4 Tappy Chicken project and shows what will happen to the resolution of Tappy Chicken when it is played on various mobile devices.
+The top section of the file handles resolution scaling for iOS devices and the bottom section handles resolution scaling for Android devices.
+Make sure to note that each **r.MobileContentScaleFactor** has a number after it.
+The number that is used after the command is what determines if the projects resolution should be scaled up or down.
+
+
+### Mobile Content Scale Factor for iOS
+When using Mobile Content Scale Factor for iOS devices inputting the following number after the r.MobileContentScaleFactor will produce the following results.
+Please note that the actual resolution for any scale factor other than 0.0 will be corrected to match the aspect ratio of the screen and clamped to native resolution.
+
+[region:note]
+The iOS scale factor is directly related to Apple's scale factor system.
+[/region]
+
+**r.MobileContentScaleFactor 0.0**: This will use the native resolution of the device.
+
+**r.MobileContentScaleFactor 1.0**: On a Retina device is the non-retina resolution.
+
+**r.MobileContentScaleFactor 2.0**: Full native resolution on iPhone 5S, iPad Air, etc.
+
+**r.MobileContentScaleFactor 3.0**: Full native resolution for iPhone 6+.
+
+### Mobile Content Scale Factor for Android
+When using the Mobile Content Scale Factor for Android devices inputting the following number after the r.MobileContentScaleFactor will produce the following results.
+Please note that the actual resolution for any scale factor other than 0.0 will be corrected to match the aspect ratio of the screen and clamped to native resolution.
+
+**r.MobileContentScaleFactor 0.0**: This will use the native resolution of the device.
+
+**r.MobileContentScaleFactor 1.0**: This will try and give you a resolution of 1280 x 720 for Landscape and 720 x 1280 for portrait.
+
+**r.MobileContentScaleFactor 2.0**: This will try and give you a resolution of 2560 x 1440 for Landscape and 1440 x 2560 for portrait.
+
