#UE4doc #docs: Removed ##Overview as according to the Instructional Guide.

diff --git a/Engine/Documentation/Source/Engine/Blueprints/BP_HowTo/NewBlueprint/NewBlueprint.INT.udn b/Engine/Documentation/Source/Engine/Blueprints/BP_HowTo/NewBlueprint/NewBlueprint.INT.udn
--- a/Engine/Documentation/Source/Engine/Blueprints/BP_HowTo/NewBlueprint/NewBlueprint.INT.udn
+++ b/Engine/Documentation/Source/Engine/Blueprints/BP_HowTo/NewBlueprint/NewBlueprint.INT.udn
@@ -6,8 +6,6 @@ Related: Engine\Blueprints
 Related: Engine\Blueprints\UserGuide
 Related: Gameplay
 
-## Overview
-
 
 
 
