#UE4doc #docs: Altered images to match the instructional guide: green -> yellow; Removed ##Overview as according to the instructional guide.

diff --git a/Engine/Documentation/Source/Engine/Blueprints/BP_HowTo/WorkingWithArrays/WorkingWithArrays.INT.udn b/Engine/Documentation/Source/Engine/Blueprints/BP_HowTo/WorkingWithArrays/WorkingWithArrays.INT.udn
--- a/Engine/Documentation/Source/Engine/Blueprints/BP_HowTo/WorkingWithArrays/WorkingWithArrays.INT.udn
+++ b/Engine/Documentation/Source/Engine/Blueprints/BP_HowTo/WorkingWithArrays/WorkingWithArrays.INT.udn
@@ -7,8 +7,6 @@ Related: Engine/Blueprints/UserGuide/Arrays/ArrayNodes
 
 [TOC(start:2)]
 
-## Overview
-
 If you are not familiar with programming terms, the best way to think of an **Array** is that it is a collection of variables of the same type, contained within a single entity. For example, if you were to create a **Float** Array, that Array would only be able to hold Float variables in it. Once you have an Array created, you can access individual variables within the Array and modify them or modify all the items in the Array at once (which greatly improves repetitive tasks by doing it all at once versus individually). 
 
 On this page, you will learn how to create Array variables inside a Blueprint as well as using the **Make Array** node to build an Array. Once you understand how to create an Array, you can check out the **Example Array Usage** section for some sample use cases for working with Arrays. 
diff --git a/Engine/Documentation/Source/Engine/Blueprints/Editor/Modes/Defaults/DefaultsTab/DefaultsTab.INT.udn b/Engine/Documentation/Source/Engine/Blueprints/Editor/Modes/Defaults/DefaultsTab/DefaultsTab.INT.udn
--- a/Engine/Documentation/Source/Engine/Blueprints/Editor/Modes/Defaults/DefaultsTab/DefaultsTab.INT.udn
+++ b/Engine/Documentation/Source/Engine/Blueprints/Editor/Modes/Defaults/DefaultsTab/DefaultsTab.INT.udn
@@ -5,8 +5,6 @@ Description:An overview of the Blueprints Defaults tab
 
 [TOC (start:2)]
 
-## Overview
-
 The **Class Defaults** tab contains information about the Blueprint's default settings and properties, and about any variables the Blueprint contains. You can modify these settings to suit your needs.
 
 [REGION:raw]
