#UE4doc #docs: ##Overview removed as according to the instructional guide.

diff --git a/Engine/Documentation/Source/Engine/Blueprints/Editor/UIBreakdowns/ClassBPUI/ClassBPUI.INT.udn b/Engine/Documentation/Source/Engine/Blueprints/Editor/UIBreakdowns/ClassBPUI/ClassBPUI.INT.udn
--- a/Engine/Documentation/Source/Engine/Blueprints/Editor/UIBreakdowns/ClassBPUI/ClassBPUI.INT.udn
+++ b/Engine/Documentation/Source/Engine/Blueprints/Editor/UIBreakdowns/ClassBPUI/ClassBPUI.INT.udn
@@ -23,9 +23,6 @@ Description:A breakdown of the UI elements of the Blueprint Editor when working
 [/OBJECT]
 [/VAR]
 
-
-## Overview
-
 ![](ClassBlueprint.png)
 
 The **Blueprint Class** editor contains several tabs by default:
