#UE4doc #docs: ##Overview removed as according to the instructional guide.

diff --git a/Engine/Documentation/Source/Engine/Blueprints/Editor/UIBreakdowns/DataOnlyUI/DataOnlyUI.INT.udn b/Engine/Documentation/Source/Engine/Blueprints/Editor/UIBreakdowns/DataOnlyUI/DataOnlyUI.INT.udn
--- a/Engine/Documentation/Source/Engine/Blueprints/Editor/UIBreakdowns/DataOnlyUI/DataOnlyUI.INT.udn
+++ b/Engine/Documentation/Source/Engine/Blueprints/Editor/UIBreakdowns/DataOnlyUI/DataOnlyUI.INT.udn
@@ -23,11 +23,6 @@ Description:Overview of the Blueprint Class Editor interface.
 [/OBJECT]
 [/VAR]
 
-
-[TOC (start:2)]
-
-## Overview
-
 [INCLUDE:Engine/Blueprints/UserGuide/Types/ClassBlueprint#datablueprint]
 
 
