#UE4doc #docs: ##Overview removed as according to the instructional guide.

diff --git a/Engine/Documentation/Source/Engine/Blueprints/Editor/UIBreakdowns/InterfaceUI/InterfaceUI.INT.udn b/Engine/Documentation/Source/Engine/Blueprints/Editor/UIBreakdowns/InterfaceUI/InterfaceUI.INT.udn
--- a/Engine/Documentation/Source/Engine/Blueprints/Editor/UIBreakdowns/InterfaceUI/InterfaceUI.INT.udn
+++ b/Engine/Documentation/Source/Engine/Blueprints/Editor/UIBreakdowns/InterfaceUI/InterfaceUI.INT.udn
@@ -25,8 +25,6 @@ Description:Overview of the Blueprint Class Editor interface.
 
 [TOC (start:2)]
 
-## Overview
-
 [INCLUDE:Engine/Blueprints/UserGuide\Types\Interface#intro]
 
 ## Blueprint Interface Editor UI
