#UE4doc #docs: ##Overview removed as according to the instructional guide.

diff --git a/Engine/Documentation/Source/Engine/Blueprints/Editor/UIBreakdowns/LevelBPUI/LevelBPUI.INT.udn b/Engine/Documentation/Source/Engine/Blueprints/Editor/UIBreakdowns/LevelBPUI/LevelBPUI.INT.udn
--- a/Engine/Documentation/Source/Engine/Blueprints/Editor/UIBreakdowns/LevelBPUI/LevelBPUI.INT.udn
+++ b/Engine/Documentation/Source/Engine/Blueprints/Editor/UIBreakdowns/LevelBPUI/LevelBPUI.INT.udn
@@ -26,8 +26,6 @@ Description:Overview of the Level Blueprint Editor interface.
 
 [TOC (start:2)]
 
-## Overview
-
 [INCLUDE:Engine/Blueprints/UserGuide/Types/LevelBlueprint#Intro]
 
 ## Interface
