#UE4doc #docs:  ##Overview removed as according to the instructional guide.

diff --git a/Engine/Documentation/Source/Engine/Blueprints/Editor/UIBreakdowns/MacroLibraryUI/MacroLibraryUI.INT.udn b/Engine/Documentation/Source/Engine/Blueprints/Editor/UIBreakdowns/MacroLibraryUI/MacroLibraryUI.INT.udn
--- a/Engine/Documentation/Source/Engine/Blueprints/Editor/UIBreakdowns/MacroLibraryUI/MacroLibraryUI.INT.udn
+++ b/Engine/Documentation/Source/Engine/Blueprints/Editor/UIBreakdowns/MacroLibraryUI/MacroLibraryUI.INT.udn
@@ -25,8 +25,6 @@ Description:Overview of the Blueprint Class Editor interface.
 
 [TOC (start:2)]
 
-## Overview
-
 [INCLUDE:Engine/Blueprints/UserGuide/Types/MacroLibrary#intro]
 
 For more information on Macro Libraries and how to use them, please see the [](Engine/Blueprints/UserGuide\Types\MacroLibrary) documentation.
