#UE4doc #docs: ##Overview removed as according to the instructional guide.

diff --git a/Engine/Documentation/Source/Engine/Blueprints/Editor/UIComponents/CompilerResults/CompilerResults.INT.udn b/Engine/Documentation/Source/Engine/Blueprints/Editor/UIComponents/CompilerResults/CompilerResults.INT.udn
--- a/Engine/Documentation/Source/Engine/Blueprints/Editor/UIComponents/CompilerResults/CompilerResults.INT.udn
+++ b/Engine/Documentation/Source/Engine/Blueprints/Editor/UIComponents/CompilerResults/CompilerResults.INT.udn
@@ -8,8 +8,6 @@ Description:Interface for the Compiler Results panel of the Blueprint Editor.
 ![](CompilerWindow.png)
 [/REGION]
 
-## Overview
-
 The **Compiler Results** panel provides feedback on the compiling of your Blueprint scripts, telling you whether a Blueprint script compile was successful or if there were _errors_ or _warnings_. Whenever errors or warnings are found, the **Compiler Results** panel gives information on the error and provides hyperlink access, enabling you to jump the graph view directly to the source of the problem.
 
 ## Interface
diff --git a/Engine/Documentation/Source/Engine/Blueprints/Editor/UIComponents/Debug/Debug.INT.udn b/Engine/Documentation/Source/Engine/Blueprints/Editor/UIComponents/Debug/Debug.INT.udn
--- a/Engine/Documentation/Source/Engine/Blueprints/Editor/UIComponents/Debug/Debug.INT.udn
+++ b/Engine/Documentation/Source/Engine/Blueprints/Editor/UIComponents/Debug/Debug.INT.udn
@@ -3,9 +3,6 @@ Title: Debug Panel
 Crumbs:%ROOT%, Engine/Blueprints, Engine/Blueprints/Editor, Engine/Blueprints/Editor/UIBreakdowns/LevelBPUI
 Description:
 
-
-## Overview
-
 The **Debug** panel provides a list of all debug tools currently in use for the Blueprint, such as Breakpoints and Watch Values. During Play-In-Editor, it also provides access to the Execution Trace, showing you each node execution taking place within a given Blueprint.
 
 For more information on debugging Blueprints, please see [](Engine/Blueprints/UserGuide\Debugging).
