#UE4doc #docs: ##Overview removed as according to the instructional guide.

diff --git a/Engine/Documentation/Source/Engine/Blueprints/Editor/UIComponents/Details/KismetDetails.INT.udn b/Engine/Documentation/Source/Engine/Blueprints/Editor/UIComponents/Details/KismetDetails.INT.udn
--- a/Engine/Documentation/Source/Engine/Blueprints/Editor/UIComponents/Details/KismetDetails.INT.udn
+++ b/Engine/Documentation/Source/Engine/Blueprints/Editor/UIComponents/Details/KismetDetails.INT.udn
@@ -6,8 +6,6 @@ Description:Overview of the Details panel within the Blueprint Editor.
 
 
 
-## Overview
-
 [EXCERPT:Intro]
 The **Details** panel is a context-sensitive area that provides editing access to the properties of selected items within the Blueprint Editor. It is composed of a search bar for fast access to specific properties, and will generally contain one or more collapsible categories to organize the included properties.
 
