#UE4doc #docs: ##Overview removed as according to the instructional guide. Made spelling changes: bluerprint -> blueprint, critea -> criteria

diff --git a/Engine/Documentation/Source/Engine/Blueprints/Editor/UIComponents/FindResults/FindResults.INT.udn b/Engine/Documentation/Source/Engine/Blueprints/Editor/UIComponents/FindResults/FindResults.INT.udn
--- a/Engine/Documentation/Source/Engine/Blueprints/Editor/UIComponents/FindResults/FindResults.INT.udn
+++ b/Engine/Documentation/Source/Engine/Blueprints/Editor/UIComponents/FindResults/FindResults.INT.udn
@@ -4,9 +4,8 @@ Crumbs: %ROOT%, Engine, Engine/Blueprints, Engine/Blueprints/Editor, Engine/Blue
 Description:
 
 
-## Overview
 
-The **Find Results** panel is a powerful search tool within the Bluerprint Editor that allows you to quickly track down a variety of objects based on the following criteria:
+The **Find Results** panel is a powerful search tool within the Blueprint Editor that allows you to quickly track down a variety of objects based on the following criteria:
 
 * Node name
 * Pin name
@@ -28,7 +27,7 @@ Unlike many search fields in Unreal Engine 4, the **Find Results** panel does no
 ![](FindResults.png)(w:940)
 [/REGION]
 
-1. **Results list** - This lists all nodes, pins, property names, comments, and property values that match the search critera.
+1. **Results list** - This lists all nodes, pins, property names, comments, and property values that match the search criteria.
 1. **Search filter** - This is where you type in what you are looking for.
 1. **Property values** - Explicitly set property values will appear in parentheses in the middle of the results.
 1. **Comments** - Node comments, if they exist, will appear in yellow text on the right side of the panel.
