#UE4doc #docs: ##Overview removed as according to the instructional guide. Altered images to match the instructional guide: green -> yellow

diff --git a/Engine/Documentation/Source/Engine/Blueprints/Editor/UIComponents/GraphEditor/GraphEditor.INT.udn b/Engine/Documentation/Source/Engine/Blueprints/Editor/UIComponents/GraphEditor/GraphEditor.INT.udn
--- a/Engine/Documentation/Source/Engine/Blueprints/Editor/UIComponents/GraphEditor/GraphEditor.INT.udn
+++ b/Engine/Documentation/Source/Engine/Blueprints/Editor/UIComponents/GraphEditor/GraphEditor.INT.udn
@@ -6,8 +6,6 @@ Description:Editable, visual representation of a network of nodes within a Bluep
 
 [TOC(start:2 end:3)]
 
-## Overview
-
 
 [EXCERPT:Intro]
 The **Graph Editor** panel is the heart of the Blueprint system. It is here that you will create the networks of nodes and wires that will define your scripted behavior. Nodes can quickly be selected by clicking on them and repositioned via dragging.
