#UE4doc #docs:  ##Overview removed as according to the instructional guide.

diff --git a/Engine/Documentation/Source/Engine/Blueprints/Editor/UIComponents/Menu/BlueprintEditorMenu.INT.udn b/Engine/Documentation/Source/Engine/Blueprints/Editor/UIComponents/Menu/BlueprintEditorMenu.INT.udn
--- a/Engine/Documentation/Source/Engine/Blueprints/Editor/UIComponents/Menu/BlueprintEditorMenu.INT.udn
+++ b/Engine/Documentation/Source/Engine/Blueprints/Editor/UIComponents/Menu/BlueprintEditorMenu.INT.udn
@@ -3,7 +3,6 @@ Title: Menu
 Crumbs:%ROOT%, Engine/Blueprints, Engine/Blueprints/Editor
 Description:Menu
 
-## Overview
 
 
 [EXCERPT:Intro]
