#UE4doc #docs: ##Overview removed as according to the instructional guide.

diff --git a/Engine/Documentation/Source/Engine/Blueprints/Editor/UIComponents/Palette/AllClassesPanel.INT.udn b/Engine/Documentation/Source/Engine/Blueprints/Editor/UIComponents/Palette/AllClassesPanel.INT.udn
--- a/Engine/Documentation/Source/Engine/Blueprints/Editor/UIComponents/Palette/AllClassesPanel.INT.udn
+++ b/Engine/Documentation/Source/Engine/Blueprints/Editor/UIComponents/Palette/AllClassesPanel.INT.udn
@@ -6,8 +6,6 @@ Description:Non-context sensitive listing of all functions and variables for Pal
 
 [TOC (start:2)]
 
-## Overview
-
 [EXCERPT:Intro]
 The **Palette** tab displays a list of all functions and variables that can be used in _Blueprints_. It is fully searchable, contains a **Favorites** panel that you can populate with the nodes you find yourself using the most, and you can filter it to search for nodes that are compatible with a specific class type. Nodes from the Palette can be added into your networks by simply dragging and dropping them from the list.
 
