#UE4doc #docs: ##Overview removed as according to the instructional guide.

diff --git a/Engine/Documentation/Source/Engine/Blueprints/Editor/UIComponents/Toolbar/ToolbarOptions/ToolbarOptions.INT.udn b/Engine/Documentation/Source/Engine/Blueprints/Editor/UIComponents/Toolbar/ToolbarOptions/ToolbarOptions.INT.udn
--- a/Engine/Documentation/Source/Engine/Blueprints/Editor/UIComponents/Toolbar/ToolbarOptions/ToolbarOptions.INT.udn
+++ b/Engine/Documentation/Source/Engine/Blueprints/Editor/UIComponents/Toolbar/ToolbarOptions/ToolbarOptions.INT.udn
@@ -5,7 +5,6 @@ Description:
 
 [TOC (start:2)]
 
-## Overview
 
 [INCLUDE:Engine/Blueprints/Editor/UIComponents/Toolbar#Intro]
 
