#UE4doc #docs:  ##Overview removed as according to the instructional guide.

diff --git a/Engine/Documentation/Source/Engine/Blueprints/Editor/UIComponents/Viewport/BlueprintViewport.INT.udn b/Engine/Documentation/Source/Engine/Blueprints/Editor/UIComponents/Viewport/BlueprintViewport.INT.udn
--- a/Engine/Documentation/Source/Engine/Blueprints/Editor/UIComponents/Viewport/BlueprintViewport.INT.udn
+++ b/Engine/Documentation/Source/Engine/Blueprints/Editor/UIComponents/Viewport/BlueprintViewport.INT.udn
@@ -5,8 +5,6 @@ Description:Overview of the viewport in the Components mode of Blueprint Classes
 
 [TOC (start:2)]
 
-## Overview
-
 In the **Viewport**, in the Blueprint Editor, you can view and manipulate your Blueprint's components. 
 
 ![](BP_Viewport.png)
