#UE4doc #docs: ##Overview removed as according to the instructional guide.

diff --git a/Engine/Documentation/Source/Engine/Blueprints/GettingStarted/BlueprintIntro.INT.udn b/Engine/Documentation/Source/Engine/Blueprints/GettingStarted/BlueprintIntro.INT.udn
--- a/Engine/Documentation/Source/Engine/Blueprints/GettingStarted/BlueprintIntro.INT.udn
+++ b/Engine/Documentation/Source/Engine/Blueprints/GettingStarted/BlueprintIntro.INT.udn
@@ -26,8 +26,6 @@ Description:If you are just getting started with Blueprints, this provides a hig
 
 [TOC (start:2)]
 
-## Overview
-
 [EXCERPT:BlueprintIntro]
 **Blueprints** are special assets that provide an intuitive, node-based interface that can be used to create new types of Actors 
 and script level events; giving designers and gameplay programmers the tools to quickly create and iterate gameplay from within 
