#UE4doc #docs: ##Overview removed as according to the instructional guide.

diff --git a/Engine/Documentation/Source/Engine/Blueprints/HowTo/BlueprintHowTo.INT.udn b/Engine/Documentation/Source/Engine/Blueprints/HowTo/BlueprintHowTo.INT.udn
--- a/Engine/Documentation/Source/Engine/Blueprints/HowTo/BlueprintHowTo.INT.udn
+++ b/Engine/Documentation/Source/Engine/Blueprints/HowTo/BlueprintHowTo.INT.udn
@@ -43,8 +43,6 @@ Description:A sample project created using different methods of Blueprint Commun
 [/OBJECT]
 [/VAR]
 
-
-## Overview
 In this **Blueprint Communication Project**, we will create a simple shooting gallery style game. Over the course of the project you will learn how to communicate information from the level to a character as well as how to pass character information to a HUD. At the end of the project, you will have something similar to below:
 
 ![](Engine\Blueprints\HowTo\BPHT_6\GameplayGIF.gif)(convert:false)
