#UE4doc #docs: Minor image adjustment

diff --git a/Engine/Documentation/Source/Engine/Rendering/Materials/Editor/MaterialEditor.INT.udn b/Engine/Documentation/Source/Engine/Rendering/Materials/Editor/MaterialEditor.INT.udn
--- a/Engine/Documentation/Source/Engine/Rendering/Materials/Editor/MaterialEditor.INT.udn
+++ b/Engine/Documentation/Source/Engine/Rendering/Materials/Editor/MaterialEditor.INT.udn
@@ -26,7 +26,7 @@ Description:Guide to using the Material Editor for creating shaders.
 [TOC(start:2)]
 
 
-## Overview
+
 
 [EXCERPT:Intro]
 The **Material Editor** provides the ability to create shaders to be applied to geometry using a node-based graph interface.
@@ -132,17 +132,12 @@ Each time a change is made to the material network, the Material must be compile
 the network have no connections or are being passed the wrong type of data, a compiler error will occur. These errors are displayed 
 in the [Graph Panel](Engine/Rendering/Materials/Editor/Interface/#GraphPanel).
 
-![error_message.png](error_message.png)
-
 Compiler errors let you know a problem exists and what that problem is by providing information about the type of expression 
-that they occurred on and a description of the error. In addition, the node that contains the error will be highlighted in red 
-in the **Graph** tab so that you can easily find it and make the necessary corrections.
+that they occurred on and a description of the error. 
 
 ![error_highlight.png](error_highlight.png)
 
-[REGION:warning]
-Compiler errors will only be visible if Stats are enabled. Be sure to check the **Stats** button ![toolbar_stats.png](toolbar_stats.png)(w:21 h:30) in the **Toolbar** tab.
-[/REGION]
+
 
 (#ExpressionSearch)
 
