Deleted "Overview"

diff --git a/Engine/Documentation/Source/Engine/Rendering/Materials/Editor/Interface/MaterialEditorInterface.INT.udn b/Engine/Documentation/Source/Engine/Rendering/Materials/Editor/Interface/MaterialEditorInterface.INT.udn
--- a/Engine/Documentation/Source/Engine/Rendering/Materials/Editor/Interface/MaterialEditorInterface.INT.udn
+++ b/Engine/Documentation/Source/Engine/Rendering/Materials/Editor/Interface/MaterialEditorInterface.INT.udn
@@ -6,7 +6,7 @@ Description:Guide to using the Material Editor UI.
 
 [TOC(start:2)]
 
-## Overview
+
 
 The Material Editor UI consists of a menu bar, a toolbar, and, by default, four open panels.
 
