#UE4doc #docs: Altered images to match the instructional guide: green -> yellow

diff --git a/Engine/Documentation/Source/Engine/Blueprints/HowTo/BPHT_2/BPHT_2.INT.udn b/Engine/Documentation/Source/Engine/Blueprints/HowTo/BPHT_2/BPHT_2.INT.udn
--- a/Engine/Documentation/Source/Engine/Blueprints/HowTo/BPHT_2/BPHT_2.INT.udn
+++ b/Engine/Documentation/Source/Engine/Blueprints/HowTo/BPHT_2/BPHT_2.INT.udn
@@ -12,7 +12,7 @@ Description:In this step we are going to make sure our playable character is set
 	[PARAM:current]
 	[/PARAM]
 	[PARAM:home]
-		[How-To Home](Engine\Blueprints\HowTo "%Engine\Blueprints\HowTo:Title%")
+		[How To Home](Engine\Blueprints\HowTo "%Engine\Blueprints\HowTo:Title%")
 	[/PARAM]
 	[PARAM:next]
 	[Next Step](Engine\Blueprints\HowTo\BPHT_3 "%Engine\Blueprints\HowTo\BPHT_3:title%") 
