#UE4doc #docs: Updated Choosing Material Function image and appropriate text

diff --git a/Engine/Documentation/Source/Engine/Rendering/Materials/Functions/Creating/CreatingFunctions.INT.udn b/Engine/Documentation/Source/Engine/Rendering/Materials/Functions/Creating/CreatingFunctions.INT.udn
--- a/Engine/Documentation/Source/Engine/Rendering/Materials/Functions/Creating/CreatingFunctions.INT.udn
+++ b/Engine/Documentation/Source/Engine/Rendering/Materials/Functions/Creating/CreatingFunctions.INT.udn
@@ -8,7 +8,7 @@ Description:A guide to the creation process for material functions.
 ## Creating Material Functions
 
 
-You can create new material functions within the **Content Browser** by **right-clicking** and choosing **New** > **Materials & Textures** > **Material Function**. This will open up the new material function within the Material Editor.
+You can create new material functions within the **Content Browser** by **clicking** and choosing **Add New** > **Materials & Textures** > **Material Function**. This will open up the new material function within the Material Editor.
 
 ![CreateNewMenu.png](CreateNewMenu.png)
 
