#UE4doc #docs:  Removed "overview"

diff --git a/Engine/Documentation/Source/Engine/Rendering/Materials/Functions/Overview/FunctionsOverview.INT.udn b/Engine/Documentation/Source/Engine/Rendering/Materials/Functions/Overview/FunctionsOverview.INT.udn
--- a/Engine/Documentation/Source/Engine/Rendering/Materials/Functions/Overview/FunctionsOverview.INT.udn
+++ b/Engine/Documentation/Source/Engine/Rendering/Materials/Functions/Overview/FunctionsOverview.INT.udn
@@ -9,7 +9,7 @@ Description:An overview of how material functions work and some critical concept
 
 [TOC(start:2)]
 
-## Overview
+
 
 Material Functions are little snippets of material graphs that can be saved in packages and reused across multiple materials.  Their purpose is to streamline the process of material creation by giving instant access to commonly used networks of material nodes. For example, if you find yourself often creating a network just to handle chaotic texture panning, then it will be much faster to save that portion of your network out as a material function and just use that whenever you need such behavior.
 
