#UE4doc #docs:  ##Overview removed as according to the instructional guide.

diff --git a/Engine/Documentation/Source/Engine/Blueprints/HowTo/BPHT_5/BPHT_5.INT.udn b/Engine/Documentation/Source/Engine/Blueprints/HowTo/BPHT_5/BPHT_5.INT.udn
--- a/Engine/Documentation/Source/Engine/Blueprints/HowTo/BPHT_5/BPHT_5.INT.udn
+++ b/Engine/Documentation/Source/Engine/Blueprints/HowTo/BPHT_5/BPHT_5.INT.udn
@@ -12,7 +12,7 @@ Description:In this section we will use a Bind Event which is bound to OnHit to
 	[PARAM:current]
 	[/PARAM]
 	[PARAM:home]
-		[How-To Home](Engine\Blueprints\HowTo "%Engine\Blueprints\HowTo:Title%")
+		[How To Home](Engine\Blueprints\HowTo "%Engine\Blueprints\HowTo:Title%")
 	[/PARAM]
 	[PARAM:next]
 	[Next Step](Engine\Blueprints\HowTo\BPHT_6 "%Engine\Blueprints\HowTo\BPHT_6:title%") 
