#UE4doc #docs: minor grammar change: a -> an

diff --git a/Engine/Documentation/Source/Engine/Blueprints/QuickStart/4/BPQS_4.INT.udn b/Engine/Documentation/Source/Engine/Blueprints/QuickStart/4/BPQS_4.INT.udn
--- a/Engine/Documentation/Source/Engine/Blueprints/QuickStart/4/BPQS_4.INT.udn
+++ b/Engine/Documentation/Source/Engine/Blueprints/QuickStart/4/BPQS_4.INT.udn
@@ -36,7 +36,7 @@ To start adding behavior to your Blueprint Class, you will need to open it in th
 	![](AdjustedTrigger.png)
 	[/REGION]
 
-1. Docked next to the **Viewport** tab is a **Construction Script** tab and a **Event Graph** tab. Since you are going to be creating gameplay behavior, you should start with the **Event Graph**. Click on that tab now.
+1. Docked next to the **Viewport** tab is a **Construction Script** tab and an **Event Graph** tab. Since you are going to be creating gameplay behavior, you should start with the **Event Graph**. Click on that tab now.
 
 	![](EventGraphStart.png)(w:500)
 
