#UE4doc #docs: Cube Collision is now known as Box Collision in v4.8

diff --git a/Engine/Documentation/Source/Engine/Blueprints/QuickStart/2/BPQS_2.INT.udn b/Engine/Documentation/Source/Engine/Blueprints/QuickStart/2/BPQS_2.INT.udn
--- a/Engine/Documentation/Source/Engine/Blueprints/QuickStart/2/BPQS_2.INT.udn
+++ b/Engine/Documentation/Source/Engine/Blueprints/QuickStart/2/BPQS_2.INT.udn
@@ -48,8 +48,8 @@ You will build the launchpad in the Level Editor, and then convert it to a Bluep
 1. In the **Details** panel, click the **Add Component** button, then select **Cube** under **Basic Shapes**.
 1. Now you should have the **Cube** component selected, so you can tweak its properties. Resize the cube to (1, 1, 0.1).
 
-1. Now, we will add a Box Collision Component, which will fire off an event whenever something overlaps it. In the **Details** panel, click the **Add Component** button, then select **Cube Collision** under **Collision**.
-1. Resize the box to be the size of the platform in X and Y, and extend a bit above the platform in Z. If you want to enter in your **Scale** directly, try (1.25, 1.25, 15.0).
+1. Now, we will add a Box Collision Component, which will fire off an event whenever something overlaps it. In the **Details** panel, click the **Add Component** button, then select **Box Collision** under **Collision**.
+1. Resize the box to be the size of the platform in X and Y, and extend a bit above the platform in Z. If you want to enter in your **Scale** directly, try (1.25, 1.25, 1.5).
 
 [REGION:tip]
 If at any time you need to change your Actor's properties, such as to scale up the whole launchpad, you can click on **LaunchPad (Instance)** right under the **Add Components** button in the **Details** panel.
