#UE4doc #docs: ##Overview removed as according to the instructional guide. Also changed two spelling errors: behaviour -> behavior,  specicified -> specified

diff --git a/Engine/Documentation/Source/Engine/Blueprints/TechnicalGuide/ExtendingBlueprints/ExposingGameplay.INT.udn b/Engine/Documentation/Source/Engine/Blueprints/TechnicalGuide/ExtendingBlueprints/ExposingGameplay.INT.udn
--- a/Engine/Documentation/Source/Engine/Blueprints/TechnicalGuide/ExtendingBlueprints/ExposingGameplay.INT.udn
+++ b/Engine/Documentation/Source/Engine/Blueprints/TechnicalGuide/ExtendingBlueprints/ExposingGameplay.INT.udn
@@ -6,9 +6,6 @@ Related: Engine/Blueprints/TechnicalGuide/Guidelines
 
 [TOC(start:2)]
 
-
-## Overview
-
 As a programmer, working with Blueprints offers you many benefits by making your code extremely flexible. For 
 example, a gameplay designer might want a new type of weapon implemented in the game. As the programmer, you 
 code the weapon as you traditionally would except that you expose some important functionality, such as the 
@@ -40,7 +37,7 @@ The simplest form of declaration for a Blueprintable class is similar to the fol
 | -- | -- |
 | **Blueprintable** | Exposes this class as an acceptable base class for creating Blueprints. The default is NotBlueprintable, unless inherited otherwise. This is inherited by subclasses. |
 | **BlueprintType** | Exposes this class as a type that can be used for variables in Blueprints. |
-| **NotBlueprintable** | Specifies that this class is NOT an acceptable base class for creating blueprints. Negates the effect of a parent class having the Blueprintable keyword specficified. |
+| **NotBlueprintable** | Specifies that this class is NOT an acceptable base class for creating blueprints. Negates the effect of a parent class having the Blueprintable keyword specified. |
 
 
 ## Readable and Writable Properties
@@ -83,7 +80,7 @@ The simplest form of declaration for a BlueprintCallable Function is similar to
 | **BlueprintPure** | This is a native function that can be called from a Blueprint, which executes native code that does not change anything about the object it is being called on, or any other global state.  This means that nothing is changed by calling this node, it just takes input, and tells you an output.  These are things like math nodes (+, - , *, etc), or variable getters, or anything that does not change anything permanently.  These do not  need to be scheduled, and do not have a connection for the white execution line.  They are automatically figured out by the compiler, based on which BlueprintCallable nodes need the data produced by these nodes. |
 | [REGION:tablesection]Native to Blueprint Communication[/REGION] ||
 | **BlueprintImplementableEvent** | This is the primary way we allow native functions to call up into Blueprints.  They are like virtual functions that you implement in the Blueprint themselves.  If there is no implementation, the function call is ignored.  It is important to note that if a BlueprintImplementableEvent does not  have a return value or an out parameter, it will appear as an event, available by **right-clicking** and selecting it in the event graph of the Blueprint.  If it does have a return value or any out parameters, it will be listed in the **My Blueprints** tab, and can then be overridden by **right-clicking** and selecting "implement" function.  Note that BlueprintImplementableEvents do not  have a native implementation of the function. |
-| **BlueprintNativeEvent** | These are the same as above, except there is a native default implementation of the function that is called if the Blueprint does not override the function.  This is useful for things where you want some sort of default behaviour if the Blueprint does not implement it, but want the Blueprint to be able to override the functionality if desired.  These are more costly, so we only put them in where the functionality is needed.  When you override the BlueprintNativeEvent, you can still call the native implementation, if desired, by **right-clicking** on the event or function entry node, and selecting "Add call to parent". |
+| **BlueprintNativeEvent** | These are the same as above, except there is a native default implementation of the function that is called if the Blueprint does not override the function.  This is useful for things where you want some sort of default behavior if the Blueprint does not implement it, but want the Blueprint to be able to override the functionality if desired.  These are more costly, so we only put them in where the functionality is needed.  When you override the BlueprintNativeEvent, you can still call the native implementation, if desired, by **right-clicking** on the event or function entry node, and selecting "Add call to parent". |
 
 
 
