#UE4doc #docs:  ##Overview removed as according to the instructional guide.

diff --git a/Engine/Documentation/Source/Engine/Blueprints/TechnicalGuide/Guidelines/Guidelines.INT.udn b/Engine/Documentation/Source/Engine/Blueprints/TechnicalGuide/Guidelines/Guidelines.INT.udn
--- a/Engine/Documentation/Source/Engine/Blueprints/TechnicalGuide/Guidelines/Guidelines.INT.udn
+++ b/Engine/Documentation/Source/Engine/Blueprints/TechnicalGuide/Guidelines/Guidelines.INT.udn
@@ -6,7 +6,6 @@ Related:Engine/Blueprints/TechnicalGuide/ExtendingBlueprints
 
 [TOC (start:2 end:2)]
 
-## Overview
 There are two driving considerations when deciding to use C++ or Blueprints:
 -	Speed
 -	Complexity of expression
