#UE4doc #docs: ##Overview removed as according to the instructional guide. Also removed instances of the same word in a row: to to -> to

diff --git a/Engine/Documentation/Source/Engine/Blueprints/UserGuide/Arrays/Arrays.INT.udn b/Engine/Documentation/Source/Engine/Blueprints/UserGuide/Arrays/Arrays.INT.udn
--- a/Engine/Documentation/Source/Engine/Blueprints/UserGuide/Arrays/Arrays.INT.udn
+++ b/Engine/Documentation/Source/Engine/Blueprints/UserGuide/Arrays/Arrays.INT.udn
@@ -3,8 +3,6 @@ Title:Blueprint Arrays
 Crumbs: %ROOT%, Engine/Blueprints, Engine/Blueprints/UserGuide
 Description:Documentation over Blueprint arrays, including creation, editing, and use.
 
-## Overview
-
 ![](Arrays.png)
 
 Just as with variable values, Blueprints can store data within arrays. If you are not versed in programming terms, you can think of an array as a collection of variables all housed within a single unit. 
@@ -63,7 +61,7 @@ To edit the values of an array's defaults:
 
 ### Setting Array Values via Node Networks
 
-If your array is intended to be populated at runtime, you will generally not use defaults. Rather, you will use nodes within the Construction Script or an Event Graph to populate each index. For instance, you could use the **Add** or **Insert** nodes to to add a new value to the next available index, or to insert a value at a given index, respectively.
+If your array is intended to be populated at runtime, you will generally not use defaults. Rather, you will use nodes within the Construction Script or an Event Graph to populate each index. For instance, you could use the **Add** or **Insert** nodes to add a new value to the next available index, or to insert a value at a given index, respectively.
 
 ![](AddArray.png)
 
