#UE4doc #docs: ##Overview removed as according to the instructional guide. Altered images to match the instructional guide: green -> yellow

diff --git a/Engine/Documentation/Source/Engine/Blueprints/UserGuide/BlueprintComms/BlueprintComms.INT.udn b/Engine/Documentation/Source/Engine/Blueprints/UserGuide/BlueprintComms/BlueprintComms.INT.udn
--- a/Engine/Documentation/Source/Engine/Blueprints/UserGuide/BlueprintComms/BlueprintComms.INT.udn
+++ b/Engine/Documentation/Source/Engine/Blueprints/UserGuide/BlueprintComms/BlueprintComms.INT.udn
@@ -49,8 +49,6 @@ related:Engine/Blueprints/UserGuide/EventDispatcher
 
 [TOC(start:2)]
 
-## Overview
-
 When you are working with a Blueprint, you may sometimes find that you want to access the functions, events, and variables that are contained in another Blueprint. The simplest way to do so is to reference the "target" Blueprint through a publicly exposed object variable, and then specify which Blueprint instance you want to use.
 
 ## Simple Way to Reference a Target Blueprint
