#UE4doc #docs: ##Overview removed as according to the instructional guide.

diff --git a/Engine/Documentation/Source/Engine/Blueprints/UserGuide/BlueprintCommsUsage/BlueprintCommsUsage.INT.udn b/Engine/Documentation/Source/Engine/Blueprints/UserGuide/BlueprintCommsUsage/BlueprintCommsUsage.INT.udn
--- a/Engine/Documentation/Source/Engine/Blueprints/UserGuide/BlueprintCommsUsage/BlueprintCommsUsage.INT.udn
+++ b/Engine/Documentation/Source/Engine/Blueprints/UserGuide/BlueprintCommsUsage/BlueprintCommsUsage.INT.udn
@@ -50,8 +50,6 @@ Related: Engine/Blueprints/UserGuide/Types/Interface/UsingInterfaces
 
 [TOC(start:2 end:2)]
 
-## Overview
-
 When you are working with a **Blueprint**, you may sometimes find that you want to access the Functions, Events, or Variables that are contained in another **Blueprint**. Referred to as **Blueprint Communication** there are several methods in which you can access and communicate between multiple **Blueprints** depending upon your needs.
 
 Outlined below are the different methods of communication and sample use cases, refer to the links in each section for further documentation. 
