#UE4doc #docs: ##Overview removed as according to the instructional guide. Fixed spelling error: communcation -> communication

diff --git a/Engine/Documentation/Source/Engine/Blueprints/UserGuide/BlueprintCommsUsage/BPComHowTo/BPComHowTo.INT.udn b/Engine/Documentation/Source/Engine/Blueprints/UserGuide/BlueprintCommsUsage/BPComHowTo/BPComHowTo.INT.udn
--- a/Engine/Documentation/Source/Engine/Blueprints/UserGuide/BlueprintCommsUsage/BPComHowTo/BPComHowTo.INT.udn
+++ b/Engine/Documentation/Source/Engine/Blueprints/UserGuide/BlueprintCommsUsage/BPComHowTo/BPComHowTo.INT.udn
@@ -5,11 +5,9 @@ Description:Examples of how to setup and use different methods of Blueprint Comm
 
 [TOC(start:2 end:4)]
 
-## Overview
-
 On this page you will learn step-by-step how to set up and use different methods of **Blueprint Communication**. 
 
-For examples of when to use each type of **Blueprint Communcation**, refer to the [](Engine\Blueprints\UserGuide\BlueprintCommsUsage) page.
+For examples of when to use each type of **Blueprint Communication**, refer to the [](Engine\Blueprints\UserGuide\BlueprintCommsUsage) page.
 
 ### Direct Blueprint Communications
 
