#UE4doc #docs: ##Overview removed as according to the instructional guide. Fixed spelling error: complie -> compile

diff --git a/Engine/Documentation/Source/Engine/Blueprints/UserGuide/CastNodes/CastNodes.INT.udn b/Engine/Documentation/Source/Engine/Blueprints/UserGuide/CastNodes/CastNodes.INT.udn
--- a/Engine/Documentation/Source/Engine/Blueprints/UserGuide/CastNodes/CastNodes.INT.udn
+++ b/Engine/Documentation/Source/Engine/Blueprints/UserGuide/CastNodes/CastNodes.INT.udn
@@ -47,8 +47,6 @@ Related: Engine/Blueprints/HowTo
 
 [TOC (start:2 end:2)]
 
-## Overview
-
 When using a **Cast To** node in **Blueprints**, put simply, you are attempting to check if the object you are casting from is the specific object you **Cast To**. In other words, say you have created a special Character Blueprint (called MyCharacter for example) that has variables in it or other custom functionality contained within it and assigned it as the Default Pawn Class (or the Character Blueprint all Player Characters use by default). 
 
 If you want to access properties of the Player Character from another Blueprint, you can use the **Get Player Character** node and affect the Player Character in very general ways (setting/getting its location, rotation, etc.) but you cannot access any of the custom functionality you have added to the MyCharacter Blueprint you have set up as you are getting the Player Character but not a specific type of Character. 
diff --git a/Engine/Documentation/Source/Engine/Blueprints/UserGuide/CheatSheet/BlueprintEditorCheatSheet.INT.udn b/Engine/Documentation/Source/Engine/Blueprints/UserGuide/CheatSheet/BlueprintEditorCheatSheet.INT.udn
--- a/Engine/Documentation/Source/Engine/Blueprints/UserGuide/CheatSheet/BlueprintEditorCheatSheet.INT.udn
+++ b/Engine/Documentation/Source/Engine/Blueprints/UserGuide/CheatSheet/BlueprintEditorCheatSheet.INT.udn
@@ -5,7 +5,6 @@ Description:Reference covering Blueprint shortcuts and useful actions.
 
 [TOC(start:2 end:2)]
 
-## Overview
 The **Blueprint** editor has a lot of productivity boosting shortcuts built-in, and while many will come naturally as you use the editor, others are a little bit hidden. This page covers the most commonly used actions and the shortcuts that can be used to perform those actions more efficiently. You can also download this list of shortcuts below. 
 
 Click to download: [Color Version](BlueprintCheatSheet-1989117414.pdf) | [Black/White Version](BlueprintCheatSheet_BlackAndWhite-1276559864.pdf)
@@ -65,7 +64,7 @@ Click to download: [Color Version](BlueprintCheatSheet-1989117414.pdf) | [Black/
 | Undo | **Ctrl + Z** |
 | Find in this Blueprint | **Ctrl + F**|
 | Find in any Blueprint | **Ctrl + Shift + F** |
-| Complie the Blueprint | **F7** |
+| Compile the Blueprint | **F7** |
 
 
 ## Variable Actions (My Blueprint)
