#UE4doc #docs: ##Overview removed as according to the instructional guide.

diff --git a/Engine/Documentation/Source/Engine/Blueprints/UserGuide/Comments/Comments.INT.udn b/Engine/Documentation/Source/Engine/Blueprints/UserGuide/Comments/Comments.INT.udn
--- a/Engine/Documentation/Source/Engine/Blueprints/UserGuide/Comments/Comments.INT.udn
+++ b/Engine/Documentation/Source/Engine/Blueprints/UserGuide/Comments/Comments.INT.udn
@@ -42,8 +42,6 @@ Description:Containers that can be used to group related nodes and provide descr
 [/VAR]
 
 
-## Overview
-
 [EXCERPT:Intro]
 **Comments** can be added directly to single _Blueprint_ nodes, or they can be added as comment boxes to group related nodes and provide descriptions 
 about their functionality. These can be used solely for organizational purposes to make the 
