#UE4doc #docs: ##Overview removed as according to the instructional guide. Altered images to match the instructional guide: green -> yellow

diff --git a/Engine/Documentation/Source/Engine/Blueprints/UserGuide/Components/ComponentsList.INT.udn b/Engine/Documentation/Source/Engine/Blueprints/UserGuide/Components/ComponentsList.INT.udn
--- a/Engine/Documentation/Source/Engine/Blueprints/UserGuide/Components/ComponentsList.INT.udn
+++ b/Engine/Documentation/Source/Engine/Blueprints/UserGuide/Components/ComponentsList.INT.udn
@@ -47,12 +47,10 @@ Related: Engine/Blueprints/UserGuide/Functions
 
 [TOC (start:2 end:2)]
 
-## Overview
-
 [INCLUDE:Shared\Glossary\C#glossary_component]
 
 [REGION:note]
-For more information on **Components**, please see the [Components Overivew](Engine\Components) documentation. 
+For more information on **Components**, please see the [Components Overview](Engine\Components) documentation. 
 [/REGION]
 
 
