#UE4doc #docs Updated images and added version/skill level tags to the Programmer Quick Start tutorial.

diff --git a/Engine/Documentation/Source/Programming/QuickStart/1/QuickStart_1.INT.udn b/Engine/Documentation/Source/Programming/QuickStart/1/QuickStart_1.INT.udn
--- a/Engine/Documentation/Source/Programming/QuickStart/1/QuickStart_1.INT.udn
+++ b/Engine/Documentation/Source/Programming/QuickStart/1/QuickStart_1.INT.udn
@@ -2,6 +2,8 @@ Availability: Public
 Title:1. Essential Project Setup
 Crumbs:%ROOT%, Programming, Programming/QuickStart
 Description:Creating a new project with the specifications we need
+SkillLevel: Beginner
+Version: 4.8
 
 [VAR:Steps]
 [OBJECT:Navigation]
diff --git a/Engine/Documentation/Source/Programming/QuickStart/2/QuickStart_2.INT.udn b/Engine/Documentation/Source/Programming/QuickStart/2/QuickStart_2.INT.udn
--- a/Engine/Documentation/Source/Programming/QuickStart/2/QuickStart_2.INT.udn
+++ b/Engine/Documentation/Source/Programming/QuickStart/2/QuickStart_2.INT.udn
@@ -2,7 +2,8 @@ Availability: Public
 Title:2. Create a C++ Class
 Crumbs:%ROOT%, Programming, Programming/QuickStart
 Description:Creating a C++ class
-Version: 4.7
+SkillLevel: Beginner
+Version: 4.8
 
 [VAR:Steps]
 [OBJECT:Navigation]
@@ -22,9 +23,9 @@ Version: 4.7
 
 %Steps%
 
- 1. In the **Unreal Editor**, we can create a new C++ class with the **Add Code to Project** command, located in the **File** drop-down menu.
+ 1. In the **Unreal Editor**, we can create a new C++ class with the **New C++ Class...** command, located in the **File** drop-down menu.
  
-	[![](AddCPPCode.png)](AddCPPCode.png)
+	[![](AddCPPClass.png)](AddCPPClass.png)
 
  2. The **Choose Parent Class** menu will open. Since **Actor** is the most basic class that can exist in an **Unreal Engine** level, we'll use the **Actor** class as our base.
  
diff --git a/Engine/Documentation/Source/Programming/QuickStart/3/QuickStart_3.INT.udn b/Engine/Documentation/Source/Programming/QuickStart/3/QuickStart_3.INT.udn
--- a/Engine/Documentation/Source/Programming/QuickStart/3/QuickStart_3.INT.udn
+++ b/Engine/Documentation/Source/Programming/QuickStart/3/QuickStart_3.INT.udn
@@ -2,6 +2,8 @@ Availability: Public
 Title:3. Write and Compile C++ Code
 Crumbs:%ROOT%, Programming, Programming/QuickStart
 Description:Writing a C++ class and compiling it with Unreal Engine
+SkillLevel: Beginner
+Version: 4.8
 
 [VAR:Steps]
 [OBJECT:Navigation]
diff --git a/Engine/Documentation/Source/Programming/QuickStart/4/QuickStart_4.INT.udn b/Engine/Documentation/Source/Programming/QuickStart/4/QuickStart_4.INT.udn
--- a/Engine/Documentation/Source/Programming/QuickStart/4/QuickStart_4.INT.udn
+++ b/Engine/Documentation/Source/Programming/QuickStart/4/QuickStart_4.INT.udn
@@ -2,6 +2,8 @@ Availability: Public
 Title:4. Test Your Code
 Crumbs:%ROOT%, Programming, Programming/QuickStart
 Description:Creating a test object in your level
+SkillLevel: Beginner
+Version: 4.8
 
 [VAR:Steps]
 [OBJECT:Navigation]
diff --git a/Engine/Documentation/Source/Programming/QuickStart/5/QuickStart_5.INT.udn b/Engine/Documentation/Source/Programming/QuickStart/5/QuickStart_5.INT.udn
--- a/Engine/Documentation/Source/Programming/QuickStart/5/QuickStart_5.INT.udn
+++ b/Engine/Documentation/Source/Programming/QuickStart/5/QuickStart_5.INT.udn
@@ -2,6 +2,8 @@ Availability: Public
 Title:5. On Your Own!
 Crumbs:%ROOT%, Programming, Programming/QuickStart
 Description:Ideas for further exploration!
+SkillLevel: Beginner
+Version: 4.8
 
 [VAR:Steps]
 [OBJECT:Navigation]
