#UE4doc #docs: ##Overview removed as according to the instructional guide. Also corrected multiple repeated words: the the -> the, to to -> to

diff --git a/Engine/Documentation/Source/Engine/Blueprints/UserGuide/Events/Events.INT.udn b/Engine/Documentation/Source/Engine/Blueprints/UserGuide/Events/Events.INT.udn
--- a/Engine/Documentation/Source/Engine/Blueprints/UserGuide/Events/Events.INT.udn
+++ b/Engine/Documentation/Source/Engine/Blueprints/UserGuide/Events/Events.INT.udn
@@ -46,8 +46,6 @@ Description: Nodes that are called from gameplay code to begin execution of an i
 
 [TOC (start:2)]
 
-## Overview
-
 [EXCERPT:Intro]
 **Events** are nodes that are called from gameplay code to begin execution of an individual network 
 within the **EventGraph**. They allow _Blueprints_ to perform a series of actions in response to 
@@ -92,7 +90,7 @@ This Blueprint Event node executes only on the server. For single player games,
 
 ![LevelReset.png](LevelReset.png)
 
-The LevelReset events sends out an execution signal when the level restarts. This is useful when you 
+The LevelReset event sends out an execution signal when the level restarts. This is useful when you 
 need something to be triggered once a level has reloaded, such as in a gaming situation when the player 
 has died but the level does not need to reload.
 
diff --git a/Engine/Documentation/Source/Engine/Blueprints/UserGuide/FlowControl/FlowControl.INT.udn b/Engine/Documentation/Source/Engine/Blueprints/UserGuide/FlowControl/FlowControl.INT.udn
--- a/Engine/Documentation/Source/Engine/Blueprints/UserGuide/FlowControl/FlowControl.INT.udn
+++ b/Engine/Documentation/Source/Engine/Blueprints/UserGuide/FlowControl/FlowControl.INT.udn
@@ -44,8 +44,6 @@ Description:Nodes that allow for controlling the flow of execution based on cond
 
 [TOC (start:2 end:3)]
 
-## Overview
-
 [EXCERPT:Intro]
 
 [INCLUDE:Shared/Glossary/F/#Glossary_Flow_Control]
@@ -128,7 +126,7 @@ the color of a light to be red. If _false_, it sets it to blue.
 | -------- | ----------- |
 | [REGION:tablesection]Input Pins[/REGION] ||
 | **(Unlabeled)** | This execution input triggers the branch check. |
-| **Condition** | Takes in an boolean value used to indicate which output pin will be triggered. |
+| **Condition** | Takes in a boolean value used to indicate which output pin will be triggered. |
 | [REGION:tablesection]Output Pins[/REGION] ||
 | **True** | This outputs an execution pulse if the incoming condition was *true*. |
 | **False** | This outputs an execution pulse if the incoming condition was *false*. |
@@ -206,8 +204,8 @@ allowing you to track when Output A has been called.
 | [REGION:tablesection]Input Pins[/REGION] ||
 | **(Unlabeled)** | This execution input triggers the FlipFlop. |
 | [REGION:tablesection]Output Pins[/REGION] ||
-| **A** | This output pin is called the the first and every odd-numbered time thereafter that the FlipFlop is triggered. |
-| **B** | This output pin is called the the second and every even-numbered time thereafter that the FlipFlop is triggered. |
+| **A** | This output pin is called the first and every odd-numbered time thereafter that the FlipFlop is triggered. |
+| **B** | This output pin is called the second and every even-numbered time thereafter that the FlipFlop is triggered. |
 | **Is A** | Outputs a boolean value indicating whether Output A is being triggered or not. This, in effect, will toggle between *true* and *false* each time the FlipFlop node is triggered. |
 [/EXCERPT:Flipflop]
 
@@ -349,7 +347,7 @@ which, when played in order, reveal a special message.
 [EXCERPT:Sequence]
 ![Sequence_Example.png](Sequence_Example.png)
 
-The **Sequence** node allows for a single execution pulse to to trigger a series of events in order. The node 
+The **Sequence** node allows for a single execution pulse to trigger a series of events in order. The node 
 may have any number of outputs, all of which get called as soon as the Sequence node receives an input. They 
 will always get called in order, but without any delay. To a typical user, the outputs will likely appear 
 to have been triggered simultaneously.
