#UE4doc #docs:  Changed outline color.

diff --git a/Engine/Documentation/Source/Engine/Physics/WalkableSlope/WalkableSlope.INT.udn b/Engine/Documentation/Source/Engine/Physics/WalkableSlope/WalkableSlope.INT.udn
--- a/Engine/Documentation/Source/Engine/Physics/WalkableSlope/WalkableSlope.INT.udn
+++ b/Engine/Documentation/Source/Engine/Physics/WalkableSlope/WalkableSlope.INT.udn
@@ -42,7 +42,7 @@ Description:
 [/OBJECT]
 [/VAR]
 
-## Overview
+
 
 The **Walkable Slope Override** on **Physics Bodies** enable you to tweak what surfaces **Characters** can walk up. Perhaps a stair case is too steep or maybe you want to enforce the "no walking on the grass" signs, these settings will enable you to do so.
 
