#UE4doc #docs: Altered images to match the instructional guide: green -> yellow

diff --git a/Engine/Documentation/Source/Engine/Blueprints/UserGuide/Functions/BlueprintFunctions.INT.udn b/Engine/Documentation/Source/Engine/Blueprints/UserGuide/Functions/BlueprintFunctions.INT.udn
--- a/Engine/Documentation/Source/Engine/Blueprints/UserGuide/Functions/BlueprintFunctions.INT.udn
+++ b/Engine/Documentation/Source/Engine/Blueprints/UserGuide/Functions/BlueprintFunctions.INT.udn
@@ -265,7 +265,7 @@ Once we do that, we have a reference to our player character and can then drag o
 ![](CallFunction.png)
 [/REGION:fullwidth]
 
-We could then hook up the rest of the script to destory the projectile after it hits the player as so: 
+We could then hook up the rest of the script to destroy the projectile after it hits the player as so: 
 
 [REGION:fullwidth]
 ![](FullProjectileScript.png)
