#UE4doc #docs: Updated headers to match instructional guide formatting.

diff --git a/Engine/Documentation/Source/Engine/Blueprints/UserGuide/Functions/FunctionCalls/FunctionCalls.INT.udn b/Engine/Documentation/Source/Engine/Blueprints/UserGuide/Functions/FunctionCalls/FunctionCalls.INT.udn
--- a/Engine/Documentation/Source/Engine/Blueprints/UserGuide/Functions/FunctionCalls/FunctionCalls.INT.udn
+++ b/Engine/Documentation/Source/Engine/Blueprints/UserGuide/Functions/FunctionCalls/FunctionCalls.INT.udn
@@ -51,18 +51,18 @@ Level Blueprint itself. Function Calls are displayed as boxes with titles that s
 Different types of function calls have different color titles.
 [/EXCERPT:Intro]
 
-**Self**
+##Self
 
 **Self Function Calls** are functions that belong to the Blueprint itself, by way of being declared in 
 the class the Blueprint derives from or a parent class.
 
-**Other**
+##Other
 
 **Other Function Calls** are functions that belong to other Objects or Actors besides the Blueprint. For example, 
 the Blueprint might have a StaticMeshComponent that can have its mesh changed via a SetStaticMesh function call. 
 Since this function belongs to the StaticMeshComponent and not the Blueprint, it is an Other Function Call.
 
-**Pure**
+##Pure
 
 **Pure Function Calls** are special actions that can be performed that do not directly affect the world 
 or the Objects in it. These are generally things like functions that output a property value, or mathematical 
