#UE4Doc #docs #code_review richard.hinckley

diff --git a/Engine/Documentation/Source/Programming/QuickStart/1/QuickStart_1.INT.udn b/Engine/Documentation/Source/Programming/QuickStart/1/QuickStart_1.INT.udn
--- a/Engine/Documentation/Source/Programming/QuickStart/1/QuickStart_1.INT.udn
+++ b/Engine/Documentation/Source/Programming/QuickStart/1/QuickStart_1.INT.udn
@@ -21,6 +21,11 @@ Version: 4.8
 [/VAR]
 
 %Steps%
+
+[REGION:note]
+You should have a copy of **Visual Studio** installed before starting this tutorial.
+[/REGION]
+
 1. Open **Unreal Engine** from the Launcher. The [](Engine/Basics/Projects/Browser) will appear.
 
 1. Click on the **New Project** tab and then select the **C++** tab. From there, select **Basic Code** so we have a clean starting point, and make sure "With Starter Content" is set. We'll need to enter a project name, so we'll use "QuickStart". We can now click **Create Project** and get started.
diff --git a/Engine/Documentation/Source/Programming/QuickStart/QuickStart.INT.udn b/Engine/Documentation/Source/Programming/QuickStart/QuickStart.INT.udn
--- a/Engine/Documentation/Source/Programming/QuickStart/QuickStart.INT.udn
+++ b/Engine/Documentation/Source/Programming/QuickStart/QuickStart.INT.udn
@@ -39,7 +39,7 @@ Version: 4.8
 [/OBJECT]
 [/VAR]
 
-Over the course of this tutorial, we'll create a new **Unreal Engine** project,  add a new C++ class to it, then compile the project and add an instance of our new class to a level. When we're finished, we'll be able to see an **Actor** we programmed in C++ moving around on screen.
+----
 
 [OBJECT:EmbeddedVideo]
 	[PARAMLITERAL:width]
@@ -53,6 +53,21 @@ Over the course of this tutorial, we'll create a new **Unreal Engine** project,
 	[/PARAMLITERAL]
 [/OBJECT]
 
+## Goals
+
+The primary goal of this tutorial is to introduce programmers to Unreal Engine 4's development environment. By the end of this quick start guide, developers should be comfortable setting up and developing **C++ Projects** in Unreal Engine 4 (UE4).
+
+## Objectives
+
+Upon completion, developers will have the skill to:
+
+* **Create a New Unreal Engine Project**
+* **Add a New C++ Class**
+* **Compile a Project** 
+* **Add an Instance of a New Class to a Level**
+
+_When we're finished, we'll be able to see an **Actor** we programmed in C++ moving around on screen._
+
 [OBJECT:TopicButtonList]
 	[PARAM:icon]
 		![](%ROOT%/start_icon.png)(convert:false)
