#UE4doc #docs: ##Overview removed as according to the instructional guide. Altered images to match the instructional guide: green -> yellow

diff --git a/Engine/Documentation/Source/Engine/Blueprints/UserGuide/Nodes/Nodes.INT.udn b/Engine/Documentation/Source/Engine/Blueprints/UserGuide/Nodes/Nodes.INT.udn
--- a/Engine/Documentation/Source/Engine/Blueprints/UserGuide/Nodes/Nodes.INT.udn
+++ b/Engine/Documentation/Source/Engine/Blueprints/UserGuide/Nodes/Nodes.INT.udn
@@ -3,8 +3,6 @@ Title:Nodes
 Crumbs: %ROOT%, Engine/Blueprints, Engine/Blueprints/UserGuide
 Description:Node graph that uses events and function calls to perform actions in response to gameplay events associated with the Blueprint.
 
-## Overview
-
 
 **Nodes** are objects such as events, function calls, flow control operations, variables, etc. 
 that can be used in graphs to define the functionality of the particular graph and Blueprint that 
