Checking in the updated Quick Start guide and How To page based on the instructional guide.

diff --git a/Engine/Documentation/Source/Engine/QuickStart/1/LDQS_1.INT.udn b/Engine/Documentation/Source/Engine/QuickStart/1/LDQS_1.INT.udn
--- a/Engine/Documentation/Source/Engine/QuickStart/1/LDQS_1.INT.udn
+++ b/Engine/Documentation/Source/Engine/QuickStart/1/LDQS_1.INT.udn
@@ -1,9 +1,13 @@
-Availability: Public
-Title:1. Navigating the Viewport
+Availability:Public
+Title:1. Required Setup
 Crumbs:%ROOT%, GettingStarted, Engine/QuickStart
-Description:An overview on the most commonly used viewport navigation controls.
-
-
+Description:In this step we will use the Project Browser to create a New Project. 
+SkillLevel:Beginner
+Version:4.8
+Related:Engine/Basics/Projects
+Related:Engine/Basics/Projects/Browser
+Related:GettingStarted/Installation
+ 
 [VAR:Steps]
 [OBJECT:Navigation]
 	[PARAM:previous]
@@ -21,28 +25,27 @@ Description:An overview on the most commonly used viewport navigation controls.
 
 %Steps%
 
-[REGION:note]
-For this tutorial series, we are using a new **Blank** project with **Starter Content** enabled. Please see the [](Engine/Basics/Projects/Browser) page for information on creating new projects, using templates and defining project settings (including **Starter Content**). 
-[/REGION]
+Before we get started, the first thing that we will need to do is create a [**Project**](Engine/Basics/Projects) to work from for the duration of this guide. In the steps below we will use Unreal Engine 4's [**Project Browser**](Engine/Basics/Projects/Browser) to create our project. 
+After completing this step, our project will be ready and you will gone through the process of creating a Project and using a Template which you can then use for your very own future projects. 
 
-Once you have created your **Blank** project (with **Starter Content**), the initial level upon load will look something like this:
+## Steps
 
-![](Engine/QuickStart/2/TableAndChairs.png)
+1. After [**Installing Unreal Engine**](GettingStarted/Installation) and launching the Unreal Editor, you will first be presented with the **Project Browser**. 
 
-Using this little area as a point of reference, get used to the **Viewport Camera Controls**. Below are the 3 most commonly used methods of navigating the Viewport in Unreal Engine 4. Once you are comfortable with navigating around the level, head to the next page of the guide: **Place Actors**.
+1. Click the **New Project** tab and select the **Blank** Blueprint Template, enter a Save Location and Name, then click **Create Project**. 
 
-## Standard Controls
+	![](ProjectBrowser1.png)(w:720)
 
-[INCLUDE:Engine\UI\LevelEditor\Viewports\ViewportControls#StandardControls]
+	The **New Project** tab is where you can create new projects based off several different template types while the **Projects** tab contains any previously created projects or samples that you have downloaded. 
 
-## WASD Fly Controls
 
-[INCLUDE:Engine\UI\LevelEditor\Viewports\ViewportControls#GameControls]
+## End Result
 
-## Orbit, Dolly, and Track
+After creating your project, once the Unreal Editor has finished loading, you will be presented with the Main Editor Interface as shown below. 
 
-[INCLUDE:Engine\UI\LevelEditor\Viewports\ViewportControls#MayaLikeControls]
+![](ExampleLevel.png)
 
-For more information on Viewport Navigation, see: [**Viewport Controls**](Engine\UI\LevelEditor\Viewports\ViewportControls)
+In the next step, we will cover the basic controls for navigating the **Viewports** inside the Unreal Editor. 
 
 %Steps%
+
diff --git a/Engine/Documentation/Source/Engine/QuickStart/2/LDQS_2.INT.udn b/Engine/Documentation/Source/Engine/QuickStart/2/LDQS_2.INT.udn
--- a/Engine/Documentation/Source/Engine/QuickStart/2/LDQS_2.INT.udn
+++ b/Engine/Documentation/Source/Engine/QuickStart/2/LDQS_2.INT.udn
@@ -1,8 +1,12 @@
 Availability: Public
-Title:2. Start a New Level
+Title:2. Navigating the Viewport
 Crumbs:%ROOT%, GettingStarted, Engine/QuickStart
-Description:Creating an empty level via the New Level dialog window.
-
+Description:Learn the different methods you can use to navigate the different viewports inside Unreal Engine 4. 
+SkillLevel:Beginner
+Version:4.8
+Related:Engine/UI/LevelEditor/Viewports
+Related:Engine/Levels
+Related:Engine/UI/LevelEditor/Viewports/ViewportControls
 
 [VAR:Steps]
 [OBJECT:Navigation]
@@ -22,20 +26,37 @@ Description:Creating an empty level via the New Level dialog window.
 
 %Steps%
 
-Next, create a new level to build your game environment in. When creating a new level, just as when you created a new project, 
-Unreal Engine enables you to select from _templates_. By default, there is a _Default_ template - which has a very simple environment - 
-and an _Empty Level_ template - which is completely blank - to choose from.
+With the project open and ready to go, the first thing you may notice is the [**Viewport**](Engine/UI/LevelEditor/Viewports) in the center of the Unreal Editor. 
+
+![](TableAndChairs.png)(w:720)
+
+Inside the Viewport is where you will do most of your level construction. The template project that we selected in the previous step includes a small sample [**Level**](Engine/Levels) and some assets for us to get us started with. 
+Using this little area as a point of reference, take a moment to get used to the [**Viewport Camera Controls**](Engine/UI/LevelEditor/Viewports/ViewportControls) by using the most common methods of navigating the Viewport in Unreal Engine 4. 
+
+
+## Standard Controls
+
+[INCLUDE:Engine\UI\LevelEditor\Viewports\ViewportControls#StandardControls]
+
+[REGION:tip]
+To switch to an [**Orthographic View**](https://wiki.unrealengine.com/Videos/Player?series=PLZlv_N0_O1gasd4IcOe9Cx9wHoBB7rxFl&video=RoiQOwCg-4Q), click the **Perspective** drop-down then select an Orthographic View Mode. 
+
+![](OrthographicViews.png)
+[/REGION]
+
+## WASD Fly Controls
 
-Let's get started creating a new level:
+[INCLUDE:Engine\UI\LevelEditor\Viewports\ViewportControls#GameControls]
 
-1. From the File menu of the **Level Editor** window, choose ![](Engine/QuickStart/3/new_button.png) **New Level...**. This will display the **New Level** dialog:
+## Orbit, Dolly, and Track
 
-	![](Engine/QuickStart/3/level_new.png)
+[INCLUDE:Engine\UI\LevelEditor\Viewports\ViewportControls#MayaLikeControls]
 
-1. For the purposes of this guide, you need a completely blank slate to work with. Click on the **Empty Level** template to select it.
+## End Result
 
-1. The new blank level will be ready for you to populate with all the visual and gameplay related items you desire.
+You should now have an understanding of the different methods you can use to navigate the Viewport and have found one that you are comfortable with using to move the camera around in the level. 
+You will be moving the camera around quite a bit as you are building your levels, so feel free to switch between the different methods at any point to get the precision that you need. 
 
-	![](Engine/QuickStart/3/level_empty.png)
+In the Next Step, we will create a **New Level** so that we can build or own sample area and start populating it with assets. 
 
 %Steps%
diff --git a/Engine/Documentation/Source/Engine/QuickStart/3/LDQS_3.INT.udn b/Engine/Documentation/Source/Engine/QuickStart/3/LDQS_3.INT.udn
--- a/Engine/Documentation/Source/Engine/QuickStart/3/LDQS_3.INT.udn
+++ b/Engine/Documentation/Source/Engine/QuickStart/3/LDQS_3.INT.udn
@@ -1,7 +1,11 @@
-Availability: Public
-Title:3. Place Actors
+﻿Availability: Public
+Title:3. Create a New Level
 Crumbs:%ROOT%, GettingStarted, Engine/QuickStart
-Description:Placing and Manipulating Actors is the foundation of building any level.
+Description:In this step we will create a brand new level which we will use to build out our sample area in. 
+SkillLevel:Beginner
+Version:4.8
+Related:Engine/Levels
+Related:Engine/Content/Browser
 
 [VAR:Steps]
 [OBJECT:Navigation]
@@ -21,116 +25,34 @@ Description:Placing and Manipulating Actors is the foundation of building any le
 
 %Steps%
 
-The concept of "placement" is something that extends to many portions of the editor. It almost always involves clicking and dragging something into the level viewport.
+Next we will create a new [**Level**](Engine/Levels) that we will use to build our game environment in. 
+While there are several different ways in which you can create a new Level, we will use the most common method by using the [**Content Browser**](Engine/Content/Browser) which is used to create several different types of assets inside the Unreal Editor. 
 
-**Place Mode** in the **Modes Panel** enables you place many commonly used **Actors**, like lights and geometry. In the following steps, you will place a floor, Directional Light, Atmospheric Fog Actor and a Player Start Actor. 
 
-[REGION:tip]
-In the simplest of terms, an **Actor** is any object that can be placed into a level, and in your time creating levels with Unreal Engine 4, you will place, move, and edit many **Actors**.
+## Steps
 
-For more information on **Actors**, see: [**Actors & Geometry**](Engine\Actors).
-[/REGION]
+1. Inside the Unreal Editor, locate the **Content Browser** window and click the **Add New** Button then select the **Level** asset. 
 
-First you will need a floor to place everything on, for this, use a Box Geometry Actor.
+	![](ContentBrowserLevel.png)
 
-1. In the **Modes Panel**, with **Place** mode enabled, click the **BSP** Category.
-1. Place a Box Geometry Actor from the **Modes Panel** by clicking and dragging it into the level.
+	This will open the **New Level** dialog window:
 
-	![](Engine/QuickStart/4/placeBox.png)
+	![](level_new.png)(w:550)
 
-	_Upon releasing the **Left Mouse Button**, the **Box Brush** will be added to the level._
+	Above, the **Default** level includes some of the commonly used assets for constructing levels while **Empty Level** is a completely blank level with no assets. For the purposes of this guide we are going to start from scratch with a completely blank slate. 
 
-	![](place2.png)
+1. Click the **Empty Level** to select it.
 
-1. With the Box selected, in the **Details** panel (lower-right window of the editor), set **Location** and **Rotation** all to **0**. 
-
-1. Set the **Scale** to **4 x 4 x 0.10** to create a floor looking surface. 
-
-	![](resize_floor.png)
-
-1. With the floor placed and positioned, switch to the **Lights** category in the **Modes Panel**.
-
-1. Place a **Directional Light** into the level, on top of the floor box.
-
-	![](Engine/QuickStart/4/placeLight.png)
-	
-1. To make it easier to find, **Move** it up above the floor by switching to the **Translation Tool** ("**W**" key).
-
-1. Click and drag the Z Axis (Blue) gizmo on the **Translation Tool**, up from the surface of the floor.
-
-	![](Engine/QuickStart/4/moveLight.png)
-
-	_If the Directional Light is not selected (1), you can **Left-click** on it in the viewport to select it (2). With the light selected, you can then click on and drag the Z Axis upward away from the floor to move the light upwards (3)._ 
-
-1. Switch to the **Visual Effects** category in the **Modes Panel**.
-
-1. Place an **Atmospheric Fog** into the level.
-	
-	![](Engine/QuickStart/4/placeFog.png)
-	
-	_The **Atmospheric Fog** Actor will add a basic sky to the scene and the viewport will become lit instead of black._
-
-1. Switch to the **Basic** category in the **Modes Panel**.
-
-1. Place a **Player Start** in the level.
-
-	![](Engine/QuickStart/4/placeStart.png)
-
-1. The **Player Start Actor** in the image above is facing down the X Axis (Red) of the Translation gizmo, which means it is pointing off the edge of the Floor. You will need to rotate it so it faces towards the center of the floor by using the **Rotation Tool**. Enable the **Rotation Tool** by pressing the "**E**" key inside the viewport.
-1. Rotate the **Player Start Actor** by clicking and dragging on the Z Axis (Blue) arc gizmo so that the light blue arrow on the **Player Start Actor** faces in towards the center of the floor.
-
-	![](rotate1.png) ![](rotate2.png)(h:320)
-
-1. Switch to the **Volumes** category in the **Modes Panel**.
-1. Place a **Lightmass Importance Volume** in the level.
-
-	![](Engine/QuickStart/4/placeVolume.png)
-
-	_The **Lightmass Importance Volume** is used to control and concentrate lighting and shadowing affects within the volume. See [**Lightmass Basics**](Engine/Rendering/LightingAndShadows/Lightmass/Basics) for more information._
-	
-1. The **Lightmass Importance Volume** will be the same size as the floor, but we need it to be larger so it can encompass everything we are going to place in the level. Switch to the **Scale Tool** by pressing the "**R**" key.
-
-1. Click and drag on each axis of the **Scale Tool** so that the **Lightmass Importance Volume** is larger than the floor box.
-
-	![](Engine/QuickStart/4/scaleVolume.png)
-
-	_Above the viewport's render mode has been set to **Unlit** (**Alt+3**) to see the volume more clearly (press **Alt+4** to return to **Lit**)._
-
-At this point, all the Actors you have placed have come from the **Modes Panel**, but you can also place things from the **Content Browser**. 
-
-Locate the **Content Browser**, it looks like this:
-
-![](Engine/QuickStart/4/contentBrowser.png)
-
-From here, you can place things like Static Meshes, Skeletal Meshes, Blueprints, and you can even place Materials onto Actors already placed in the level. For more information on the **Content Browser**, see: [**Content Browser**](Engine\Content\Browser). 
-
-Let's add some **Static Meshes** to the level: 
 	
-1. Included in the **Starter Content** is a props folder. **Left-click** on the **Props** folder in the **Asset Tree** view (highlighted below). 
+## End Result
 
-	![](Engine/QuickStart/4/openProps.png)
+The level will be created and ready to populate with visual and gameplay related assets to construct our playable environment. 
 
-	_The **Content Window** (right-side of the **Content Browser**) will become populated with assets that are located inside that folder._
-
-1. Scroll down and find the Static Mesh asset named **SM_TableRound**.
-
-	![](Engine/QuickStart/4/findSMTableRound.png)
-
-1. Click and drag the **SM_TableRound** Static Mesh asset into the level.
-
-	![](Engine/QuickStart/4/placeTable.png)
-
-1. Using what you have learned, make a small scene by adding more Actors from the **Modes Panel** and **Content Browser**.  
-
-	![](Engine/QuickStart/4/placeThings.png)
-
-	_Place some **Point Lights**, more **Props** and some walls with a roof from the **Architecture** folder!_
-
-In the next section, you will learn how to edit the Actors you have placed in the level.
-
-%Steps%
+![](level_empty.png)(w:640)
 
+In the next step, we will begin adding assets and building out our environment. 
 
+%Steps%	
 
 
 
diff --git a/Engine/Documentation/Source/Engine/QuickStart/4/LDQS_4.INT.udn b/Engine/Documentation/Source/Engine/QuickStart/4/LDQS_4.INT.udn
--- a/Engine/Documentation/Source/Engine/QuickStart/4/LDQS_4.INT.udn
+++ b/Engine/Documentation/Source/Engine/QuickStart/4/LDQS_4.INT.udn
@@ -1,7 +1,15 @@
-Availability: Public
-Title:4. Edit Placed Actors
+﻿Availability: Public
+Title:4. Placing Actors in the Level
 Crumbs:%ROOT%, GettingStarted, Engine/QuickStart
-Description:With some Actors in the Level, it's now time to add some variation.
+Description:Here we will begin placing and manipulating Actors in our level to build out our environment. 
+SkillLevel:Beginnner
+Version:4.8 
+Related:Engine/UI/LevelEditor/Modes/PlaceMode
+Related:GettingStarted/HowTo/ManipulatingActors
+Related:Engine/UI/LevelEditor/Details
+Related:Engine/Actors
+Related:Engine/Rendering/LightingAndShadows
+
 
 [VAR:Steps]
 [OBJECT:Navigation]
@@ -21,39 +29,92 @@ Description:With some Actors in the Level, it's now time to add some variation.
 
 %Steps%
 
-![](Engine\QuickStart\4\placeThings.png)
+In this step we will begin placing [**Actors**](Engine/Actors) (which are things like lights or geometry for example) into our empty level. 
+We will cover the two most common ways of adding Actors to the level, one of which is through [**Place Mode**](Engine/UI/LevelEditor/Modes/PlaceMode) and the other is through the [**Content Browser**](Engine/Content/Browser). 
+After completing this step, you will know how to place Actors inside your own levels and can begin manipulating those Actors to create your environment. 
+
+## Steps
 
-With a number of Actors placed, it is time to edit a few of them. With the first being the **Directional Light Actor**.
+1. In the **Modes Panel**, with **Place Mode** enabled, click the [**BSP**](Engine/Actors/Brushes) category and select the **Box**. 
 
-1. Select the **Directional Light Actor** by **Left-clicking** on it in the Viewport.
+	![](ModesPanel.png)
 
-	![](Engine\QuickStart\5\selectDL.png)
+1. **Left-click** and drag the **Box** into the **Level Viewport**. 
 
-1. In the **Details Panel** under the **Light** category, enable **Used as Atmosphere Sun Light**:
+	![](placeBox.png)
 
-	![](Engine\QuickStart\5\enableSun.png)
+	Upon releasing the **Left Mouse Button**, the **Box** will be added to the level.
 
-	Depending on the rotation of your **Directional Light Actor** the sky color will change, and if you rotate the Viewport around you will see that the sun now aligns with the **Directional Light Actor**. This is a real time process so you can rotate the **Directional Light Actor** (press **E** to switch to **Rotation Mode**) and the sky will change color from night, sunrise, day time, and sunset.
+	![](place2.png)
 
-Now change the **Material** on one of your placed **Static Mesh Actors** by first selecting it. 
+1. With the **Box** still selected, in the [**Details**](Engine/UI/LevelEditor/Details) panel (lower-right window of the editor), set **Location** and **Rotation** all to **0**. 
 
-1. With your Actor selected, in the **Details** panel under **Materials**, click the drop-down box under **Element 0**. 
+1. Set the **Scale** to **4 x 4 x 0.1**. 
 
-	![](material_1.png)(w:900)
+	![](resize_floor.png)
 
-	In the pop-up window, select the new **Material** for the Static Mesh to use (below we selected **M_Brick_Clay_New**. 
+	We will use this as the floor in which the player can walk around on. 
 
-	![](material_2.png)(w:900)
+1. In the **Modes Panel** select the **Lights** tab, then drag-and-drop a **Directional Light** into the level on top of the floor. 
+
+	![](placeLight.png)
 	
-1.	All **Actors** in your level have many properties for you to adjust. Explore, and when you are ready to continue, click **Next Step**. 
+1. On the on the [**Translation Tool**](https://docs.unrealengine.com/latest/INT/Engine/UI/LevelEditor/Viewports/ViewportControls/index.html#transformcontrols), click and drag the Z Axis (Blue) gizmo up away from the surface of the floor.
 
-	![](Engine\QuickStart\5\beforeBake.png)(w:900)
+	![](moveLight.png)
 
-	_Try changing the color of your lights and adding more Materials to your Actors._
+	If the **Directional Light** becomes unselected, you can re-select it by **Left-clicking** on it in the **Level Viewport**. 
 
-%Steps%
+1. In the **Modes Panel**, select the **Visual Effects** tab and drag-and-drop an [**Atmospheric Fog**](Engine/Actors/FogEffects/AtmosphericFog) to the level. 
+
+	![](placeFog.png)
+	
+	The **Atmospheric Fog** Actor will add a basic sky to the level and the level will become illuminated instead of dark. 
+
+1. From the **Modes Panel** on the **Basic** tab, drag-and-drop a [**Player Start**](Engine/Actors/PlayerStart) into the level. 
+
+	![](placeStart.png)
+
+1. From the **Modes Panel** on the **Volumes** tab, drag-and-drop a [**Lightmass Importance Volume**](Engine/Rendering/LightingAndShadows/Lightmass/Basics) into the level. 
+
+	![](placeVolume.png)
+
+	The **Lightmass Importance Volume** is used to control and concentrate lighting and shadowing affects within the volume. When placing the Lightmass Importance Volume in the level, the default size of the volume does not cover our playable area so we will need to scale it up. 
+	
+1. Inside the **Level Viewport**, press **R** to switch to the [**Scale Tool**](https://docs.unrealengine.com/latest/INT/Engine/UI/LevelEditor/Viewports/ViewportControls/index.html#transformcontrols). 
+
+1. Click and drag the white box in the center of the **Scale Tool** so that the **Lightmass Importance Volume** encapsulates the floor. 
 
+	![](scaleVolume.png)
 
+1. Inside the **Content Browser** under **Content > StarterContent > Props**, drag-and-drop the **SM_TableRound** into the level. 
 
+	![](placeTable.png)
 
+	Try to place the table in the center of the floor using the **Move Tool** (Keyboard Shortcut **W** if it is not selected). 
 
+1. Also under **Content > StarterContent > Props**, drag-and-drop the **SM_Chair** into the level. 
+
+	![](addChair.png)
+
+1. With the **SM_Chair** selected inside the **Level Viewport**, press **E** to access the [**Rotation Tool**](https://docs.unrealengine.com/latest/INT/Engine/UI/LevelEditor/Viewports/ViewportControls/index.html#transformcontrols). 
+	
+	![](Rotate1.png)(w:440)
+
+1. **Left-click** and drag the **Blue** axis arc (the gizmo will update to show degrees) and rotate the chair to face the table. 
+
+	![](Rotate2.png)(w:440)
+
+1. Using the placement methods above, create a small scene by adding more Actors from the **Modes Panel** and **Content Browser**.  
+
+	![](Engine/QuickStart/5/placeThings.png)
+
+	Try adding some lights, props, walls and a roof (found under the **Content > StarterContent > Architecture** folder). 
+
+## End Result
+
+You should now have a Level that is populated with several different Actors types from both the **Modes** panel and **Content Browser**. You also have manipulated those Actors with the **Move**, **Rotate** and **Scale** tools. Depending on the Actors you placed in the final step, you should have something similar to above. 
+
+In the next section, you will learn how to edit the Actors you have placed in the level and customize their properties. 
+
+%Steps%
diff --git a/Engine/Documentation/Source/Engine/QuickStart/5/LDQS_5.INT.udn b/Engine/Documentation/Source/Engine/QuickStart/5/LDQS_5.INT.udn
--- a/Engine/Documentation/Source/Engine/QuickStart/5/LDQS_5.INT.udn
+++ b/Engine/Documentation/Source/Engine/QuickStart/5/LDQS_5.INT.udn
@@ -1,9 +1,13 @@
-Availability: Public
-Title:5. Running the Build Process
+﻿Availability: Public
+Title:5. Editing Placed Actors
 Crumbs:%ROOT%, GettingStarted, Engine/QuickStart
-Description:With a number of Actors placed and customized, one of the final steps for creating any level is the Build Process.
-
-
+Description:This this step we will take the Actors we placed and edit their properties to give our level a more customized look. 
+SkillLevel:Beginner
+Version:4.8
+Related:Engine/Rendering/Materials
+Related:Engine/Rendering/LightingAndShadows/LightTypes
+Related:Engine/Content/FBX/ImportOptions
+Related:
 
 [VAR:Steps]
 [OBJECT:Navigation]
@@ -23,39 +27,47 @@ Description:With a number of Actors placed and customized, one of the final step
 
 %Steps%
 
-![](Engine\QuickStart\5\beforeBake.png)(w:940)
-
-By now you have noticed the "Preview" labels in the shadows and the light leaking under walls. This is because all the lights in the scene are static and use precomputed, or baked lighting, which has not been calculated yet. The "Preview" text is there to remind you that what you are seeing in the viewport currently is not what you will see in the game. To get the scene looking better, we need to run the **Build** process.
+With several different **Actors** placed inside our level, the next step involves [**Editing Actor Properites** ](GettingStarted/HowTo/DetailsPanel) which can change the look or the way the Actor functions in the level, giving us a more customized looking level. 
+We will start with editing the properties of our **Directional Light Actor** then shift our focus to applying [**Materials**](Engine/Rendering/Materials) to some of the [**Static Mesh Actors**](Engine/Actors/StaticMeshActor) that you have placed in your level. 
 
-**Building** in Unreal Editor is fairly straight forward: choose your quality setting and click the ![**Build**](Engine\QuickStart\6\build_button.png) button.
+Once you have finished this step, you will have seen where to access and modify the properties of Actors so that you can begin editing and experimenting with different settings inside your own levels. 
 
+## Steps
 
+1. Select the **Directional Light Actor** by **Left-clicking** on it in the Viewport.
 
-1. To choose your quality level, click the ![**Down Arrow**](Engine\QuickStart\6\downArrow_button.png) next to the ![**Build**](Engine\QuickStart\6\build_button.png) button to expand the **Build Options** then expand the **Lighting Quality** menu:
+	![](selectDL.png)
 
-	![](Engine\QuickStart\6\buildSettings.png)
+1. In the **Details Panel** under the **Light** category, enable **Used as Atmosphere Sun Light**:
 
-1. Set the **Quality Level** to **Production**.
+	![](enableSun.png)
 
-	[REGION:note]
-	Before beginning a **Build**, remember, it can take a while on larger levels. Our small hut may not take too long, but a scene with hundreds of lights that covers many square kilometers will take substantially longer.
-	[/REGION]
+	Depending on the rotation of your **Directional Light Actor** the sky color will change and if you rotate the Viewport around you will see that the sun now aligns with the Directional Light Actor. 
+	This is a real time process so you can rotate the Directional Light Actor (press **E** to switch to **Rotation Mode**) and the sky will change color from night, sunrise, day time, and sunset.
 
-1. Click the ![**Build**](Engine\QuickStart\6\build_button.png) button.
-1. Wait for the **Build** to complete.
+Next we will change the **Material** on one of your placed **Static Mesh Actors** by first selecting it. 
 
-	![](Engine\QuickStart\6\wait.png)
+1. With your Actor selected, in the **Details** panel under **Materials**, click the drop-down box under **Element 0**. 
 
-1. When it completes, your lighting should automatically update:
+	![](material_1.png)(w:900)
 
-![](Engine\QuickStart\6\built.png)(w:940)
+1. In the pop-up window, select the **M_Brick_Clay_New** Material. 
 
+	![](material_2.png)(w:900)
+	
+1.	All **Actors** in your level have many properties for you to adjust inside the **Details** panel. Explore changing their settings!
 
-Finally, remember that **Player Start** you placed? If you hit the ![**Play**](Engine\QuickStart\6\play_button.png) button, you will be able to fly around your level using the **Mouse** and **WASD** keys. While in Play mode, your game is actually running and the renderer is exactly the same thing you would see in a stand alone client.
-
-%Steps%
+	![](beforeBake.png)(w:900)
+	
+	Try changing the **Light Color** of your lights, applying more **Materials** or changing the **Scale** of the Actors in your level. 
 
+## End Result
 
+You have seen how to apply **Materials** to **Static Mesh Actors** in your level as well as how to access the **Properties** of an **Actor**. 
+After experimenting with the various properties, your sample level should start to resemble the image above and look more complete. 
+Feel free to continue changing the various properties of each Actor in your level.  
 
+When you are ready, proceed to the next step where we will run through the **Build** process which will update the settings in our level. 
 
+%Steps%
 
diff --git a/Engine/Documentation/Source/Engine/QuickStart/6/LDQS_6.INT.udn b/Engine/Documentation/Source/Engine/QuickStart/6/LDQS_6.INT.udn
--- a/Engine/Documentation/Source/Engine/QuickStart/6/LDQS_6.INT.udn
+++ b/Engine/Documentation/Source/Engine/QuickStart/6/LDQS_6.INT.udn
@@ -1,9 +1,9 @@
-Availability: Public
-Title:6. On Your Own!
+﻿Availability: Public
+Title:6. Running the Build Process
 Crumbs:%ROOT%, GettingStarted, Engine/QuickStart
-Description:Take what you have learned and polish up your level, then continue on to read related topics. 
-
-
+Description:With a number of Actors placed and customized, one of the final steps for creating any level is the Build Process.
+SkillLevel:Beginner
+Version:4.8
 
 [VAR:Steps]
 [OBJECT:Navigation]
@@ -16,50 +16,54 @@ Description:Take what you have learned and polish up your level, then continue o
 		[Quick Start Home](Engine/QuickStart)
 	[/PARAM]
 	[PARAM:next]
+		[Next Step](Engine/QuickStart/7 "%Engine/QuickStart/7:title%")
 	[/PARAM]
 [/OBJECT]
 [/VAR]
 
 %Steps%
 
-Using what you have learned over the course of this Quick Start guide, try to do the following: 
+By now you may have noticed the "Preview" labels in the shadows and the light leaking under walls. 
 
-![](update1.png)(w:830)
+![](beforeBake2.png)
 
-* Change the lighting of the level to a moonlight, night scene. 
-* Add another room, attached to the first room. 
+This is because all the lights in the scene are static and use precomputed, or baked lighting, which has not been calculated yet. The "Preview" text is there to remind you that what you are seeing in the viewport currently is not what you will see in the game. 
 
-![](update2.png)(h:380) ![](update4.png)(h:380)
+In this step we will go through the **Build** process which will build all levels (precompute lighting data and visibility data, generate any navigation networks and update all brush models). We will also take a look at **Light Quality** settings inside of the **Build Options** which we can use to adjust the quality of our lighting when it is built. 
 
-* On the attached room, try to make it elevated and join them with stairs. 
-* Add some bushes, a couch, shelves and a front door. 
+## Steps
 
-![](update6.png)(w:830)
+1. From the Main Tool Bar, click the down-arrow next to the **Build** option. 
 
-* Add different kinds of lights with different colors. 
-* Use different Materials on some of your Actors. 
+	![](buildOptionsButton.png)
 
-For more information on the topics covered in this quick start guide, and across the entire editor, see the [**Unreal Editor Manual**](Engine).
+1.	Under **Lighting Quality**, choose the **Production** setting. 
 
-As for specifics covered in this quick start:
+	![](productionLighting.png)
 
-* For more information on the Level Editor, see: [**Level Editor**](Engine\UI\LevelEditor)
+	This will give us the highest quality lighting but is the slowest in terms of computation time and will increase the time it takes to **Build** the game. Our level is small, so it should not impact us too much, but keep this in mind when you are working on larger levels as you may want to leave it on a mid-low level setting while creating your level and switching it to **Production** on a "final pass" at your level. 
 
-* For more information on Viewports, see: [**Viewports**](Engine\UI\LevelEditor\Viewports)
+1. Wait for the **Build** to complete.
 
-* For more information on the Editing Modes available in Unreal Editor, 4 see: [**Level Editor Modes**](Engine\UI\LevelEditor\Modes)
+	![](wait.png)
 
-* For more information on the **Content Browser**, see: [**Content Browser**](Engine\Content\Browser)
+	You will see the progress in the lower-right corner of the Unreal Editor as seen in the image above. 
+	
+	Once the **Build** process has completed, the level lighting will update giving you a better indication of the final result. 
 
-* For more information on the **Details** Panel, see: [**Details Panel**](Engine\UI\LevelEditor\Details)
+	![](built.png)
 
-* For more information on Building, see: [**Lightmass**](Engine\Rendering\LightingAndShadows\Lightmass)
+1. From the Main Toolbar, click the **Play** Button to play in the editor. 
 
-* For more information on Lighting, see: [**Lighting Quick Start Guide**](Engine/Rendering/LightingAndShadows/QuickStart)
+	![](PlayInEditor.png)
 
-%Steps%
+	Using **WASD** to move and the **Mouse** to turn the camera, you can fly around your level. 
 
 
+## End Result
 
+At this point you should have **Built** the level lighting and previewed your game with the **Play** in Editor feature. Each of the steps leading up to this point have been aimed at getting you quickly up to speed on how to perform the most common actions when constructing levels inside the Unreal Editor. 
 
+In the next step, you will be given some examples of how you can take this sample level and build upon it using the methods that were provided throughout the course of this guide. Also, additional links to more detailed documentation related to the topics covered in this guide are provided which will serve as reference when constructing levels on your own. 
 
+%Steps%
\ No newline at end of file
diff --git a/Engine/Documentation/Source/Engine/QuickStart/7/LDQS_7.INT.udn b/Engine/Documentation/Source/Engine/QuickStart/7/LDQS_7.INT.udn
--- a/Engine/Documentation/Source/Engine/QuickStart/7/LDQS_7.INT.udn
+++ b/Engine/Documentation/Source/Engine/QuickStart/7/LDQS_7.INT.udn
@@ -1,9 +1,9 @@
-Availability: Docs
-Title:7. Further Reading
+Availability: Public
+Title:7. On Your Own!
 Crumbs:%ROOT%, GettingStarted, Engine/QuickStart
-Description:Further reading on the topics covered in the Content Quick Start Guide.
-
-
+Description:Using the methods provided during this guide, polish up your level with additional content and read up on related topics. 
+SkillLevel:Beginner
+Version:4.7
 
 [VAR:Steps]
 [OBJECT:Navigation]
@@ -22,7 +22,22 @@ Description:Further reading on the topics covered in the Content Quick Start Gui
 
 %Steps%
 
----
+Using the methods that were provided during the course of this guide, try do to the following on your own: 
+
+![](update1.png)(w:830)
+
+* Change the lighting of the level to a moonlight, night scene. 
+* Add another room, attached to the first room. 
+
+![](update2.png)(h:380) ![](update4.png)(h:380)
+
+* On the attached room, try to make it elevated and join them with stairs. 
+* Add some bushes, a couch, shelves and a front door. 
+
+![](update6.png)(w:830)
+
+* Add different kinds of lights with different colors. 
+* Use different Materials on some of your Actors. 
 
 For more information on the topics covered in this quick start guide, and across the entire editor, see the [**Unreal Editor Manual**](Engine).
 
@@ -40,6 +55,7 @@ As for specifics covered in this quick start:
 
 * For more information on Building, see: [**Lightmass**](Engine\Rendering\LightingAndShadows\Lightmass)
 
----
+* For more information on Lighting, see: [**Lighting Quick Start Guide**](Engine/Rendering/LightingAndShadows/QuickStart)
 
 %Steps%
+
diff --git a/Engine/Documentation/Source/Engine/QuickStart/EditorQuickStart.INT.udn b/Engine/Documentation/Source/Engine/QuickStart/EditorQuickStart.INT.udn
--- a/Engine/Documentation/Source/Engine/QuickStart/EditorQuickStart.INT.udn
+++ b/Engine/Documentation/Source/Engine/QuickStart/EditorQuickStart.INT.udn
@@ -39,9 +39,9 @@ Version: 4.8
 [/OBJECT]
 [/VAR]
 
-![](Engine\QuickStart\6\built.png)(w:940)
+![](Finished.png)
 
-The Unreal Editor Quick Start Guide walks you through the basics of working with Unreal Engine 4 including: navigating the viewports, creating a new level, placing and editing Actors in your level, as well as the process of building and running your levels. At the end of this guide, you will have a room similar to the one pictured above. The final step of this guide will give you some ideas to take it and make it your own as well as provide links to related documentation. 
+The Unreal Editor Quick Start Guide walks you through the basics of working with Unreal Engine 4 including: navigating the viewports, creating a new level, placing and editing Actors in your level, as well as the process of building and running your levels. At the end of this guide, you will have a room similar to the one pictured above. The final step of this guide will give you some ideas to take it and make it your own using what you have learned as well as provide links to additional documentation. 
 
 [OBJECT:TopicButtonList]
 	[PARAM:icon]
@@ -59,6 +59,7 @@ The Unreal Editor Quick Start Guide walks you through the basics of working with
 	* [](Engine/QuickStart/4 "%Engine/QuickStart/4:description%")
 	* [](Engine/QuickStart/5 "%Engine/QuickStart/5:description%")
 	* [](Engine/QuickStart/6 "%Engine/QuickStart/6:description%")
+	* [](Engine/QuickStart/7 "%Engine/QuickStart/7:description%")
 	[/PARAM]
 [/OBJECT]
 
diff --git a/Engine/Documentation/Source/Gameplay/HowTo/PossessPawns/Blueprints/GHT_4B.INT.udn b/Engine/Documentation/Source/Gameplay/HowTo/PossessPawns/Blueprints/GHT_4B.INT.udn
--- a/Engine/Documentation/Source/Gameplay/HowTo/PossessPawns/Blueprints/GHT_4B.INT.udn
+++ b/Engine/Documentation/Source/Gameplay/HowTo/PossessPawns/Blueprints/GHT_4B.INT.udn
@@ -1,272 +1,101 @@
 Availability: Public
-Title:How to Possess Pawns in Blueprint
+Title:Possessing Pawns in Blueprint
 Crumbs: %ROOT%, Gameplay, Gameplay/HowTo/PossessPawns
-Description:A How To Guide for switching between characters or possessing different types of pawns in Blueprints.
-
-[TOC (start:2)]
-
-## Overview
-
-On this page, you will learn about **Pawn Possessing** and how to switch between multiple characters in a level. You will also learn how to switch from a 3rd person character and enter/exit a drivable vehicle through **Blueprints**. 
+Description:Illustrates using the Possess node in Blueprints to switch between and control multiple characters. 
+Version: 4.8
+SkillLevel: Beginner
+Related: Engine/Blueprints
+Related: Gameplay/Framework/Pawn
+Related: Gameplay/Framework/Pawn/Character
+Related: Engine/Basics/Projects/Browser
+Related: Gameplay/Framework/Controller/PlayerController
+
+Learn how to use the **Possess** and **UnPossess** [Blueprint](Engine/Blueprints) nodes to take control of and provide input to a [Pawn](Gameplay/Framework/Pawn) or [Character](Gameplay/Framework/Pawn/Character) in your project.  
+
+[REGION:note]
+For the purposes of this guide we are using the **Blueprint Side Scroller** template, however you can use your own project if you would like. If you do not know how to create projects or use templates, refer to the [Project Browser](Engine/Basics/Projects/Browser) page for more information. 
+[/REGION]
  
-## Switching Between Multiple Characters
-
-Below you will find steps you can take to switch between multiple characters using the **Level Blueprint**. 
-
-_The steps below are based off using the **Blueprint Side Scroller Template** project, however you can use any project you wish._
-
-1.	Inside your project, drag into your level the additional characters you would like to control. 
-
-	![](GHT4B_1.png)
-
-	Here we added two **SideScrollerCharacter** Blueprints to the level. 
-
-1.	Select each of your characters in the level (the two we added to the level). 
-
-	![](GHT4B_3.png)
-
-1.	From the Main Toolbar, click the **Blueprints** button and **Open Level Blueprint**. 
-
-	![](GHT4B_2.png)
-
-1.	**Right-click** in the graph area and select the **Add references to selected Actors** option. 
-
-	![](GHT4B_4.png)
-
-	This will allow us to reference our characters and possess them when we press a button. 
-
-1.	**Right-click** in the graph and add an **Event Begin Play** node, **Right-click** again and add a **Get Player Character** node.  
-
-1.	Drag off the **Return Value** of the **Get Player Character** and add the **Cast To SideScrollerCharacter** node. 
-
-	![](GHT4B_5.png)
-
-	We are going to store a reference to the default character so that we may switch back to them when we press a button. 
-
-1.	Drag off the **As SideScrollerCharacter** pin and select **Promote to variable** and call the variable **DefaultCharacter**. 
-
-	![](GHT4B_6.png)
-
-1.	**Right-click** in the graph and search for and add the **1** Key Event, **2** Key Event and **3** Key Event nodes. 
-
-1.	While holding down **Control**, drag in the **DefaultCharacter** variable we created above. 
-
-	![](GHT4B_7.png)
-
-	When we press 1, we will switch to the default character, pressing 2 will switch to the second and 3 to the third characters. 
-
-1.	**Right-click** in the graph and search for and add the **Possess** node. 
-
-	![](GHT4B_8.png)
+## Steps
 
-	You may need to un-check the **Context Sensitive** option in order for the node to appear in the list. 
+1. From the **Content Browser**, drag-and-drop into your level the additional characters you would like to control. 
 
-1.	Connect each of the characters to the **In Pawn** pin and the **Key Events** to each of the **Possess** nodes. 
+	![](Pawns1b.png)(w:860)
 
-	![](GHT4B_9.png)
-
-1.	**Right-click** in the graph and add a **Get Player Controller** node and connect it to each of the **Target** pins. 
-
-	![](GHT4B_10.png)(w:760)
-
-	The **Possess** nodes targets a **Player Controller** and here we are assigning the default Player Controller to target (if you had a Player Controller Blueprint, you could **Cast To** it and hook it up to the target. For this example, we are using the default controller). 
-
-1.	Click the **Compile** button to compile the Blueprint, then from the Main Toolbar, click the **Play** button to play in the editor. 
-
-	![](GHT4B_11.png)
-
-When you play in the editor, you should be able to press the 1, 2, or 3 keys to cycle between the three characters in the level. 
-
-
-
-
-## Entering/Exiting a Vehicle
-
-Below you will find steps that can be used to have a third person character enter/exit a drivable vehicle. 
-
-_This example is based off using the **Blueprint 3rd Person and Vehicle Template** projects, however you can use any projects you wish._
-
-1.	Create a **Blueprint Vehicle Template** and **Blueprint 3rd Person Template** project. 
-
-1.	Inside the **Blueprint 3rd Person** project, under **Game/Blueprints**, **Right-click** on the **ThirdPersonCharacter** Blueprint and choose **Migrate**(located within **Asset Actions**). 
-
-	![](GHT4B_12.png)
-
-	You want to migrate this Blueprint to the **MyDocuments/Unreal Projects/ProjectName/Content** folder of your **Vehicle** project. 
-
-1.	Inside the **Vehicle Project**, open the **World Settings** option from the Main Toolbar. 
-
-1.	Under **Game Mode**, set the **Default Pawn Class** to use the **ThirdPersonCharacter** Blueprint.
-
-	![](GHT4B_13.png)
-
-	This is indicating that the default player character will use the 3rd Person character instead of the vehicle. 
-
-1.	Under the **Project Settings**, find the **Input** section and add the Bindings highlighted below.
-
-	![](Bindings.png)
-
-	These are the Bindings that were used and are needed to control the 3rd Person Character. Also make sure that the **Scale** values are set as shown as well (you can change these later if you would like to fit your preferred control method). 
+	Here we added two additional **SideScrollerCharacter** Blueprints from the **Content > SideScrollerBP > Blueprints** folder to the level giving us a total of three characters in our level. With this particluar template project, a character is already placed inside the level by default. 
 	
-1.	In the **Content Browser** under **Game/Sedan**, open the **Sedan** Blueprint. 
-
-1.	Click the **Components** tab, then in the **Components** window, click **Add Component** and add an **Arrow** and **Box**
-
-1.	Name the Arrow **SpawnLocation** and the Box **Trigger** and move them near the driver's side door. 
-
-	![](GHT4B_15.png)
-
-	The Arrow will define the location of the character when they exit the vehicle while the Box acts as a trigger, determining if the character is close enough to the vehicle to enter it. When inside the trigger, the player can press a button and enter the vehicle, while outside the trigger they will not be able to enter the vehicle. 
+1.	In the **World Outliner**, select each of the characters you want to control. 
 
-	![](GHT4B_16.png)
+	![](Pawns2b.png)
 
-	Here we moved the box trigger near the driver side door and the arrow just outside of the box. 
+1. In the **World Outliner**, select each of the characters you want to control. 
 
-1.	For the **SpawnLocation** (Arrow), in the **Details** panel under **Location**, set the **Z** to **131.0**. 
-
-1.	In the **Content Browser**, under the **Game/Blueprints** folder, open the **ThirdPersonCharacter** Blueprint. 
-
-1.	Add a variable called **Car** of the **WheeledVehicle** type and a **Bool** variable called **NearCar**. 
-
-	![](GHT4B_17.png)
-
-	Here we are going to get reference to the car and assign it to the **Car** variable. We are also going to set **NearCar** to _true_ when the player enters the **Trigger** in the **Sedan** Blueprint which will determine if we allow the player to press a button to enter the car. 
-
-1.	Return to the **Sedan** Blueprint and add a **Character** variable called **Player**. 
-
-	![](GHT4B_18.png)
-
-1.	**Right-click** on the **Trigger** box and add both the **OnComponentBeginOverlap** and **OnComponentEndOverlap** events. 
-
-	![](GHT4B_19.png)
-
-1.	In the graph, drag off the **Other Actor** pin of the **OnComponentBeginOverlap** node and add the **Cast To ThirdPersonCharacter** node. 
-
-1.	Off the **As My Character C** pin, add the **Set Near Car** variable and click the check box to set it to **True**. 
-
-	![](GHT4B_20.png)(w:780)
-
-	By using the **Cast To** node, we are able to access and update the variables in our **ThirdPersonCharacter** Blueprint. 
-
-1.	Drag off **As My Character C** again and add the **Set Car** node, then drag off the **Car** pin and add a **Get Reference To Self** node. 
-
-	[REGION:lightbox]
-	[![](GHT4B_21.png)(w:780)](GHT4B_21.png)
-	[/REGION]
-
-	_Click image for full view._
-
-	Here we created the reference to the **Sedan** Blueprint and assigned it to the variable in our **ThirdPersonCharacter** Blueprint. You can also add a **Print String** node to the end with some text to indicate that you are in the **Trigger** box (we set this to **Enter Car**). 
-
-1.	Drag off the **Other Actor** pin of the **OnComponentEndOverlap** node and **Cast To ThirdPersonCharacter**. 
-
-1.	Off the **As My Character C** pin, add the **Set Near Car** node and leave the checkbox unchecked which sets it to **False**. 
-
-	![](GHT4B_22.png)
-
-1.	In the **MyBlueprint** window, while holding down **Control**, drag in the **SpawnLocation** variable and the **Player** variable. 
-
-	![](GHT4B_29.png)
-
-1.	**Right-click** in the graph and add the **E** Key Event and connect it to a **Set Actor Transform** node. 
-
-1.	Drag off the **SpawnLocation** variable, un-check **Context Sensitive**, and add the **Get World Transform** node. 
-
-1.	Connect the **Get World Transform** to the **Set Actor Transform** node. 
-
-1.	Connect the **Player** variable to the **Target** on the **Set Actor Transform** node. 
-
-	![](GHT4B_23.png)
-
-	The last few steps should result in something similar to above (do not worry about the additional wires on the Player node yet). 
-
-1.	Drag off the **Player** variable and add a **Set Actor Hidden in Game** node (leave the checkbox unchecked). 
-
-1.	Drag off the **Player** variable again and add a **Set Actor Enable Collision** node (check the checkbox to enable collision). 
-
-1.	**Right-click** in the graph and add a **Possess** node.
-
-	![](GHT4B_30.png)
-
-	To access the **Possess** node, you may need to uncheck **Context Sensitive** for it to appear in the list. 
-
-1.	Connect the **Player** variable to the **In Pawn** of the **Possess** node, then drag off **Target** and add a **Get Player Controller** node.
-
-1.	Connect all the remaining wires as shown below. 
-
-	[REGION:lightbox]
-	[![](GHT4B_24.png)(w:780)](GHT4B_24.png)
-	[/REGION]
-
-	_Click image for full view._
-		
-1.	Open the **ThirdPersonCharacter** Blueprint, in the graph add a **E** Key Event and connect the **Pressed** pin to a **Branch** node. 
-
-1.	Hold **Control** and drag in the **NearCar** variable and connect it to the **Branch** condition. 
-
-1.	**Right-click** in the graph and add a **Possess** node (uncheck **Context Sensitive** if it does not appear in the list). 
-
-1.	Hold **Control** and drag in the **Car** variable and connect it to the **In Pawn** pin on the **Possess** node. 
-
-1.	Drag off the **Target** of the **Possess** node and add a **Get Player Controller** node, you graph should look similar to below. 
+	![](Pawns2b.png)
+	
+1. From the Main Toolbar, click the **Blueprints** button and **Open Level Blueprint**. 
 
-	![](GHT4B_25.png)
+	![](Pawns3b.png)
 
-1.	Following the **Possess** node, add a **Set Actor Hidden in Game** node (set to _True_) then a **Set Actor Enable Collision** (set to _False_). 
+1. **Right-click** in the graph area, then select **Create References to selected Actors** option from the context menu. 
 
-1.	Drag off the **Car** variable, re-check the **Context Sensitive** option. 
+	![](Pawns4b.png)(w:640)
 
-1.	Search for and add a **Cast To Sedan** node following the **Set Actor Enable Collision** node. 
+	This will allow us to reference each of the characters that we've selected in the level. 
 
-1.	Drag off the **As Sedan C** pin and add the **Set Player** node, then drag off **Player** and add a **Get Reference To Self** node.
+1. **Right-click** in the graph area, then search for the **1** Keyboard Event and select it. 
 
-1.	Your full sequence in **ThirdPersonCharacter** should look like below. 
+	![](Pawns5b.png)(w:640)
 
-	[REGION:lightbox]
-	[![](GHT4B_26.png)(w:780)](GHT4B_26.png)
-	[/REGION]
+	This will allow us to fire off an event whenever the 1 key is pressed or released. 
 
-	_Click image for full view._
+1. Add the **2** and **3** Keyboard Events to the graph. 
 
-	Here we are saying when **E** is pressed (if **NearCar** is _True_ which is set by entering the trigger box in the **Sedan** Blueprint), take possession of the car and set the **ThirdPersonCharacter** Actor to hidden in game and remove its collision.  We are then telling the **Sedan** Blueprint that its **Player** variable in the Blueprint should reference the **ThirdPersonCharacter** Blueprint so that we can return to **ThirdPersonCharacter** when exiting the vehicle. 
+1. **Right-click** in the graph and search for and add the **Get Player Controller** node. 
 
-1.  If there is a Sedan already placed in the level, delete it. Then, from the **Content Browser**, drag the **Sedan** Blueprint into the level. 
+	![](Pawns6b.png)(w:640)
 
-1.	If you click **Compile** and then **Play** from the Main Toolbar to play in the editor, you should end up with something similar to below. 
+	The [**Player Controller**](Gameplay/Framework/Controller/PlayerController) is used to take the input from a human player and translate that into actions for a Pawn. 
 
-	![](GIF1.gif)(convert:false)
+1. **Left-click** and drag off the **Return Value** of the Get Player Controller node and search for and add the **Possess** node. 
 
-	[REGION:note]
-	This example is intended to serve as a starting point to demonstrate how you could approach entering/exiting a vehicle. With this implementation, issues may arise when exiting the vehicle on angled surfaces or exiting a vehicle when going off jumps.
-	[/REGION]
+	![](Pawns7b.png)(w:640)
 
-	Inside the **Sedan** Blueprint, you can add the additional section of script to **Break Transform** and zero out the **Pitch** and **Roll** of the **Spawn Location** marker, this should help orientate the character when exiting the vehicle while it is angled or airborne. 
+	This will tell the Player Controller which target Pawn to possess and take control of. When the **Possess** function is called, it will automatically check if a Pawn is currently controlled and **UnPossess** it first before attempting to Possess a new Pawn. 
 	
-	[REGION:lightbox]
-	[![](GHT4B_Add.png)(w:780)](GHT4B_Add.png)
-	[/REGION]
+	You can use the UnPossess function if you want to have the player release control of their Pawn and enter, say for example, a spectator type of state where they are not directly controlling a playable character. 
 
-	_Click image for full view._
+1. Create two more **Possess** nodes and connect the **1**, **2** and **3** Keyboard Events as shown below. 
 
+	![](Pawns8.png)(w:640)
 	
+	We are now set up to possess a Pawn when 1, 2 or 3 is pressed. Next we need to define which Pawn to possess from our references. 
 
+1. Connect each of the **SideScrollerCharacter** references to each of the **In Pawn** pins as shown below. 
 
+	![](Pawns9.png)(w:640)
 
+	Our scripted functionality is complete and we are ready to test it. 
 
+1. Click **Compile** from the toolbar to update your script. 
 
+	![](Pawns10b.png)
 
+1. Click **Play** from the toolbar to play the game in the editor. 
 
-
-
-
-
+	![](Pawns11b.png)
 	
-
-
-
-
-
-
-
-
-
+## End Result
+
+When you press 1, 2 or 3 on your keyboard you will switch between each of the characters in your level. 
+
+[OBJECT:EmbeddedVideo]
+[PARAMLITERAL:width]
+640
+[/PARAMLITERAL]
+[PARAMLITERAL:height]
+360
+[/PARAMLITERAL]
+[PARAMLITERAL:videoid]
+RFgeogZOnvQ
+[/PARAMLITERAL]
+[/OBJECT]
