#UE4doc #docs: Fixed spelling errors: dialog, receive, possibility, earlier

diff --git a/Engine/Documentation/Source/Engine/Blueprints/UserGuide/OnlineNodes/OnlineNodes.INT.udn b/Engine/Documentation/Source/Engine/Blueprints/UserGuide/OnlineNodes/OnlineNodes.INT.udn
--- a/Engine/Documentation/Source/Engine/Blueprints/UserGuide/OnlineNodes/OnlineNodes.INT.udn
+++ b/Engine/Documentation/Source/Engine/Blueprints/UserGuide/OnlineNodes/OnlineNodes.INT.udn
@@ -92,7 +92,7 @@ This node is also used for leaving sessions. Clients also need to call this node
 
 ## Error Handling Nodes
 
-Unfortunately, when dealing with network operations, errors are always a possiblity and game should handle them gracefully. Perhaps the host of the session you try to join has just quit a moment eariler, or perhaps your internet service went down temporarily, or any number of other problems occur. Through Blueprints, your games have more flexibility in how they handle these network-related errors. 
+Unfortunately, when dealing with network operations, errors are always a possibility and game should handle them gracefully. Perhaps the host of the session you try to join has just quit a moment earlier, or perhaps your internet service went down temporarily, or any number of other problems occur. Through Blueprints, your games have more flexibility in how they handle these network-related errors. 
 
 Failures that originate from the Session Operations described above will be reported through the **OnFailure** execution pin, so you can set up your game to handle these failures as you wish. For certain kinds of failures however, or failures that occur in the middle of gameplay, there are a couple of Error Events that can be called within a **GameInstance** Blueprint, **Network Error** and **Travel Error**. 
 
@@ -100,7 +100,7 @@ Failures that originate from the Session Operations described above will be repo
 
 ![](EventNetworkError.png)
 
-The **Network Error** Event allows games to respond to the event, recieve an enumeration value that describes the type of error, and respond appropriately - a game may want to display a dailog box to the user, for example. 
+The **Network Error** Event allows games to respond to the event, receive an enumeration value that describes the type of error, and respond appropriately - a game may want to display a dialog box to the user, for example. 
 
 ### Event Travel Error
 
