#UE4doc #docs: ##Overview removed as according to the instructional guide. Altered images to match the instructional guide: green -> yellow

diff --git a/Engine/Documentation/Source/Engine/Blueprints/UserGuide/Timelines/Creation/CreatingTimelines.INT.udn b/Engine/Documentation/Source/Engine/Blueprints/UserGuide/Timelines/Creation/CreatingTimelines.INT.udn
--- a/Engine/Documentation/Source/Engine/Blueprints/UserGuide/Timelines/Creation/CreatingTimelines.INT.udn
+++ b/Engine/Documentation/Source/Engine/Blueprints/UserGuide/Timelines/Creation/CreatingTimelines.INT.udn
@@ -3,8 +3,6 @@ Title: Creating Timelines
 Crumbs: %ROOT%, Engine, Engine/Blueprints, Engine/Blueprints/UserGuide, Engine/Blueprints/UserGuide/Timelines
 Description: This document contains an overview of how to create Timeline nodes in Blueprints.
 
-## Overview
-
 %Description%
 
 ## Creating Timelines
