#UE4doc #docs: ##Overview removed as according to the instructional guide.

diff --git a/Engine/Documentation/Source/Engine/Blueprints/UserGuide/Timelines/Examples/LightFlickering/LightFlickering.INT.udn b/Engine/Documentation/Source/Engine/Blueprints/UserGuide/Timelines/Examples/LightFlickering/LightFlickering.INT.udn
--- a/Engine/Documentation/Source/Engine/Blueprints/UserGuide/Timelines/Examples/LightFlickering/LightFlickering.INT.udn
+++ b/Engine/Documentation/Source/Engine/Blueprints/UserGuide/Timelines/Examples/LightFlickering/LightFlickering.INT.udn
@@ -3,9 +3,6 @@ Title:Flickering Lights
 Crumbs: %ROOT%, Engine, Engine/Blueprints, Engine/Blueprints/UserGuide, Engine/Blueprints/UserGuide/Timelines
 Description: An example use of a timeline in which we set up a flickering light that also changes color.
 
-
-## Overview 
-
 %Description%
 
 ## Flickering Light
