#UE4doc #docs: ##Overview removed as according to the instructional guide.

diff --git a/Engine/Documentation/Source/Engine/Blueprints/UserGuide/Timelines/Examples/OpeningDoors/OpeningDoors.INT.udn b/Engine/Documentation/Source/Engine/Blueprints/UserGuide/Timelines/Examples/OpeningDoors/OpeningDoors.INT.udn
--- a/Engine/Documentation/Source/Engine/Blueprints/UserGuide/Timelines/Examples/OpeningDoors/OpeningDoors.INT.udn
+++ b/Engine/Documentation/Source/Engine/Blueprints/UserGuide/Timelines/Examples/OpeningDoors/OpeningDoors.INT.udn
@@ -3,8 +3,6 @@ Title:Opening Doors
 Crumbs: %ROOT%, Engine, Engine/Blueprints, Engine/Blueprints/UserGuide, Engine/Blueprints/UserGuide/Timelines
 Description: An example use of a timeline in which we set up a classic proximity-based opening door.
 
-## Overview
-
 %Description%
 
 ## Level Setup
