#UE4doc #docs: ##Overview removed as according to the instructional guide. Fixed a spelling error: documenation -> documentation

diff --git a/Engine/Documentation/Source/Engine/Blueprints/UserGuide/Types/Interface/InterfaceBlueprint.INT.udn b/Engine/Documentation/Source/Engine/Blueprints/UserGuide/Types/Interface/InterfaceBlueprint.INT.udn
--- a/Engine/Documentation/Source/Engine/Blueprints/UserGuide/Types/Interface/InterfaceBlueprint.INT.udn
+++ b/Engine/Documentation/Source/Engine/Blueprints/UserGuide/Types/Interface/InterfaceBlueprint.INT.udn
@@ -43,8 +43,6 @@ Description:Blueprints that declare functions to define an interface between Blu
 
 [TOC (start:2)]
 
-## Overview
-
 [EXCERPT:Intro]
 A **Blueprint Interface** is a collection of one or more functions - name only, no implementation - that can be added to 
 other _Blueprints_. Any _Blueprint_ that has the _Interface_ added is guaranteed to have those functions. The functions 
@@ -139,7 +137,7 @@ To edit a function's signature:
 
 ## Reference Links
 
-Refer to the following links for additional documenation regarding **Blueprint Communication**: 
+Refer to the following links for additional documentation regarding **Blueprint Communication**: 
 
 * For information on **Blueprint Communication Guidelines**, see: [**Blueprint Communication Usage**](Engine\Blueprints\UserGuide\BlueprintCommsUsage) 
 
