#UE4doc #docs: ##Overview removed as according to the instructional guide.

diff --git a/Engine/Documentation/Source/Engine/Blueprints/UserGuide/Types/LevelBlueprint/LevelScript.INT.udn b/Engine/Documentation/Source/Engine/Blueprints/UserGuide/Types/LevelBlueprint/LevelScript.INT.udn
--- a/Engine/Documentation/Source/Engine/Blueprints/UserGuide/Types/LevelBlueprint/LevelScript.INT.udn
+++ b/Engine/Documentation/Source/Engine/Blueprints/UserGuide/Types/LevelBlueprint/LevelScript.INT.udn
@@ -45,8 +45,6 @@ Related: Engine/Blueprints
 
 [TOC (start:2 end:3)]
 
-## Overview
-
 [EXCERPT:Intro]
 A **Level Blueprint** is a specialized type of **Blueprint** that acts as a level-wide global event graph. 
 Each level in your project has its own _Level Blueprint_ created by default that can be edited within the Unreal Editor, however new _Level Blueprints_ 
diff --git a/Engine/Documentation/Source/Engine/Blueprints/UserGuide/Types/MacroLibrary/MacroLibrary.INT.udn b/Engine/Documentation/Source/Engine/Blueprints/UserGuide/Types/MacroLibrary/MacroLibrary.INT.udn
--- a/Engine/Documentation/Source/Engine/Blueprints/UserGuide/Types/MacroLibrary/MacroLibrary.INT.udn
+++ b/Engine/Documentation/Source/Engine/Blueprints/UserGuide/Types/MacroLibrary/MacroLibrary.INT.udn
@@ -43,8 +43,6 @@ Description:Blueprints that declare and define macros for reuse in other Bluepri
 
 [TOC(start:2)]
 
-## Overview
-
 [EXCERPT:Intro]
 A **Blueprint Macro Library** is a container that holds a collection of _Macros_ or self-contained graphs that can 
 be placed as nodes in other _Blueprints_. These can be time-savers as they can store commonly used sequences of nodes 
