#UE4doc #docs: ##Overview removed as according to the instructional guide.

diff --git a/Engine/Documentation/Source/Engine/Blueprints/UserGuide/UserConstructionScript/UserConstructionScript.INT.udn b/Engine/Documentation/Source/Engine/Blueprints/UserGuide/UserConstructionScript/UserConstructionScript.INT.udn
--- a/Engine/Documentation/Source/Engine/Blueprints/UserGuide/UserConstructionScript/UserConstructionScript.INT.udn
+++ b/Engine/Documentation/Source/Engine/Blueprints/UserGuide/UserConstructionScript/UserConstructionScript.INT.udn
@@ -49,8 +49,6 @@ Description:Executed when instances of a Blueprint are created to perform initia
 
 [TOC(start:2 end:3)]
 
-## Overview
-
 [EXCERPT:Intro]
 
 [INCLUDE:Shared/Glossary/C/#Glossary_Construction_Script]
