#UE4doc #docs: ##Overview removed as according to the instructional guide. Altered images to match the instructional guide: green -> yellow. Fixed minor spelling errors.

diff --git a/Engine/Documentation/Source/Engine/Blueprints/UserGuide/Variables/Variables.INT.udn b/Engine/Documentation/Source/Engine/Blueprints/UserGuide/Variables/Variables.INT.udn
--- a/Engine/Documentation/Source/Engine/Blueprints/UserGuide/Variables/Variables.INT.udn
+++ b/Engine/Documentation/Source/Engine/Blueprints/UserGuide/Variables/Variables.INT.udn
@@ -44,8 +44,6 @@ Description: Properties that hold a value or reference an Object or Actor in the
 
 [TOC (start:2)]
 
-## Overview
-
 [EXCERPT:Intro]
 **Variables** are properties that hold a value or reference an Object or Actor in the world. These 
 properties can be accessible internally to the **Blueprint** containing them, or they can be made 
@@ -112,7 +110,7 @@ Creating a variable inside a Blueprint can be achieved by following the steps be
 
 	![](HT5.png)
 
-	Enter a name for your Blueprint, then proceed to the **Details** panel to adjust the properites of the variable.
+	Enter a name for your Blueprint, then proceed to the **Details** panel to adjust the properties of the variable.
 
 1.  In the **Details** panel, there are several settings that can be used to define how your variable is used or accessed. 
 
@@ -153,7 +151,7 @@ Here you can set whether or not the variable is serialized and is zero-filled at
 
 [EXCERPT:editingvars]
 
-When a variable is set to **Editable**, this allows the variable to modified from outside of the Blueprint that the variable resides in. 
+When a variable is set to **Editable**, this allows the variable to be modified from outside of the Blueprint that the variable resides in. 
 
 [REGION:raw]
 ![](HT21.png)
@@ -334,7 +332,7 @@ Get nodes supply portions of the network with a variable value. Once created, th
 
 #### Set Nodes
 
-Set nodes allow you to change the value of a variable. Note that these must be called by an execution wire in order to exectute.
+Set nodes allow you to change the value of a variable. Note that these must be called by an execution wire in order to execute.
 
 ![](SetVariable.png)
 
