#UE4doc #docs: fixed a spelling error: partilcuar -> particular

diff --git a/Engine/Documentation/Source/Engine/Components/Components.INT.udn b/Engine/Documentation/Source/Engine/Components/Components.INT.udn
--- a/Engine/Documentation/Source/Engine/Components/Components.INT.udn
+++ b/Engine/Documentation/Source/Engine/Components/Components.INT.udn
@@ -51,7 +51,7 @@ Each Component could then be accessed in this manner (turning on the lights, hon
 
 Contrary to the default behavior of sub-objects in general, Components created as sub-objects within an Actor are instanced, meaning each Actor instance of a particular class gets its own unique instances of the Components. 
 A simple way to visualize this is to imagine the vehicle described above. A Car class might use Components to represent the wheels of the car. 
-Four Wheel Components would be created as sub-objects in the default properties of the class and assigned to a "Wheels" array. When a new Car instance is created, new instances of the Wheel Components are created specifically for that particluar Car. 
+Four Wheel Components would be created as sub-objects in the default properties of the class and assigned to a "Wheels" array. When a new Car instance is created, new instances of the Wheel Components are created specifically for that particular Car. 
 If this were not the case, when one Car in the world moved, the wheels of all Cars would turn; which is clearly not the desired behavior. Instancing of Components by default simplifies the process of quickly adding unique sub-objects to Actors.
 
 [REGION:note]
