#UE4doc #docs: fixed a spelling error: properites -> properties

diff --git a/Engine/Documentation/Source/Engine/Components/Camera/CameraComponents.INT.udn b/Engine/Documentation/Source/Engine/Components/Camera/CameraComponents.INT.udn
--- a/Engine/Documentation/Source/Engine/Components/Camera/CameraComponents.INT.udn
+++ b/Engine/Documentation/Source/Engine/Components/Camera/CameraComponents.INT.udn
@@ -25,7 +25,7 @@ The **SpringArmComponent** tries to maintain its children at a fixed distance fr
 
 _Above a CameraComponent (Camera) and SpringArmComponent (red line) are used to create a 3rd person perspective for a character._
 
-You can adjust several camera related properites on the SpringArmComponent such as the **TargetArmLength** which is the natural length of the spring arm when there are no collisions, **Probe Size** or how big should the probe sphere be when checking for collisions, as well as things like **CameraLag** useful if you want the camera to lag slightly behind the object it is attached to. 
+You can adjust several camera related properties on the SpringArmComponent such as the **TargetArmLength** which is the natural length of the spring arm when there are no collisions, **Probe Size** or how big should the probe sphere be when checking for collisions, as well as things like **CameraLag** useful if you want the camera to lag slightly behind the object it is attached to. 
 
 
 
