#UE4doc #docs Updated Programming Quick Start Guide title enumeration.

diff --git a/Engine/Documentation/Source/Programming/QuickStart/1/QuickStart_1.INT.udn b/Engine/Documentation/Source/Programming/QuickStart/1/QuickStart_1.INT.udn
--- a/Engine/Documentation/Source/Programming/QuickStart/1/QuickStart_1.INT.udn
+++ b/Engine/Documentation/Source/Programming/QuickStart/1/QuickStart_1.INT.udn
@@ -1,5 +1,5 @@
 Availability: Public
-Title:1. Essential Project Setup
+Title:1 - Essential Project Setup
 Crumbs:%ROOT%, Programming, Programming/QuickStart
 Description:Creating a new project with the specifications we need
 SkillLevel: Beginner
diff --git a/Engine/Documentation/Source/Programming/QuickStart/2/QuickStart_2.INT.udn b/Engine/Documentation/Source/Programming/QuickStart/2/QuickStart_2.INT.udn
--- a/Engine/Documentation/Source/Programming/QuickStart/2/QuickStart_2.INT.udn
+++ b/Engine/Documentation/Source/Programming/QuickStart/2/QuickStart_2.INT.udn
@@ -1,5 +1,5 @@
 Availability: Public
-Title:2. Create a C++ Class
+Title:2 - Create a C++ Class
 Crumbs:%ROOT%, Programming, Programming/QuickStart
 Description:Creating a C++ class
 SkillLevel: Beginner
diff --git a/Engine/Documentation/Source/Programming/QuickStart/3/QuickStart_3.INT.udn b/Engine/Documentation/Source/Programming/QuickStart/3/QuickStart_3.INT.udn
--- a/Engine/Documentation/Source/Programming/QuickStart/3/QuickStart_3.INT.udn
+++ b/Engine/Documentation/Source/Programming/QuickStart/3/QuickStart_3.INT.udn
@@ -1,5 +1,5 @@
 Availability: Public
-Title:3. Write and Compile C++ Code
+Title:3 - Write and Compile C++ Code
 Crumbs:%ROOT%, Programming, Programming/QuickStart
 Description:Writing a C++ class and compiling it with Unreal Engine
 SkillLevel: Beginner
diff --git a/Engine/Documentation/Source/Programming/QuickStart/4/QuickStart_4.INT.udn b/Engine/Documentation/Source/Programming/QuickStart/4/QuickStart_4.INT.udn
--- a/Engine/Documentation/Source/Programming/QuickStart/4/QuickStart_4.INT.udn
+++ b/Engine/Documentation/Source/Programming/QuickStart/4/QuickStart_4.INT.udn
@@ -1,5 +1,5 @@
 Availability: Public
-Title:4. Test Your Code
+Title:4 - Test Your Code
 Crumbs:%ROOT%, Programming, Programming/QuickStart
 Description:Creating a test object in your level
 SkillLevel: Beginner
diff --git a/Engine/Documentation/Source/Programming/QuickStart/5/QuickStart_5.INT.udn b/Engine/Documentation/Source/Programming/QuickStart/5/QuickStart_5.INT.udn
--- a/Engine/Documentation/Source/Programming/QuickStart/5/QuickStart_5.INT.udn
+++ b/Engine/Documentation/Source/Programming/QuickStart/5/QuickStart_5.INT.udn
@@ -1,5 +1,5 @@
 Availability: Public
-Title:5. On Your Own!
+Title:5 - On Your Own!
 Crumbs:%ROOT%, Programming, Programming/QuickStart
 Description:Ideas for further exploration!
 SkillLevel: Beginner
