#UE4doc #docs: Unbolded text

diff --git a/Engine/Documentation/Source/Engine/Rendering/LightingAndShadows/LightFunctions/LightFunctions.INT.udn b/Engine/Documentation/Source/Engine/Rendering/LightingAndShadows/LightFunctions/LightFunctions.INT.udn
--- a/Engine/Documentation/Source/Engine/Rendering/LightingAndShadows/LightFunctions/LightFunctions.INT.udn
+++ b/Engine/Documentation/Source/Engine/Rendering/LightingAndShadows/LightFunctions/LightFunctions.INT.udn
@@ -46,15 +46,15 @@ Description:Lighting projected using a Material as the source.
 ![](Engine/Rendering/LightingAndShadows/lighting_overview_banner_001.png)(convert:false)
 [/REGION]
 
-## Overview
+
 
 [EXCERPT:LightFunctions]
-**Light Functions** are - in essence - a **Material** that can be applied to filter a light's intensity.  
+**Light Functions** are - in essence - a Material that can be applied to filter a light's intensity.  
 
-It should be noted however that you cannot change the color of lights using a **Light Function**, only the **Light Color** setting can do that. 
+It should be noted however that you cannot change the color of lights using a Light Function, only the Light Color setting can do that. 
 [/EXCERPT:LightFunctions]
 
-In the image below (_click it to view it full size_), a **Material** uses **Panner** and **ScalarParameter** nodes to move a texture while generating a flickering effect. 
+In the image below (_click it to view it full size_), a Material uses Panner and ScalarParameter nodes to move a texture while generating a flickering effect. 
 [REGION:lightbox]
 [![](LI_LightFunction.png)(w:900)](LI_LightFunction.png)
 [/REGION]
@@ -62,13 +62,13 @@ In the image below (_click it to view it full size_), a **Material** uses **Pann
 
 ## Creating a Light Function 
 
-You can create a **Light Function** by setting the **Material Domain** of a **Material** to use the **Light Function** attribute.
+You can create a Light Function by setting the Material Domain of a Material to use the Light Function attribute.
 
 ![](MAT_LightFunction.png)(w:520)
 
 
 ## Applying a Light Function 
 
-You can apply a **Light Function** to a light from inside the **Details** panel under the **Light Function Material** setting. 
+You can apply a Light Function to a light from inside the **Details** panel under the Light Function Material setting. 
 
 ![](AddLightFunction.png)(w:640)
