#UE4doc #docs: Updated out of date character image.

diff --git a/Engine/Documentation/Source/Engine/Components/Rendering/RenderingComponents.INT.udn b/Engine/Documentation/Source/Engine/Components/Rendering/RenderingComponents.INT.udn
--- a/Engine/Documentation/Source/Engine/Components/Rendering/RenderingComponents.INT.udn
+++ b/Engine/Documentation/Source/Engine/Components/Rendering/RenderingComponents.INT.udn
@@ -126,7 +126,7 @@ While you specify the vectors for the position and tangent on the Component itse
 
 ![](splinemesh2.png)
 
-Above we can move the **Start Transform** and **End Transform** independently of moving the entire Actor, allowing us to strech or rotate it as we desire. This example and setup can be found inside the [](Resources/ContentExamples/Blueprint_Splines) Content Examples map. 
+Above we can move the **Start Transform** and **End Transform** independently of moving the entire Actor, allowing us to stretch or rotate it as we desire. This example and setup can be found inside the [](Resources/ContentExamples/Blueprint_Splines) Content Examples map. 
 
 ## Text Render Component
 
