#UE4 Docs
Fix for Jira UEDOC-1789

Desc: Updated to add skill tables for the landing page and link to Persona landing page.

diff --git a/Engine/Documentation/Source/Engine/Animation/AnimHowTo/AnimHowToOverview.INT.udn b/Engine/Documentation/Source/Engine/Animation/AnimHowTo/AnimHowToOverview.INT.udn
--- a/Engine/Documentation/Source/Engine/Animation/AnimHowTo/AnimHowToOverview.INT.udn
+++ b/Engine/Documentation/Source/Engine/Animation/AnimHowTo/AnimHowToOverview.INT.udn
@@ -3,22 +3,40 @@ Title:Animation How To's
 Crumbs: %ROOT%, Engine, Engine\Animation, Engine\Animation\AnimHowTo
 Description:Examples of different aspects of the Animation system in Unreal Engine 4 are demonstrated through several step-by-step How To guides.
 
-The **Animation How To's** pages provide several step-by-step examples of how to work with and use different aspects of the Animation system in Unreal Engine 4. Whether you are new to Animation in UE4 or want to learn how to work with certain aspects of the Animation system, the How To pages below cover a broad range of topics and should give you a place to start or expand upon your current knowledge. 
+The **Animation How To's** provide several step-by-step examples of working with and using different aspects of [**Persona**](Engine/Animation), which is the animation tool inside Unreal Engine 4. 
+Whether you are new to animation in Unreal Engine 4 or want to learn how to work with certain aspects of Persona, the How To pages below cover a broad range of topics and aim to give you a place to start or expand upon your current knowledge. 
 
-Refer to the table below for more information and click each link to view its contents. 
+Refer to the tables below for more information and click each link to view its contents. 
 
+## Beginner
 |How To | Description |
 | ------- | ------------ |
 | **[](Engine\Animation\AnimHowTo\PreviewingAnimations)** | %Engine\Animation\AnimHowTo\PreviewingAnimations:description% |
+
+
+## Intermediate
+|How To | Description |
+| ------- | ------------ |
+| **[](Engine\Animation\AnimHowTo\BlendSpace)** | %Engine\Animation\AnimHowTo\BlendSpace:description% |
+| **[](Engine\Animation\AnimHowTo\AimOffset)** | %Engine\Animation\AnimHowTo\AimOffset:description% |
 | **[](Engine\Animation\AnimHowTo\Retargeting)** | %Engine\Animation\AnimHowTo\Retargeting:description% |
+
+
+## Advanced
+|How To | Description |
+| ------- | ------------ |
 | **[](Engine\Animation\AnimHowTo\AdditiveAnimations)** | %Engine\Animation\AnimHowTo\AdditiveAnimations:description% |
-| **[](Engine\Animation\AnimHowTo\AimOffset)** | %Engine\Animation\AnimHowTo\AimOffset:description% |
-| **[](Engine\Animation\AnimHowTo\BlendSpace)** | %Engine\Animation\AnimHowTo\BlendSpace:description% |
-<!-- Add this for 4.7 release as it is a 4.7 feature. 
 | **[](Engine\Animation\AnimHowTo\LayerEditing)** | %Engine\Animation\AnimHowTo\LayerEditing:description% |
--->
 
 
+<!-- Convert these to Quick Starts
+| **[](Engine\Animation\AnimHowTo\AimOffset)** | %Engine\Animation\AnimHowTo\AimOffset:description% |Beginner|
+| **[](Engine\Animation\AnimHowTo\BlendSpace)** | %Engine\Animation\AnimHowTo\BlendSpace:description% |Beginner|
+
+Split this into two How To's
+| **[](Engine\Animation\AnimHowTo\Retargeting)** | %Engine\Animation\AnimHowTo\Retargeting:description% |Intermediate|
+-->
+
 
 
 
