#UE4doc #docs:  Removed extraneous bolding as according to the instructional guide

diff --git a/Engine/Documentation/Source/Engine/Components/SkeletalMesh/SkeletalMeshComponents.INT.udn b/Engine/Documentation/Source/Engine/Components/SkeletalMesh/SkeletalMeshComponents.INT.udn
--- a/Engine/Documentation/Source/Engine/Components/SkeletalMesh/SkeletalMeshComponents.INT.udn
+++ b/Engine/Documentation/Source/Engine/Components/SkeletalMesh/SkeletalMeshComponents.INT.udn
@@ -11,12 +11,12 @@ Skeletal Mesh Components are used for anything that has complex animation data a
 
 ## Skeletal Mesh Components
 
-**SkeletalMeshComponents** are used to create an instance of a `USkeletalMesh`. The **Skeletal Mesh** (outward appearance) has a complex **Skeleton** (interconnected bones) inside which helps move the individual vertices of the Skeletal Mesh to match the current animation that is being played. 
-This makes **SkeletalMeshComponents** ideal for things like characters, creatures, complex machinery, or anything that needs to deform or display complex motion. See [](Engine/Content/Types/SkeletalMeshes) and [](Engine/Actors/SkeletalMeshActors) for more information on working with **Skeletal Meshes**. 
+**SkeletalMeshComponents** are used to create an instance of a `USkeletalMesh`. The Skeletal Mesh (outward appearance) has a complex Skeleton (interconnected bones) inside which helps move the individual vertices of the Skeletal Mesh to match the current animation that is being played. 
+This makes SkeletalMeshComponents ideal for things like characters, creatures, complex machinery, or anything that needs to deform or display complex motion. See [](Engine/Content/Types/SkeletalMeshes) and [](Engine/Actors/SkeletalMeshActors) for more information on working with Skeletal Meshes. 
 
 ![](mesh_component.png)
 
-Above, a **SkeletalMeshComponent** is used with a **Character Blueprint** to create a playable character. In addition to specifying the **Skeletal Mesh** asset to use, you can also define the **Animation Mode** for the mesh to use (either an [Animation Blueprint](Engine/Animation/AnimBlueprints) or an [Animation Asset](Engine/Animation/Sequences)). 
+Above, a SkeletalMeshComponent is used with a Character Blueprint to create a playable character. In addition to specifying the Skeletal Mesh asset to use, you can also define the Animation Mode for the mesh to use (either an [Animation Blueprint](Engine/Animation/AnimBlueprints) or an [Animation Asset](Engine/Animation/Sequences)). 
 
 
 
