#UE4doc #docs: Removed extraneous bolding as according to the instructional guide

diff --git a/Engine/Documentation/Source/Engine/Components/StaticMesh/StaticMeshComponents.INT.udn b/Engine/Documentation/Source/Engine/Components/StaticMesh/StaticMeshComponents.INT.udn
--- a/Engine/Documentation/Source/Engine/Components/StaticMesh/StaticMeshComponents.INT.udn
+++ b/Engine/Documentation/Source/Engine/Components/StaticMesh/StaticMeshComponents.INT.udn
@@ -11,13 +11,13 @@ Static Mesh Components allow a piece of geometry to be used as a sub-object for
 
 ## Static Mesh Component
 
-The **StaticMeshComponent** is used to create an instance of a `UStaticMesh`. A **Static Mesh** is a piece of geometry that consists of a static set of polygons and are the basic unit used to create world geometry for levels in Unreal Engine 4. In addition to building levels, **Static Meshes** can be used for creating movers such as doors or lifts, rigid body physics objects, foliage and terrain decorations, procedurally created buildings, game objectives, and many more visual elements. 
+The **StaticMeshComponent** is used to create an instance of a `UStaticMesh`. A Static Mesh is a piece of geometry that consists of a static set of polygons and are the basic unit used to create world geometry for levels in Unreal Engine 4. In addition to building levels, Static Meshes can be used for creating movers such as doors or lifts, rigid body physics objects, foliage and terrain decorations, procedurally created buildings, game objectives, and many more visual elements. 
 
-Below, a **StaticMeshComponent** representing a light fixture along with a **SpotLightComponent** and **PointLightComponent** are used to create an alarm that could cycle between two different types of lights.
+Below, a StaticMeshComponent representing a light fixture along with a SpotLightComponent and PointLightComponent are used to create an alarm that could cycle between two different types of lights.
 
 ![](static_mesh.png)(w:720)
 
-Another example is in the form of a door where two **StaticMeshComponents** are placed side-by-side. Through script you could move one (or both) to open a pathway. 
+Another example is in the form of a door where two StaticMeshComponents are placed side-by-side. Through script you could move one (or both) to open a pathway. 
 
 ![](static_mesh2.png)(w:720)
 
