#UE4doc #docs:  Removed extraneous bolding as according to the instructional guide. Fixed a spelling errror: exmple -> example

diff --git a/Engine/Documentation/Source/Engine/Components/Widget/WidgetComponents.INT.udn b/Engine/Documentation/Source/Engine/Components/Widget/WidgetComponents.INT.udn
--- a/Engine/Documentation/Source/Engine/Components/Widget/WidgetComponents.INT.udn
+++ b/Engine/Documentation/Source/Engine/Components/Widget/WidgetComponents.INT.udn
@@ -11,23 +11,23 @@ Widget Components allow you to manifest 3D UI elements created through [Unreal M
 
 ## Widget Component
 
-The **WidgetComponent** itself is a 3D instance of a **Widget Blueprint** that you can interact with in your game world. 
+The **WidgetComponent** itself is a 3D instance of a Widget Blueprint that you can interact with in your game world. 
 
-In the example below, we have added a **WidgetComponent** which will now allow us to specify a **Widget Class Blueprint**. 
+In the example below, we have added a WidgetComponent which will now allow us to specify a Widget Class Blueprint. 
 
 ![](Engine/UMG/HowTo/Create3DWidgets/UMG_3D_21.png)
 
-_This example uses a **Widget Blueprint** to display a interactive menu that appears in the game world._
+_This example uses a Widget Blueprint to display a interactive menu that appears in the game world._
 
-You can change settings to how the **WidgetComponent** is perceived in the world by changing the **Draw Size** or **Max Interaction** distance. 
+You can change settings to how the WidgetComponent is perceived in the world by changing the Draw Size or Max Interaction distance. 
 
 ![](Engine/UMG/HowTo/Create3DWidgets/UMG_3D_22.png)(w:720)
 
-Once we place the Actor that contains our **WidgetComponent** in our level, our **Widget Class Blueprint** is displayed in the world. 
+Once we place the Actor that contains our WidgetComponent in our level, our Widget Class Blueprint is displayed in the world. 
 
 ![](Engine/UMG/HowTo/Create3DWidgets/UMG_3D_30.png)(w:720)
 
-See [](Engine/UMG/HowTo/Create3DWidgets) for a step-by-step guide to creating this exmaple using a **WidgetCompoment** and **Widget Blueprint**. 
+See [](Engine/UMG/HowTo/Create3DWidgets) for a step-by-step guide to creating this example using a WidgetCompoment and Widget Blueprint. 
 
 [REGION:note]
 This Component is considered as experimental and may be substantially changed or removed in future releases. 
