#UE4 Docs

Jira UEDOC-1792
Desc: Replaced the GIF on this page with a YouTube linked video.

diff --git a/Engine/Documentation/Source/Engine/Animation/Blendspaces/Blendspaces.INT.udn b/Engine/Documentation/Source/Engine/Animation/Blendspaces/Blendspaces.INT.udn
--- a/Engine/Documentation/Source/Engine/Animation/Blendspaces/Blendspaces.INT.udn
+++ b/Engine/Documentation/Source/Engine/Animation/Blendspaces/Blendspaces.INT.udn
@@ -116,7 +116,17 @@ _Blend Spaces_ can also be created in a one-dimensional format, known as a **_Bl
 
 Aside from having only one axis instead of two, a _Blend Space 1D_ is exactly the same as a standard two-axis _Blend Space_. As such, this document will focus on the 2D version. But you should know that all the rules and properties still carry over.
 
-![](BlendSpace_1D.gif)(convert:false)
+[OBJECT:EmbeddedVideo]
+[PARAMLITERAL:width]
+640
+[/PARAMLITERAL]
+[PARAMLITERAL:height]
+360
+[/PARAMLITERAL]
+[PARAMLITERAL:videoid]
+2mkn8FZL2KA
+[/PARAMLITERAL]
+[/OBJECT]
 
 ## Architecture
 
