#UE4doc #docs: Multiple bolding and unbolding edits

diff --git a/Engine/Documentation/Source/Engine/Rendering/LightingAndShadows/Lightmass/Basics/Lightmass.INT.udn b/Engine/Documentation/Source/Engine/Rendering/LightingAndShadows/Lightmass/Basics/Lightmass.INT.udn
--- a/Engine/Documentation/Source/Engine/Rendering/LightingAndShadows/Lightmass/Basics/Lightmass.INT.udn
+++ b/Engine/Documentation/Source/Engine/Rendering/LightingAndShadows/Lightmass/Basics/Lightmass.INT.udn
@@ -46,9 +46,7 @@ Description:High level introduction to Lightmass.
 ![](Engine/Rendering/LightingAndShadows/lighting_overview_banner_001.png)(convert:false)
 [/REGION]
 
-[TOC(start:2)]
-
-## Overview
+[TOC(start:2)
 
 [INCLUDE:Engine/Rendering/LightingAndShadows/Lightmass#LightMass_002]
 
diff --git a/Engine/Documentation/Source/Engine/Rendering/LightingAndShadows/Lightmass/Lightmass.INT.udn b/Engine/Documentation/Source/Engine/Rendering/LightingAndShadows/Lightmass/Lightmass.INT.udn
--- a/Engine/Documentation/Source/Engine/Rendering/LightingAndShadows/Lightmass/Lightmass.INT.udn
+++ b/Engine/Documentation/Source/Engine/Rendering/LightingAndShadows/Lightmass/Lightmass.INT.udn
@@ -30,15 +30,15 @@ Description:A run down of all the features and settings of Lightmass.
 [TOC(start:2 end:2)]
 
 
-## Overview
+
 
 [EXCERPT:LightMass_002]
 **Lightmass** creates lightmaps with complex light interactions like area shadowing and diffuse interreflection.  It is used to precompute portions of the lighting contribution of lights with stationary and static mobility. 
 [/EXCERPT:LightMass_002]
 
-Communication between the editor and **Lightmass** is handled by the **Swarm Agent**, which manages the lighting build locally and can also distribute the lighting build to remote machines. The **Swarm Agent**, which opens minimized by default, also tracks lighting build progress and keeps you up to date with which machines are working for you, what they are working on, and how many threads each one is using.  
+Communication between the editor and Lightmass is handled by the Swarm Agent, which manages the lighting build locally and can also distribute the lighting build to remote machines. The Swarm Agent, which opens minimized by default, also tracks lighting build progress and keeps you up to date with which machines are working for you, what they are working on, and how many threads each one is using.  
 
-A sample image of the **Swarm Agent** is pictured below (the bar near the bottom shows how much of the build is complete).
+A sample image of the Swarm Agent is pictured below (the bar near the bottom shows how much of the build is complete).
 
 ![](SwarmAgent.png)(w:640)
 
@@ -46,13 +46,13 @@ A sample image of the **Swarm Agent** is pictured below (the bar near the bottom
 
 ### Diffuse Interreflection
 
-**Diffuse Interreflection** is by far the most visually important global illumination lighting effect.  Light bounces by default with **Lightmass**, and the **BaseColor** term of your material controls how much light (and what color) bounces in all directions.  This effect is sometimes called **Color Bleeding**.  **Diffuse Interreflection** is incoming light reflecting equally in all directions, which means that it is not affected by the viewing direction or position.
+**Diffuse Interreflection** is by far the most visually important global illumination lighting effect.  Light bounces by default with Lightmass, and the BaseColor term of your material controls how much light (and what color) bounces in all directions.  This effect is sometimes called Color Bleeding.  Diffuse Interreflection is incoming light reflecting equally in all directions, which means that it is not affected by the viewing direction or position.
 
-Here is a scene built by **Lightmass** with a single **Directional Light** and only direct lighting shown.  The areas that are not directly visible to the light are black.  This is the result without global illumination.
+Here is a scene built by Lightmass with a single Directional Light and only direct lighting shown.  The areas that are not directly visible to the light are black.  This is the result without global illumination.
 
 ![3DirectOnly.jpg](3DirectOnly.jpg)(w:720)
 
-This is what the first diffuse global illumination bounce looks like.  Notice the shadow behind the chair on the left, this is called an **Indirect Shadow** because it is the shadow of the indirect light.  The brightness and color of the diffuse bounce depends on the incoming light and the diffuse term of the material that the light interacts with.  Each bounce is darker than the previous, since some of the light gets absorbed by surfaces instead of being reflected.  The pillar bases get more indirect light than other surfaces because they are closer to the areas in direct light.
+This is what the first diffuse global illumination bounce looks like.  Notice the shadow behind the chair on the left, this is called an Indirect Shadow because it is the shadow of the indirect light.  The brightness and color of the diffuse bounce depends on the incoming light and the diffuse term of the material that the light interacts with.  Each bounce is darker than the previous, since some of the light gets absorbed by surfaces instead of being reflected.  The pillar bases get more indirect light than other surfaces because they are closer to the areas in direct light.
 
 ![3FirstBounceOnly.jpg](3FirstBounceOnly.jpg)(w:720)
 
@@ -65,7 +65,7 @@ And here is the scene with four diffuse bounces combined.  Simulating global ill
 ![3FourBounces.jpg](3FourBounces.jpg)(w:720)
 
 
-Bounced lighting picks up the diffuse color of the underlying material, as you can see below.  This is where the term **Color Bleeding** comes from.  **Color Bleeding** is most noticeable with highly saturated colors.  You can exaggerate the effect by raising **DiffuseBoost** on the **Primitive**, **Material** or **Level**.
+Bounced lighting picks up the diffuse color of the underlying material, as you can see below.  This is where the term Color Bleeding comes from.  Color Bleeding is most noticeable with highly saturated colors.  You can exaggerate the effect by raising DiffuseBoost on the Primitive, Material, or Level.
 
 [REGION:imagetable]
 | ![diffusebounce_T.JPG](diffusebounce_T.JPG)(h:445) | ![diffusebounce_L.JPG](diffusebounce_L.JPG)(h:445) |
@@ -76,7 +76,7 @@ Bounced lighting picks up the diffuse color of the underlying material, as you c
 
 ### Character lighting
 
-**Lightmass** places samples in a uniform 3d grid inside the [Lightmass Importance Volume](#LightmassImportanceVolume) at a low resolution, and also on top of upward facing surfaces that characters are likely to walk on, at a higher resolution.  Each lighting sample captures indirect lighting from all directions, but not direct lighting.
+Lightmass places samples in a uniform 3d grid inside the [Lightmass Importance Volume](#LightmassImportanceVolume) at a low resolution, and also on top of upward facing surfaces that characters are likely to walk on, at a higher resolution.  Each lighting sample captures indirect lighting from all directions, but not direct lighting.
 
 The first image is a debug visualization of lighting samples placed above the floor, the second image is the same scene in Lit mode.  Notice how the samples above the red tapestry picked up the red bounced light.  The samples are visualized as a single color, but they really capture lighting from all directions.
 
@@ -87,21 +87,21 @@ The first image is a debug visualization of lighting samples placed above the fl
 [/REGION]
 
 
-The **Indirect Lighting Cache** then interpolates the indirect lighting for a movable object using these lighting samples.  The indirect lighting affects the color of the light environment shadow, but not the direction.  Use **Show->Visualize->Volume Lighting Samples** to preview these in your viewport.
+The Indirect Lighting Cache then interpolates the indirect lighting for a movable object using these lighting samples.  The indirect lighting affects the color of the light environment shadow, but not the direction.  Use Show->Visualize->Volume Lighting Samples to preview these in your viewport.
 
 For more information, see the [Indirect Lighting Cache](Engine/Rendering/LightingAndShadows/IndirectLightingCache) documentation.
 
 #### Limitations
 
-   * The default settings for volume sample placement will result in a large number of samples in a large map.  This will cause **Indirect Lighting Cache** interpolation times to be very large.  Use **Static Lighting Level Scale** to reduce the sample count in large maps.
-   * Movable objects outside the **Lightmass Importance Volume** will have black indirect lighting.
+   * The default settings for volume sample placement will result in a large number of samples in a large map.  This will cause Indirect Lighting Cache interpolation times to be very large.  Use Static Lighting Level Scale to reduce the sample count in large maps.
+   * Movable objects outside the Lightmass Importance Volume will have black indirect lighting.
 
 
 ### Ambient Occlusion
 
-**Lightmass** calculates detailed indirect shadows automatically, but it can be useful to exaggerate indirect shadows for artistic purposes or to enhance the perception of proximity in a scene.  
+Lightmass calculates detailed indirect shadows automatically, but it can be useful to exaggerate indirect shadows for artistic purposes or to enhance the perception of proximity in a scene.  
 
-Ambient occlusion is the indirect shadowing you would get from a uniformly bright upper hemisphere, like an overcast sky.  **Lightmass** supports calculating ambient occlusion, applying it to direct and indirect illumination, and then baking it into lightmaps.  Ambient occlusion is enabled by default, and can be disabled by unchecking the checkbox **Use Ambient Occlusion** in the **Lightmass Settings** of **Lightmass** under the **World Settings**.
+Ambient occlusion is the indirect shadowing you would get from a uniformly bright upper hemisphere, like an overcast sky.  Lightmass supports calculating ambient occlusion, applying it to direct and indirect illumination, and then baking it into lightmaps.  Ambient occlusion is enabled by default, and can be disabled by unchecking the checkbox Use Ambient Occlusion in the Lightmass Settings of Lightmass under the World Settings.
 
 In the first image is a scene with indirect lighting but no ambient occlusion.  In the second image is the same scene with ambient occlusion applied to both the direct and indirect lighting, note the darkening where objects come together.
 
@@ -127,7 +127,7 @@ In the first image is a scene with indirect lighting but no ambient occlusion.
 [/REGION]
 
 
-Ambient occlusion is nearly free as far as lighting build time goes when **Num Indirect Lighting Bounces** is greater than 0.
+Ambient occlusion is nearly free as far as lighting build time goes when Num Indirect Lighting Bounces is greater than 0.
 
 #### Limitations
 
@@ -137,7 +137,7 @@ Ambient occlusion is nearly free as far as lighting build time goes when **Num I
 
 ### Masked shadows
 
-**Lightmass** takes the opacity mask of **BLEND_Masked Materials** into account when calculating shadows.  The part of the material that gets clipped in the editor viewports also does not cause any shadowing, which allows much more detailed shadowing from trees and foliage.
+Lightmass takes the opacity mask of BLEND_Masked Materials into account when calculating shadows.  The part of the material that gets clipped in the editor viewports also does not cause any shadowing, which allows much more detailed shadowing from trees and foliage.
 
 [REGION:imagetable]
 | ![masked.JPG](masked.JPG)(h:420) | ![masked_L.JPG](masked_L.JPG)(h:420) |
@@ -149,19 +149,19 @@ Ambient occlusion is nearly free as far as lighting build time goes when **Num I
 
 ### Bent Normal sky occlusion
 
-**Lightmass** generates directional occlusion in the form of a bent normal whenever a [Sky Light](Engine/Rendering/LightingAndShadows/LightTypes/SkyLight) with **Stationary Mobility** is enabled.  
+Lightmass generates directional occlusion in the form of a bent normal whenever a [Sky Light](Engine/Rendering/LightingAndShadows/LightTypes/SkyLight) with Stationary Mobility is enabled.  
 
 ### Distance field shadowmaps
 
-**Lightmass** computes distance field shadowmaps for [Stationary Lights](Engine/Rendering/LightingAndShadows/LightMobility/StationaryLights).  Distance field shadow maps maintain their curvy shape very well even under lower resolutions; however they do not support area shadowing or translucent shadows.
+Lightmass computes distance field shadowmaps for [Stationary Lights](Engine/Rendering/LightingAndShadows/LightMobility/StationaryLights).  Distance field shadow maps maintain their curvy shape very well even under lower resolutions; however they do not support area shadowing or translucent shadows.
 
 ## Features for Static Lights Only
 
 ### Area lights and shadows
 
-With **Lightmass**, all lights with **Static Mobility** are area lights by default.  The shape used by **Point** and **Spot Light** sources is a sphere, whose radius is set by **Light Source Radius** under **Lightmass Settings**.  **Directional Light** sources use a disk, positioned at the edge of the scene.  Light source size is one of the two factors controlling shadow softness, as larger light sources will create softer shadows.  The other factor is distance from the receiving location to the shadow caster.  Area shadows get softer as this distance increases, just like in real life.
+With Lightmass, all lights with Static Mobility are area lights by default.  The shape used by Point and Spot Light sources is a sphere, whose radius is set by Light Source Radius under Lightmass Settings.  Directional Light sources use a disk, positioned at the edge of the scene.  Light source size is one of the two factors controlling shadow softness, as larger light sources will create softer shadows.  The other factor is distance from the receiving location to the shadow caster.  Area shadows get softer as this distance increases, just like in real life.
 
-In the first image is a **Static Directional Light** with only static lighting, the penumbra size is the same everywhere.  In the second image, **Lightmass** calculated area shadows whose sharpness is controlled by both the light source size and occluder distance.  Notice how the shadow of the pillar is much sharper close to where the pillar meets the ground.
+In the first image is a Static Directional Light with only static lighting, the penumbra size is the same everywhere.  In the second image, Lightmass calculated area shadows whose sharpness is controlled by both the light source size and occluder distance.  Notice how the shadow of the pillar is much sharper close to where the pillar meets the ground.
 
 [REGION:imagetable]
 | ![7OldUE3Shadows.jpg](7OldUE3Shadows.jpg)(h:520) | ![7LightmassShadows.jpg](7LightmassShadows.jpg)(h:520) |
@@ -169,7 +169,7 @@ In the first image is a **Static Directional Light** with only static lighting,
 | | |
 [/REGION]
 
-The light source radius for **Point** and **Spot Lights** is drawn in yellow wireframe, while the influence radius is drawn in teal.  In most cases, you will want to make sure the light source does not intersect any shadow casting geometry, or the light will be emitting on both sides of that geometry.
+The light source radius for Point and Spot Lights is drawn in yellow wireframe, while the influence radius is drawn in teal.  In most cases, you will want to make sure the light source does not intersect any shadow casting geometry, or the light will be emitting on both sides of that geometry.
 
 ![](AreaLight.png)
 
@@ -182,7 +182,7 @@ Light passing through a translucent material that is applied to a static shadow
 
 #### Translucent shadow color
 
-The amount of light passing through the material is called **Transmission**, and ranges between 0 and 1 for each color channel.  A value of 0 would be fully opaque, and a value of 1 would mean that the incident light passes through unaffected.  There is no material input for **Transmission**, so currently it is derived from the other material inputs as follows:
+The amount of light passing through the material is called Transmission, and ranges between 0 and 1 for each color channel.  A value of 0 would be fully opaque, and a value of 1 would mean that the incident light passes through unaffected.  There is no material input for Transmission, so currently it is derived from the other material inputs as follows:
 
 * Lit materials
 	* BLEND_Translucent and BLEND_Additive: Transmission = Lerp(White, BaseColor, Opacity)
@@ -197,7 +197,7 @@ This means that at an opacity of 0, the material will not filter out incident li
 
 There are several factors controlling translucent shadow sharpness.
 
-In the first image, a large light source was used (directional light with a **Light Source Angle** of 5) and in the second, a small light source was used (**Light Source Angle** of 0).
+In the first image, a large light source was used (directional light with a Light Source Angle of 5) and in the second, a small light source was used (Light Source Angle of 0).
 
 [REGION:imagetable]
 | ![5LargeLightSource.jpg](5LargeLightSource.jpg)(h:560) | ![5SmallLightSource.jpg](5SmallLightSource.jpg)(h:560) |
@@ -205,7 +205,7 @@ In the first image, a large light source was used (directional light with a **Li
 | | |
 [/REGION]
 
-In the first image, a small light source was used, but the lightmap resolution was too low to capture the sharp translucent shadows.  In the second image, the material was exported at too low of a resolution (controlled by **Export Resolution Scale** in the **Material Editor**) to capture the sharp shadows.
+In the first image, a small light source was used, but the lightmap resolution was too low to capture the sharp translucent shadows.  In the second image, the material was exported at too low of a resolution (controlled by Export Resolution Scale in the Material Editor) to capture the sharp shadows.
 
 [REGION:imagetable]
 | ![5LowLightmapResolution.jpg](5LowLightmapResolution.jpg)(h:565) | ![5LowExportResolutionScale.jpg](5LowExportResolutionScale.jpg)(h:565) |
@@ -213,7 +213,7 @@ In the first image, a small light source was used, but the lightmap resolution w
 | | |
 [/REGION]
 
-Indirect light is also affected by **Translucent Materials**.  The window in this image filters incoming light based on its **Transmission**, and that light then bounces around the scene with its color modified.
+Indirect light is also affected by Translucent Materials.  The window in this image filters incoming light based on its Transmission, and that light then bounces around the scene with its color modified.
 
 ![TranslucentShadowIndirectLight.jpg](TranslucentShadowIndirectLight.jpg)
 
@@ -232,7 +232,7 @@ Indirect light is also affected by **Translucent Materials**.  The window in thi
 
 During rendering, lit pixel color is determined as BaseColor * Lighting, so base color directly affects how visible the lighting will be.  High contrast or dark diffuse textures make lighting difficult to notice, while low contrast, mid-range diffuse textures let the lighting details show through.
 
-Compare the lighting clarity between the scene in the first image built with mid-range diffuse textures, to the scene in the second image also built with **Lightmass** but with noisy, dark diffuse textures.  Only the most high-frequency changes are noticeable in the scene in the second image, like the shadow transitions.
+Compare the lighting clarity between the scene in the first image built with mid-range diffuse textures, to the scene in the second image also built with Lightmass but with noisy, dark diffuse textures.  Only the most high-frequency changes are noticeable in the scene in the second image, like the shadow transitions.
 
 [REGION:imagetable]
 ![94MidToneDiffuse.jpg](94MidToneDiffuse.jpg)(w:465) | ![94DarkDiffuse.jpg](94DarkDiffuse.jpg)(w:465) |
@@ -240,7 +240,7 @@ Compare the lighting clarity between the scene in the first image built with mid
 | | |
 [/REGION]
 
-The **Unlit** view mode is useful for viewing the **Diffuse** term.  The scene in the first image looks flatter and more monotone in the **Unlit** view mode, which means the lighting is doing all the work and variations in final pixel color are mostly due to lighting differences.  **To get good lighting, your scene should look flat and boring in the Unlit view mode.**  Baking lighting and macro features into the diffuse textures will counteract the lighting.
+The **Unlit** view mode is useful for viewing the Diffuse term.  The scene in the first image looks flatter and more monotone in the Unlit view mode, which means the lighting is doing all the work and variations in final pixel color are mostly due to lighting differences.  (To get good lighting, your scene should look flat and boring in the Unlit view mode.)  Baking lighting and macro features into the diffuse textures will counteract the lighting.
 
 [REGION:imagetable]
 ![94MidToneDiffuseUnlit.jpg](94MidToneDiffuseUnlit.jpg)(w:465) | ![94DarkDiffuseUnlit.jpg](94DarkDiffuseUnlit.jpg)(w:465) |
@@ -273,16 +273,16 @@ Using texture lightmaps with high resolution is the best way to get detailed, qu
 
 #### Lightmass Solver quality
 
-**Lightmass Solver** settings are automatically set based on what quality of build is requested in the **Lighting Build Options** dialog.  Production should give good enough quality that the artifacts are not clearly noticeable with a diffuse texture applied.  
+**Lightmass Solver** settings are automatically set based on what quality of build is requested in the Lighting Build Options dialog.  Production should give good enough quality that the artifacts are not clearly noticeable with a diffuse texture applied.  
 
 ## Getting the Best Lighting Build Times
 
-There are several ways to improve your **Lightmass** build times:
+There are several ways to improve your Lightmass build times:
 
    * Only have high resolution lightmaps in areas that have high-frequency (quickly changing) lighting. Reduce the lightmap resolution for Brush surfaces and Static Meshes that are not in direct lighting or affected by sharp indirect shadows. This will give you high resolution shadows in the areas that are most noticeable.
    * Surfaces that are never visible to the player should be set to the lowest possible lightmap resolution.
    * Use a [Lightmass Importance Volume](#LightmassImportanceVolume) to contain the areas that are most important (just around the playable area).
-   * Optimize the lightmap resolutions across the map so that build time for meshes is more even.  The lighting build can never be faster than the slowest single object, regardless of how many machines are doing the distributed build.  **Avoid large continuous meshes that surround a large part of the level and with high lightmap resolution.**  You will get better build times if you break these up into more modular pieces, especially on machines with many cores.
+   * Optimize the lightmap resolutions across the map so that build time for meshes is more even.  The lighting build can never be faster than the slowest single object, regardless of how many machines are doing the distributed build.  Avoid large continuous meshes that surround a large part of the level and with high lightmap resolution. You will get better build times if you break these up into more modular pieces, especially on machines with many cores.
    * Meshes with a lot of self-occlusion will take longer to build, for example a rug with many layers parallel to each other will take much longer to build than a flat floor.
 
 The **Lighting Build Info** dialog is a very important tool for improving lighting build times.  First, build lighting in the level that you want to see stats for.  Then, open the dialog under *Build->Lighting Info->Lighting StaticMesh Info*.  Change the drop down to *Lighting Build Info*.  This will show a sorted list of meshes and how long they took to compute lighting for.
@@ -294,12 +294,12 @@ The **Lighting Build Info** dialog is a very important tool for improving lighti
 ### Lightmass Importance Volume
 
 [EXCERPT:LightMass_001]
-Many maps have meshes out to the edge of the grid in the editor, but the actual playable area that needs high quality lighting is much smaller.  **Lightmass** emits photons based on the size of the level, so those background meshes will greatly increase the number of photons that need to be emitted, and lighting build times will increase.  The **Lightmass Importance Volume** controls the area that **Lightmass** emits photons in, allowing you to concentrate it only on the area that needs detailed indirect lighting.  Areas outside the importance volume get only one bounce of indirect lighting at a lower quality.
+Many maps have meshes out to the edge of the grid in the editor, but the actual playable area that needs high quality lighting is much smaller.  Lightmass emits photons based on the size of the level, so those background meshes will greatly increase the number of photons that need to be emitted, and lighting build times will increase.  The Lightmass Importance Volume controls the area that Lightmass emits photons in, allowing you to concentrate it only on the area that needs detailed indirect lighting.  Areas outside the importance volume get only one bounce of indirect lighting at a lower quality.
 [/EXCERPT:LightMass_001]
 
 An overhead wireframe view of a multiplayer map is shown in the first image.  The actual playable area that needs high quality lighting is the small green blob at the center.  
 
-In the second image, a close-up of the playable area of the multiplayer map is shown, with the correctly setup **Lightmass Importance Volume** selected.  The **Lightmass Importance Volume** reduced the radius of the region to light from 80,000 units to 10,000 units, which is 64x less area to light.
+In the second image, a close-up of the playable area of the multiplayer map is shown, with the correctly setup **Lightmass Importance Volume** selected.  The Lightmass Importance Volume reduced the radius of the region to light from 80,000 units to 10,000 units, which is 64x less area to light.
 
 [REGION:imagetable]
 ![JacintoImportanceFar.jpg](JacintoImportanceFar.jpg)(w:465) | ![JacintoImportanceNear.jpg](JacintoImportanceNear.jpg)(w:465) |
@@ -308,24 +308,24 @@ In the second image, a close-up of the playable area of the multiplayer map is s
 [/REGION]
 
 [EXCERPT:LightMass_003]
-You can add a **Lightmass Importance Volume** to a level by dragging a **Lightmass Importance Volume** object from the **Volumes** tab of the **Modes** menu into the level then scaling it to the needed size.
+You can add a Lightmass Importance Volume to a level by dragging a **Lightmass Importance Volume** object from the **Volumes** tab of the **Modes** menu into the level then scaling it to the needed size.
 
 ![](AddingLightMassVolume.png)(w:840)
 
-You can also convert a **Brush** into a **Lightmass Importance Volume** by clicking the **Convert Actor** drop down box inside the **Details** panel under **Actor**.
+You can also convert a Brush into a Lightmass Importance Volume by clicking the **Convert Actor** drop down box inside the **Details** panel under **Actor**.
 
 ![](ConvertBrush_Actor.png)
 
-After clicking the drop down box, a menu will appear where you can select the type of **Actor** to replace the **Brush** with.
+After clicking the drop down box, a menu will appear where you can select the type of Actor to replace the Brush with.
 
 ![](ConvertBrush_LightMassVol.png)
 
-If you place multiple **Lightmass Importance Volumes**, most of the lighting work will be done with a bounding box that contains all of them.  However, **Volume Lighting Samples** are only placed inside the smaller volumes.
+If you place multiple Lightmass Importance Volumes, most of the lighting work will be done with a bounding box that contains all of them.  However, Volume Lighting Samples are only placed inside the smaller volumes.
 [/EXCERPT:LightMass_003]
 
 ### World Settings 
 	
-**Lightmass** settings can be adjusted inside the **World Settings** window under the **Lightmass** section. 
+Lightmass settings can be adjusted inside the **World Settings** window under the **Lightmass** section. 
 
 ![worldlightmasssettings.png](worldlightmasssettings.png)
 
@@ -353,7 +353,7 @@ You can access **World Settings** by clicking the ![](WorldSettings_Icon.png) ic
 
 ### Light Settings
 
-Below are **Lightmass** settings that can be adjusted inside the properties of a light under the **Lightmass** section. 
+Below are Lightmass settings that can be adjusted inside the properties of a light under the **Lightmass** section. 
 
 ![LightLightmassSettings.png](LightLightmassSettings.png)
 
@@ -366,7 +366,7 @@ Below are **Lightmass** settings that can be adjusted inside the properties of a
 
 ### Primitive Component Settings
 
-Below are **Lightmass** settings that can be adjusted on a **Brush** added from the **Geometry** tab of the **Modes** menu. These options are found under the **Details** panel of the **Brush**. 
+Below are Lightmass settings that can be adjusted on a **Brush** added from the **Geometry** tab of the **Modes** menu. These options are found under the **Details** panel of the **Brush**. 
 
 ![objectlightmasssettings.png](objectlightmasssettings.png)
 
@@ -379,9 +379,9 @@ Below are **Lightmass** settings that can be adjusted on a **Brush** added from
 
 
 ### Base Material Settings
-Below are the **Lightmass** settings that can be adjusted inside a **Material** for the **Base Node** under the **Details** panel. 
+Below are the Lightmass settings that can be adjusted inside a Material for the Base Node under the **Details** panel. 
 
-For more information on the **Material Editor**, see the [Material Editor User Guide](Engine/Rendering/Materials/Editor).
+For more information on the Material Editor, see the **[Material Editor User Guide](Engine/Rendering/Materials/Editor)**.
 
 ![materialsettings.png](materialsettings.png)
 
