#UE4 Docs

Jira UEDOC-1798
Replaced video attachment with YouTube video.

diff --git a/Engine/Documentation/Source/Engine/Animation/Blendspaces/Editor/BlendspaceEditor.INT.udn b/Engine/Documentation/Source/Engine/Animation/Blendspaces/Editor/BlendspaceEditor.INT.udn
--- a/Engine/Documentation/Source/Engine/Animation/Blendspaces/Editor/BlendspaceEditor.INT.udn
+++ b/Engine/Documentation/Source/Engine/Animation/Blendspaces/Editor/BlendspaceEditor.INT.udn
@@ -172,8 +172,16 @@ Tooltip Display** option.
 
 ![Graph Tooltip](bs_samples_tooltip.png)
 
-[REGION:caption]
-**Right-click and Save As to download a video**
-[/REGION]
-
-[![](PlayBlendSpaceVideo.png)(w:640)](BlendSpace.mp4)
+An example of a previewing a Blend Space is illustrated below. 
+
+[OBJECT:EmbeddedVideo]
+[PARAMLITERAL:width]
+640
+[/PARAMLITERAL]
+[PARAMLITERAL:height]
+360
+[/PARAMLITERAL]
+[PARAMLITERAL:videoid]
+L435S6VaUMY
+[/PARAMLITERAL]
+[/OBJECT]
\ No newline at end of file
