#UE4 Docs

Jira UEDOC-1800
Replaced GIF with YouTube video of Morph Target Previewer.

diff --git a/Engine/Documentation/Source/Engine/Animation/Persona/MorphTargetPreviewer/MorphTargetPreviewer.INT.udn b/Engine/Documentation/Source/Engine/Animation/Persona/MorphTargetPreviewer/MorphTargetPreviewer.INT.udn
--- a/Engine/Documentation/Source/Engine/Animation/Persona/MorphTargetPreviewer/MorphTargetPreviewer.INT.udn
+++ b/Engine/Documentation/Source/Engine/Animation/Persona/MorphTargetPreviewer/MorphTargetPreviewer.INT.udn
@@ -36,7 +36,17 @@ Related: Engine/Animation/AnimBlueprints
 
 [EXCERPT:Overview]
 
-![](MorphGif.gif)(w:720 convert:false)
+[OBJECT:EmbeddedVideo]
+[PARAMLITERAL:width]
+640
+[/PARAMLITERAL]
+[PARAMLITERAL:height]
+360
+[/PARAMLITERAL]
+[PARAMLITERAL:videoid]
+CfAI0c-Br8Y
+[/PARAMLITERAL]
+[/OBJECT]
 
 The **Morph Target Previewer** allows you to preview any morph targets (also known as "morphs" and "blend shapes", depending on what 3D animation package you are using) that are applied to your Skeletal Meshes. Each Morph Target is additively blended into the existing Skeletal Mesh geometry. This allows for multiple morph targets to be combined to create complex vertex-driven animation, which is often ideal for such things as facial expressions.
 [/EXCERPT:Overview]
