#UE4 Docs

Jira UEDOC-1802
Replaced GIF with standard image instead.

diff --git a/Engine/Documentation/Source/Engine/Animation/Persona/Toolbar/PersonaToolbar.INT.udn b/Engine/Documentation/Source/Engine/Animation/Persona/Toolbar/PersonaToolbar.INT.udn
--- a/Engine/Documentation/Source/Engine/Animation/Persona/Toolbar/PersonaToolbar.INT.udn
+++ b/Engine/Documentation/Source/Engine/Animation/Persona/Toolbar/PersonaToolbar.INT.udn
@@ -84,7 +84,7 @@ Persona toolbar buttons provide easy access to switching between asset modes, su
 
 ## Mode Switching Tabs
 
-![](Engine/Animation/Persona\ModeSelector.gif)(convert:false)
+![](ModeTabPersona.png)
 
 The Persona Toolbar also contains buttons to switch between each of its four modes. These buttons have some functionality in and of themselves that can be useful when trying to track down a particular animation asset. You change Persona modes by simply clicking on the button for the desired mode.
 
