#UE4doc #docs: Removed "Overview" header on these pages and bolded/unbolded type per instructional guide

diff --git a/Engine/Documentation/Source/Engine/Rendering/LightingAndShadows/LightMobility/DynamicLights/DynamicLights.INT.udn b/Engine/Documentation/Source/Engine/Rendering/LightingAndShadows/LightMobility/DynamicLights/DynamicLights.INT.udn
--- a/Engine/Documentation/Source/Engine/Rendering/LightingAndShadows/LightMobility/DynamicLights/DynamicLights.INT.udn
+++ b/Engine/Documentation/Source/Engine/Rendering/LightingAndShadows/LightMobility/DynamicLights/DynamicLights.INT.udn
@@ -48,7 +48,7 @@ Description:A fully dynamic light capable of changing all of its properties duri
 
 [TOC(start:2 end:2)]
 
-## Overview
+
 
 [EXCERPT:Movable001]
 
diff --git a/Engine/Documentation/Source/Engine/Rendering/LightingAndShadows/LightMobility/StaticLights/StaticLights.INT.udn b/Engine/Documentation/Source/Engine/Rendering/LightingAndShadows/LightMobility/StaticLights/StaticLights.INT.udn
--- a/Engine/Documentation/Source/Engine/Rendering/LightingAndShadows/LightMobility/StaticLights/StaticLights.INT.udn
+++ b/Engine/Documentation/Source/Engine/Rendering/LightingAndShadows/LightMobility/StaticLights/StaticLights.INT.udn
@@ -48,7 +48,7 @@ Description:A completely static light that has no overhead during a running game
 
 [TOC(start:2)]
 
-## Overview
+
 
 [INCLUDE:Resources/ContentExamples\Lighting\2_1#StaticLighting01]
 
@@ -56,7 +56,7 @@ The primary use case for static lights is for low power devices on mobile platfo
 
 ## Light Source Radius
 
-As mentioned above, **Light Source Radius** for **Static Lights** have the additional effect of softening the shadows they cast, much in the same way area lights work in many 3D rendering packages. Below, the light on the _left_ has a small **Light Source Radius** and casts sharper shadows. The light on the _right_ has a larger **Light Source Radius** and so casts softer shadows.
+As mentioned above, **Light Source Radius** for Static Lights have the additional effect of softening the shadows they cast, much in the same way area lights work in many 3D rendering packages. Below, the light on the _left_ has a small Light Source Radius and casts sharper shadows. The light on the _right_ has a larger Light Source Radius and so casts softer shadows.
 
 ![](Engine\Rendering\LightingAndShadows\Basics\LI_lightRadiusShadows.png)
 
diff --git a/Engine/Documentation/Source/Engine/Rendering/LightingAndShadows/LightShafts/LightShafts.INT.udn b/Engine/Documentation/Source/Engine/Rendering/LightingAndShadows/LightShafts/LightShafts.INT.udn
--- a/Engine/Documentation/Source/Engine/Rendering/LightingAndShadows/LightShafts/LightShafts.INT.udn
+++ b/Engine/Documentation/Source/Engine/Rendering/LightingAndShadows/LightShafts/LightShafts.INT.udn
@@ -48,7 +48,7 @@ Description:Simulating light scattered in the atmosphere.
 
 [TOC(start:2)]
 
-## Overview
+
 
 [EXCERPT:Lightshafts]
 **Light Shafts** can be generated by directional lights to simulate the real world effect of crepuscular rays, or atmospheric shadowing of atmospheric in-scattering. These rays add depth and realism to any scene.
diff --git a/Engine/Documentation/Source/Engine/Rendering/LightingAndShadows/LightTypes/Directional/Directional.INT.udn b/Engine/Documentation/Source/Engine/Rendering/LightingAndShadows/LightTypes/Directional/Directional.INT.udn
--- a/Engine/Documentation/Source/Engine/Rendering/LightingAndShadows/LightTypes/Directional/Directional.INT.udn
+++ b/Engine/Documentation/Source/Engine/Rendering/LightingAndShadows/LightTypes/Directional/Directional.INT.udn
@@ -48,10 +48,10 @@ Description:The basics of Directional Lights.
 
 [TOC(start:2 end:2)]
 
-## Overview
+
 
 [EXCERPT:Intro]
-The **Directional Light** simulates light that is being emitted from a source that is infinitely far away. This means that all shadows cast by this light will be parallel, making this the ideal choice for simulating sunlight. The **Directional Light** when placed can be set to one of three **Mobility** settings:
+The **Directional Light** simulates light that is being emitted from a source that is infinitely far away. This means that all shadows cast by this light will be parallel, making this the ideal choice for simulating sunlight. The Directional Light when placed can be set to one of three **Mobility** settings:
 [/EXCERPT:Intro]
 
 [REGION:raw]
@@ -98,7 +98,7 @@ Below you will see a representation of sunlight being cast through an opening in
 [![](directional_001.png)(w:450)](directional_001.png) [![](directional_002.png)(w:450)](directional_002.png)
 [/REGION]
 
-The left image shows just the light, the right image has its **Shadow Frustum** enabled in order to see that the **Directional Light** is casting parallel light rays. 
+The left image shows just the light, the right image has its **Shadow Frustum** enabled in order to see that the Directional Light is casting parallel light rays. 
 
 [OBJECT:topiclistnotitle]
 	[PARAM:icon]
@@ -112,7 +112,7 @@ The left image shows just the light, the right image has its **Shadow Frustum**
 
 ## Directional Light Properties
 
-The properties of a **Directional Light** are broken up into 5 categories: **Light**, **Light Shafts**, **Lightmass**, **Light Function** and **Cascaded Shadow Maps**. 
+The properties of a Directional Light are broken up into 5 categories: Light, Light Shafts, Lightmass, Light Function and Cascaded Shadow Maps:
 
 
 ### Light
diff --git a/Engine/Documentation/Source/Engine/Rendering/LightingAndShadows/LightTypes/LightTypes.INT.udn b/Engine/Documentation/Source/Engine/Rendering/LightingAndShadows/LightTypes/LightTypes.INT.udn
--- a/Engine/Documentation/Source/Engine/Rendering/LightingAndShadows/LightTypes/LightTypes.INT.udn
+++ b/Engine/Documentation/Source/Engine/Rendering/LightingAndShadows/LightTypes/LightTypes.INT.udn
@@ -43,7 +43,7 @@ DescriptionLong:The types of lights available in Unreal Engine 4: Point, Spot, D
 [/OBJECT]
 [/VAR]
 
-## Overview
+
 
 Unreal Engine 4 has four light types: **Directional**, **Point**, **Spot**, and **Sky**. Directional lights are primarily used as your primary outdoor light or any light that needs to appear as if it is casting light from extreme or near infinite distances. Point lights are your classic "light bulb" like light, emitting light in all directions from a single point. Spot lights emit light from a single point, but have their light limited by a set of cones.  Sky lights capture the background of your scene and apply it as lighting to your level's meshes.
 
