#UE4 Docs

Jira UEDOC-1815
Fixed by removing Blueprint Basic User Guide (information exists on other pages).
Update the Blueprint How To Directory as well to match the current instructional guide.

diff --git a/Engine/Documentation/Source/Engine/Blueprints/BP_HowTo/BasicUsage/BasicUsage.INT.udn b/Engine/Documentation/Source/Engine/Blueprints/BP_HowTo/BasicUsage/BasicUsage.INT.udn
--- a/Engine/Documentation/Source/Engine/Blueprints/BP_HowTo/BasicUsage/BasicUsage.INT.udn
+++ b/Engine/Documentation/Source/Engine/Blueprints/BP_HowTo/BasicUsage/BasicUsage.INT.udn
@@ -1,4 +1,4 @@
-Availability: Public
+Availability: docs
 Title:Blueprint Basic User Guide
 Crumbs: %ROOT%, Engine, Engine\Blueprints, Engine\Blueprints\BP_HowTo
 Description:A quick guide on the elements needed to begin working with Blueprints. 
diff --git a/Engine/Documentation/Source/Engine/Blueprints/BP_HowTo/BlueprintHowTo.INT.udn b/Engine/Documentation/Source/Engine/Blueprints/BP_HowTo/BlueprintHowTo.INT.udn
--- a/Engine/Documentation/Source/Engine/Blueprints/BP_HowTo/BlueprintHowTo.INT.udn
+++ b/Engine/Documentation/Source/Engine/Blueprints/BP_HowTo/BlueprintHowTo.INT.udn
@@ -6,33 +6,51 @@ Related: Engine/Blueprints
 Related: Engine/Blueprints/UserGuide
 Related: Gameplay
 
-## Overview
-If you are new to **Blueprints** in Unreal Engine 4, this **Blueprint - How To** series will provide information on the core aspects needed to working with Blueprints including Basic Usage, Adding and Connecting Nodes, Creating Variables and others as seen in the table below. 
 
-When you are ready and have an understanding of the core aspects behind Blueprints, **Gameplay** oriented How To's can be found in the [](Gameplay) which illustrate gameplay related elements like Spawning/Destroying Actors, Respawning (OnDestroyed), Using Timers, and more. 
+If you are new to [Blueprints Visual Scripting](Engine/Blueprints) in Unreal Engine 4, the Blueprint How To's will provide information on the core aspects needed to working with Blueprints including Creating Blueprints, Adding and Connecting Nodes, Creating Variables and others as seen in the tables below. 
 
-Refer to the table below for more information and click each link to view its contents. 
+When you are ready and have an understanding of the core aspects behind Blueprints, **Gameplay** oriented Blueprint How To's can be found in the [](Gameplay) which illustrate gameplay related elements like Spawning/Destroying Actors, Respawning (OnDestroyed), Using Timers, and more. 
+
+Refer to the tables below for more information and click each link to view its contents. 
+
+
+## Beginner
 
 |How To... | Description |
 | ------- | ------------ |
-| **[](Engine/Blueprints/BP_HowTo/BasicUsage)** | %Engine/Blueprints/BP_HowTo/BasicUsage:description% |
 | **[](Engine/Blueprints/UserGuide/Types/ClassBlueprint/Creation)** | %Engine/Blueprints/UserGuide/Types/ClassBlueprint/Creation/:description% |
 | **[](Engine/Blueprints/BP_HowTo/PlacingNodes)** | %Engine/Blueprints/BP_HowTo/PlacingNodes:description% |
 | **[](Engine/Blueprints/BP_HowTo/ConnectingNodes)** | %Engine/Blueprints/BP_HowTo/ConnectingNodes:description% |
 | **[](Engine/Blueprints/BP_HowTo/NavigatingGraphs)** | %Engine/Blueprints/BP_HowTo/NavigatingGraphs:description% |
 | **[](Engine/Blueprints/BP_HowTo/CollapsingGraphs)** | %Engine/Blueprints/BP_HowTo/CollapsingGraphs:description% |
+| **[](Engine/Blueprints/BP_HowTo/Debugging)** | %Engine/Blueprints/BP_HowTo/Debugging:description% |
+
+
+## Intermediate
+
+|How To... | Description |
+| ------- | ------------ |
 | **[](Engine/Blueprints/BP_HowTo/ActorReference)** | %Engine/Blueprints/BP_HowTo/ActorReference:description% |
 | **[](Engine/Blueprints/BP_HowTo/MakingFunctions)** | %Engine/Blueprints/BP_HowTo/MakingFunctions:description% |
 | **[](Engine/Blueprints/BP_HowTo/WorkingWithArrays)** | %Engine/Blueprints/BP_HowTo/WorkingWithArrays:description% |
 | **[](Engine/Blueprints/UserGuide/Variables)** | The entire user guide for Variables. |
+
+## Advanced
+
+|How To... | Description |
+| ------- | ------------ |
 | **[](Engine\Blueprints\BP_HowTo\DirectBlueprintComs)** | %Engine\Blueprints\BP_HowTo\DirectBlueprintComs:description% |
 | **[](Engine\Blueprints\BP_HowTo\Casting)** | %Engine\Blueprints\BP_HowTo\Casting:description% |
 | **[](Engine\Blueprints\BP_HowTo\Interfaces)** | %Engine\Blueprints\BP_HowTo\Interfaces:description% |
 | **[](Engine\Blueprints\BP_HowTo\EventDispatcher)** | %Engine\Blueprints\BP_HowTo\EventDispatcher:description% |
 | **[](Engine/Blueprints/BP_HowTo/MakingMacros)** | %Engine/Blueprints/BP_HowTo/MakingMacros:description% |
-| **[](Engine/Blueprints/BP_HowTo/Debugging)** | %Engine/Blueprints/BP_HowTo/Debugging:description% |
 
 
+<!-- To remove since information is redundant. 
+
+| **[](Engine/Blueprints/BP_HowTo/BasicUsage)** | %Engine/Blueprints/BP_HowTo/BasicUsage:description% |
+
+-->
 
 
 
