#UE4doc #docs: altered reserved words to match v4.8 of the engine.

diff --git a/Engine/Documentation/Source/Engine/Animation/AnimHowTo/AdditiveAnimations/AdditiveHowTo.INT.udn b/Engine/Documentation/Source/Engine/Animation/AnimHowTo/AdditiveAnimations/AdditiveHowTo.INT.udn
--- a/Engine/Documentation/Source/Engine/Animation/AnimHowTo/AdditiveAnimations/AdditiveHowTo.INT.udn
+++ b/Engine/Documentation/Source/Engine/Animation/AnimHowTo/AdditiveAnimations/AdditiveHowTo.INT.udn
@@ -29,9 +29,9 @@ Once downloaded, you can add the Animation Starter Pack to your project via the
 
 In this section, we assign our player character and make modifications to the animation we will blend in with our movement animations. 
 
-1.  With your project open, inside the **Content Browser** under the **Game/Blueprints** folder, open the **MyGame** Blueprint. 
+1.  With your project open, inside the **Content Browser** under the **Game/Blueprints** folder, open the **ThirdPersonGameMode** Blueprint. 
 
-1.	Under **Default Pawn Class**, click the dropdown menu and select the **ASP_Character**. 
+1.	Under **Default Pawn Class**, click the dropdown menu and select the **Ue4ASP_Character**. 
 
 	![](Animations1.png)
 
@@ -65,7 +65,7 @@ In this section, we assign our player character and make modifications to the an
 
 Next we will make some updates to our Character Blueprint so that we can determine when they are pressing the shoot button. 
 
-1.	Inside the **Content Browser** under **Game/AnimStarterPack/Character**, open the **ASP_Character** Blueprint. 
+1.	Inside the **Content Browser** under **Game/AnimStarterPack/Character**, open the **Ue4ASP_Character** Blueprint. 
 
 1.	Locate the **Crouching** section, and delete the **InputAction Crouch** Event (with the Warning!) sign. 
 
@@ -101,7 +101,7 @@ Next we will make some updates to our Character Blueprint so that we can determi
 
 Here we will set up the Animation Graph and place the nodes needed to blend between our current pose and the fire weapon animation. 
 
-1.	Also under **Game/AnimStarterPack/Character**, open the **ASP_HeroTPP_AnimBlueprint** Animation Blueprint. 
+1.	Also under **Game/AnimStarterPack/Character**, open the **UE4ASP_HeroTPP_AnimBlueprint** Animation Blueprint. 
 
 1.	Inside the **MyBlueprint** window, open the **AnimGraph**. 
 
@@ -159,7 +159,7 @@ Finally, we update the way we blend between animations so that our upperbody pla
 
 	![](Animations18b.png)
 
-1.	Drag off the **As ASP Character** pin of the **Cast To ASP_Character** node and add the **Get Fire Button Down** node. 
+1.	Drag off the **As Ue 4ASP Character** pin of the **Cast To Ue4ASP_Character** node and add the **Get Fire Button Down** node. 
 
 	![](Animations19.png)
 
