#UE4doc #docs: Altered images to match the instructional guide: green -> yellow. Updated use of reserved words to match v4.8.

diff --git a/Engine/Documentation/Source/Engine/Animation/AnimHowTo/BlendSpace/BlendspaceHowTo.INT.udn b/Engine/Documentation/Source/Engine/Animation/AnimHowTo/BlendSpace/BlendspaceHowTo.INT.udn
--- a/Engine/Documentation/Source/Engine/Animation/AnimHowTo/BlendSpace/BlendspaceHowTo.INT.udn
+++ b/Engine/Documentation/Source/Engine/Animation/AnimHowTo/BlendSpace/BlendspaceHowTo.INT.udn
@@ -28,7 +28,7 @@ Once downloaded, you can add the Animation Starter Pack to your project via the
 
 1.  With your project open, inside the **Content Browser** under the **Game/Blueprints** folder, open the **MyGame** Blueprint. 
 
-1.	Under **Default Pawn Class**, click the dropdown menu and select the **ASP_Character**. 
+1.	Under **Default Pawn Class**, click the dropdown menu and select the **Ue4ASP_Character**. 
 
 	![](Engine\Animation\AnimHowTo\AdditiveAnimations\Animations1.png)
 
@@ -84,7 +84,7 @@ Once downloaded, you can add the Animation Starter Pack to your project via the
 
 	This will smooth out the transition between the animations in the _Blend Space_ to prevent popping from pose-to-pose. 
 
-1.	Click **Save** then close the _Blend Space_, then in the **Content Browser**, open the **ASP_HeroTPP_AnimBlueprint**. 
+1.	Click **Save** then close the _Blend Space_, then in the **Content Browser**, open the **Ue4ASP_HeroTPP_AnimBlueprint**. 
 
 1.	Inside the _Animation Blueprint_, click the **MyBlueprint** tab, then expand **AnimGraph** and **Double-click Locomotion**. 
 
@@ -106,7 +106,7 @@ Once downloaded, you can add the Animation Starter Pack to your project via the
 
 1.	Click **Compile** and **Save** then close the _Animation Blueprint_. 
 
-1.	In the **Content Browser**, open the **ASP_Character** _Blueprint_. 
+1.	In the **Content Browser**, open the **Ue4ASP_Character** _Blueprint_. 
 
 1.	With the _Blueprint_ open, find the **Crouching** section and replace the **InputAction Crouch** node with a **C** Key Event. 
 
