Removed Overview and bolded/unbolded type per instruction guide.

diff --git a/Engine/Documentation/Source/Engine/Rendering/LightingAndShadows/LightTypes/Point/Point.INT.udn b/Engine/Documentation/Source/Engine/Rendering/LightingAndShadows/LightTypes/Point/Point.INT.udn
--- a/Engine/Documentation/Source/Engine/Rendering/LightingAndShadows/LightTypes/Point/Point.INT.udn
+++ b/Engine/Documentation/Source/Engine/Rendering/LightingAndShadows/LightTypes/Point/Point.INT.udn
@@ -48,10 +48,10 @@ Description:The basics of Point Lights.
 
 [TOC(start:2 end:2)]
 
-## Overview
+
 
 [EXCERPT:Intro]
-**Point Lights** work much like a real world light bulb, emitting light in all directions from the light bulb's tungsten filament. However, for the sake of performance, Point Lights are simplified down emitting light equally in all directions from just a single point in space.  The **Point Light** when placed can be set to one of three **Mobility** settings:
+**Point Lights** work much like a real world light bulb, emitting light in all directions from the light bulb's tungsten filament. However, for the sake of performance, Point Lights are simplified down emitting light equally in all directions from just a single point in space.  The Point Light when placed can be set to one of three Mobility settings:
 
 [OBJECT:topiclistnotitle]
 	[PARAM:icon]
@@ -82,20 +82,20 @@ Description:The basics of Point Lights.
 ![](spacer.png)
 [/REGION]
 
-Two examples of a **Point Light** placed inside a level are shown below.
+Two examples of a Point Light placed inside a level are shown below.
 
 ![](pointLight_004.png)(w:450) ![](pointLight_005.png)(w:450 convert:false)
 
-The image on the left is a **Point Light** without its radius showing, and the same light with it enabled is shown in the right image, giving a good impression of where the light will affect the world.
+The image on the left is a Point Light without its radius showing, and the same light with it enabled is shown in the right image, giving a good impression of where the light will affect the world.
 
-While the light from a **Point Light** only emits from that point in space, having no shape, Unreal Engine 4 can give a point light a radius and a length for use within reflections and specular highlights to give the point light more of a physical feel.
+While the light from a Point Light only emits from that point in space, having no shape, Unreal Engine 4 can give a point light a radius and a length for use within reflections and specular highlights to give the point light more of a physical feel.
 
 ![](pointLight_001.png)(w:300 convert:false) ![](pointLight_002.png)(w:300) ![](pointLight_003.png)(w:300)
 
 
 ## Point Light Properties
 
-The properties of a **Point Light** are broken up into 4 categories: **Light**, **Light Profiles**, **Lightmass**, and **Light Function**. 
+The properties of a **Point Light** are broken up into 4 categories: Light, Light Profiles, Lightmass, and Light Function. 
 
 ### Light
 Property | Description |
diff --git a/Engine/Documentation/Source/Engine/Rendering/LightingAndShadows/LightTypes/Spot/Spot.INT.udn b/Engine/Documentation/Source/Engine/Rendering/LightingAndShadows/LightTypes/Spot/Spot.INT.udn
--- a/Engine/Documentation/Source/Engine/Rendering/LightingAndShadows/LightTypes/Spot/Spot.INT.udn
+++ b/Engine/Documentation/Source/Engine/Rendering/LightingAndShadows/LightTypes/Spot/Spot.INT.udn
@@ -48,11 +48,11 @@ Description:The basics of Spot Lights.
 
 [TOC(start:2 end:2)]
 
-## Overview
+
 [EXCERPT:Intro]
 A **Spot Light** emits light from a single point in a cone shape. Users are given two cones to shape the light - the **Inner Cone Angle** and **Outer Cone Angle**. Within the _Inner Cone Angle_, the light achieves full brightness. As you go from the extent of the inner radius to the extents of the _Outer Cone Angle_, a falloff takes place, creating a penumbra, or softening around the Spot Light's disc of illumination. The Radius of the light defines the length of the cones. More simply, this will work like a flash light or stage can light. 
 
-Like other lights, a **Spot Light** can be set to one of three **Mobility** settings: 
+Like other lights, a Spot Light can be set to one of three Mobility settings: 
 [/EXCERPT:Intro]
 
 
@@ -63,7 +63,7 @@ Like other lights, a **Spot Light** can be set to one of three **Mobility** sett
 	[PARAM:description]
 	**Static** - (pictured left) which means that the light cannot be changed in game. This is the fastest method for rendering and allows for baked lighting.
 
-	**Stationary** - (also pictured left) which means the light will only have its shadowing and bounced lighting from static geometry baked by **Lightmass**, all other lighting will be dynamic. This setting also allows for the light to change color and intensity in game, but, it does not move and allows partial baked lighting. 
+	**Stationary** - (also pictured left) which means the light will only have its shadowing and bounced lighting from static geometry baked by Lightmass, all other lighting will be dynamic. This setting also allows for the light to change color and intensity in game, but, it does not move and allows partial baked lighting. 
 	[/PARAM]
 [/OBJECT]
 
@@ -88,12 +88,12 @@ Two examples of a **Spot Light** placed inside a level are shown below demonstra
 
 ![](spotlight_example.png)
 
-On the left a **Spot Light** without its cones showing light radius, on the right the same **Spot Light** with its cone effectors showing.
+On the left a Spot Light without its cones showing light radius, on the right the same Spot Light with its cone effectors showing.
 
 
 ## Spot Light Properties
 
-The properties of a **Spot Light** are broken up into 4 categories: **Light**, **Light Profiles**, **Lightmass**, and **Light Function**. 
+The properties of a Spot Light are broken up into 4 categories: Light, Light Profiles, Lightmass, and Light Function. 
 
 ### Light
 Property | Description |
diff --git a/Engine/Documentation/Source/Engine/Rendering/LightingAndShadows/LitTranslucency/LitTranslucency.INT.udn b/Engine/Documentation/Source/Engine/Rendering/LightingAndShadows/LitTranslucency/LitTranslucency.INT.udn
--- a/Engine/Documentation/Source/Engine/Rendering/LightingAndShadows/LitTranslucency/LitTranslucency.INT.udn
+++ b/Engine/Documentation/Source/Engine/Rendering/LightingAndShadows/LitTranslucency/LitTranslucency.INT.udn
@@ -25,7 +25,7 @@ Description:Explanation of how translucent surfaces are lit and cast shadows inc
 
 [TOC(start:2)]
 
-## Overview
+
 
 Translucent effects generally fall into a few categories: volumetric, volumetric but dense enough to have normal information, and surfaces.  Different lighting techniques are needed for each of these, so a material must specify the Translucency Lighting Mode that should be used.
 
@@ -49,7 +49,7 @@ _The left sphere is lit translucency using the *Indirect Lighting Cache*, the ri
 
 Translucency can cast shadows onto the opaque world and onto itself and other lit translucency Actors. This is implemented with Fourier Opacity maps, which do a great job for shadowing from blobby volumes, but have severe ringing artifacts with more opaque translucent surfaces. Translucency self shadowing goes through the lighting volume for point and spot lights, so it is often not visible due to low resolution unless the effect is very large and dense. Directional lights however do translucent self shadowing per-pixel, and get much higher quality. Directional lights also have subsurface shading for lit materials using the subsurface shading model.
 
-Translucent self-shadowing uses per-object shadows, which means that it needs user specified fixed particle system bounds and they need to be correct. The easiest way to set this up is to author your particle movement, then **right-click** on the 'show bounds' button on the Cascade toolbar, which will pop up a dialog that allows you to generate fixed bounds.  Very large self-shadowing particle systems will get reduced shadowmap resolution, since the shadowmap is stretched to cover the system bounds.  Verify that your bounds are reasonable by enabling Show Bounds, which can be found under **Show** -> **Advanced** -> **Bounds**.  Then select the emitter in the editor and it will draw the bounding box and sphere.
+Translucent self-shadowing uses per-object shadows, which means that it needs user specified fixed particle system bounds and they need to be correct. The easiest way to set this up is to author your particle movement, then **right-click** on the 'show bounds' button on the Cascade toolbar, which will pop up a dialog that allows you to generate fixed bounds.  Very large self-shadowing particle systems will get reduced shadowmap resolution, since the shadowmap is stretched to cover the system bounds.  Verify that your bounds are reasonable by enabling Show Bounds, which can be found under Show -> Advanced -> Bounds.  Then select the emitter in the editor and it will draw the bounding box and sphere.
 
 [REGION:imagetable]
 ![Translucent Particle Self-Shadowing](Engine/Rendering/Overview/self_shadow_1.png)(w:456) | ![Translucent Particle Self-Shadowing](Engine/Rendering/Overview/self_shadow_2.png)(w:456) |
@@ -69,6 +69,6 @@ _The left sphere is translucent, the right is opaque, and they are both setup as
 
 ## Limitations
 
-* **Static** shadowing from static lights is currently not handled for lit translucency. However, **dynamic** shadows from stationary lights is supported.
+* Static shadowing from static lights is currently not handled for lit translucency. However, dynamic shadows from stationary lights is supported.
 * Lit translucent surfaces are missing direct specular.
 * Lit translucent surfaces get all their direct lighting through the translucency volume lighting texture, which causes them to be lower resolution than needed for most surface materials (glass, water).
diff --git a/Engine/Documentation/Source/Engine/Rendering/Materials/Functions/Reference/Blends/Blends.INT.udn b/Engine/Documentation/Source/Engine/Rendering/Materials/Functions/Reference/Blends/Blends.INT.udn
--- a/Engine/Documentation/Source/Engine/Rendering/Materials/Functions/Reference/Blends/Blends.INT.udn
+++ b/Engine/Documentation/Source/Engine/Rendering/Materials/Functions/Reference/Blends/Blends.INT.udn
@@ -5,7 +5,7 @@ Description:Functions designed to blend one color with another, similar to blend
 
 [TOC (start:2 end:3)]
 
-## Overview
+
 
 A Blend is a type of function that performs mathematical calculations in the color information of a texture so that one texture can blend into another in a particular manner.
 
