#UEdoc #docs updated image for groupings setting toolbar

diff --git a/Engine/Documentation/Source/Engine/Actors/FogEffects/HeightFog/HeightFog.INT.udn b/Engine/Documentation/Source/Engine/Actors/FogEffects/HeightFog/HeightFog.INT.udn
--- a/Engine/Documentation/Source/Engine/Actors/FogEffects/HeightFog/HeightFog.INT.udn
+++ b/Engine/Documentation/Source/Engine/Actors/FogEffects/HeightFog/HeightFog.INT.udn
@@ -52,7 +52,6 @@ Description:Guide to creating height-based, distance fog in levels.
 
 
 
-## Overview
 
 This document covers Exponential Height Fog and its uses. It assumes that you have a level already set up in which you would like to add Exponential Height Fog.
 
