#UE4doc #docs: Removed Overview.

diff --git a/Engine/Documentation/Source/Engine/Rendering/Materials/Functions/Reference/Gradient/Gradient.INT.udn b/Engine/Documentation/Source/Engine/Rendering/Materials/Functions/Reference/Gradient/Gradient.INT.udn
--- a/Engine/Documentation/Source/Engine/Rendering/Materials/Functions/Reference/Gradient/Gradient.INT.udn
+++ b/Engine/Documentation/Source/Engine/Rendering/Materials/Functions/Reference/Gradient/Gradient.INT.udn
@@ -5,7 +5,7 @@ Description:Procedurally generated gradients to add to your materials, eliminati
 
 [TOC (start:2 end:3)]
 
-## Overview
+
 
 The Gradient functions produce procedurally generated gradients made from Texture coordinate expressions. They save memory over having to create a texture-based gradient.
 
diff --git a/Engine/Documentation/Source/Engine/Rendering/Materials/Functions/Reference/ImageAdjustment/ImageAdjustment.INT.udn b/Engine/Documentation/Source/Engine/Rendering/Materials/Functions/Reference/ImageAdjustment/ImageAdjustment.INT.udn
--- a/Engine/Documentation/Source/Engine/Rendering/Materials/Functions/Reference/ImageAdjustment/ImageAdjustment.INT.udn
+++ b/Engine/Documentation/Source/Engine/Rendering/Materials/Functions/Reference/ImageAdjustment/ImageAdjustment.INT.udn
@@ -5,7 +5,7 @@ Description:Functions for making adjustments to existing image textures, such as
 
 [TOC (start:2 end:3)]
 
-## Overview
+
 
 The Image Adjustment functions exist as a way to perform basic color correction operations on textures. They are useful in that they allow corrective actions or variations on a texture without having to worry about the overhead of loading a separate version into memory.
 
