#UE4doc #docs: updated reserved words to match v4.8.

diff --git a/Engine/Documentation/Source/Engine/Animation/AnimHowTo/LayerEditing/LayerEditingHowTo.INT.udn b/Engine/Documentation/Source/Engine/Animation/AnimHowTo/LayerEditing/LayerEditingHowTo.INT.udn
--- a/Engine/Documentation/Source/Engine/Animation/AnimHowTo/LayerEditing/LayerEditingHowTo.INT.udn
+++ b/Engine/Documentation/Source/Engine/Animation/AnimHowTo/LayerEditing/LayerEditingHowTo.INT.udn
@@ -9,7 +9,7 @@ Related: Engine\Animation\AnimHowTo\AdditiveAnimations
 Related: Engine\Animation\Skeleton
 Related: Engine\Animation\Overview
 
-**Animation Layer Editing** allows you to modify any existing [](Engine\Animation\Sequences) inside Unreal Engine 4 as part of **Additive Layer Tracks**. With Animation Layer Editing you can take an existing animation and adjust it to produce a different take on that animation, or you can create a separate, new animation with its own motion you provide based off the original. 
+**Animation Layer Editing** allows you to modify any existing [](Engine\Animation\Sequences) inside Unreal Engine 4 as part of Additive Layer Tracks. With Animation Layer Editing you can take an existing animation and adjust it to produce a different take on that animation, or you can create a separate, new animation with its own motion you provide based off the original. 
 
 This page will show you how to use the Animation Layer Editing tools and take an Idle animation and repurpose it into a Reload animation. 
 
@@ -23,7 +23,7 @@ For this guide, we created a New Project using the **Blueprint First Person** te
 
 In this section, we will create a copy of the Idle animation that we can then modify, leaving the original intact. 
 
-1.  Inside the **Content Browser**, under the **Game/FirstPersonBP/Animations** folder, open the **FirstPerson_Idle** animation. 
+1.  Inside the **Content Browser**, under the **Game/FirstPersonBP/Animations** folder, open the **Prone_Idle** animation. 
 
 	![](AnimEdit1.png)
 
@@ -51,7 +51,7 @@ With the **FirstPerson_Reload** animation created, we can begin editing it throu
 
 	![](AnimEdit5.png)
 
-1.	In the **Skeleton Tree**, click the **b_LeftArm** Bone so that it is selected. 
+1.	In the **Skeleton Tree**, click the **upperarm_l** Bone so that it is selected. 
 
 	![](AnimEdit6.png)
 
@@ -69,7 +69,7 @@ With the **FirstPerson_Reload** animation created, we can begin editing it throu
 
 	![](AnimEdit8.png)(w:520)
 
-	Above, three separate tracks were added for the **b_LeftArm** Bone: a **Translation** track (1), a **Rotation** track (2) and a **Scale** track (3).
+	Above, three separate tracks were added for the **upperarm_l** Bone: a **Translation** track (1), a **Rotation** track (2) and a **Scale** track (3).
 
 1.	Click the **ToNext** button to advance to frame 30. 
 
