#UEdoc #docs - Updated images and removed overview header

diff --git a/Engine/Documentation/Source/Engine/Actors/Placement/PlacingActors.INT.udn b/Engine/Documentation/Source/Engine/Actors/Placement/PlacingActors.INT.udn
--- a/Engine/Documentation/Source/Engine/Actors/Placement/PlacingActors.INT.udn
+++ b/Engine/Documentation/Source/Engine/Actors/Placement/PlacingActors.INT.udn
@@ -69,8 +69,6 @@ At the most fundamental level, an **Actor** is any Object that you can place in
 
 [TOC(start:2)]
 
-## Overview
-
 [EXCERPT:Place01]
 When you create a new project, the level that is loaded when you open your project is typically a blank slate. To build the desired environment or populate the world, you must place **Actors** in the level. You can do this in two different ways, both of which involve using the **Content Browser** and results in a new instance of the Actor being created. Each instance in the level can then be moved around or have its properties modified.
 [/EXCERPT:Place01]
