Test

diff --git a/Engine/Documentation/Source/Engine/Actors/Grouping/Grouping.INT.udn b/Engine/Documentation/Source/Engine/Actors/Grouping/Grouping.INT.udn
--- a/Engine/Documentation/Source/Engine/Actors/Grouping/Grouping.INT.udn
+++ b/Engine/Documentation/Source/Engine/Actors/Grouping/Grouping.INT.udn
@@ -67,7 +67,6 @@ Description:Guide to creating and working with groups of actors within Unreal Ed
 
 [TOC(start:2)]
 
-## Overview
 
 In the Unreal Engine, you can group Actors together, which allows you to easily manage multiple Actors simultaneously.  Grouped Actors can be transformed as a single unit through rotation, translation, and scaling. Groups can be unlocked to temporarily allow manipulation of individual Actors, and then locked again to freeze the Actors inside the group to prevent changes to an individual Actor's transforms. You can add Actors to, or remove them from, a group. You can also nest groups by having one group as a member of another group.
 
