test

diff --git a/Engine/Documentation/Source/Engine/Actors/FogEffects/AtmosphericFog/AtmosphericFog.INT.udn b/Engine/Documentation/Source/Engine/Actors/FogEffects/AtmosphericFog/AtmosphericFog.INT.udn
--- a/Engine/Documentation/Source/Engine/Actors/FogEffects/AtmosphericFog/AtmosphericFog.INT.udn
+++ b/Engine/Documentation/Source/Engine/Actors/FogEffects/AtmosphericFog/AtmosphericFog.INT.udn
@@ -68,9 +68,6 @@ This may be done via drag-and-drop:
 
 ![AtmoActor.png](AtmoActor.png)(w:420)
 
-Or by selecting the Actor from the Class Viewer, and choosing to add it via the **right-click** context menu:
-
-![Atmo_Context.png](Atmo_Context.png)(w:420)
 
 ## Atmospheric Fog Properties
 
@@ -101,9 +98,9 @@ The Atmospheric Fog Actor contains the following properties:
 [REGION:imagetable]
 | ![Atmo0.5_4.png](Atmo0.5_4.png)(w:460) | ![Atmo0.35_4.png](Atmo0.35_4.png)(w:460) |
 | ------ | ------ |
-| Decay Height of 0.5 (8 km) | Decay Height of 0.35 (2.744 km) |
+| Density Decay Height of 0.5 (8 km) | Density Decay Height of 0.35 (2.744 km) |
 | ![Atmo1.0_4.png](Atmo1.0_4.png)(w:460) | ![Atmo1.0_1.png](Atmo1.0_1.png)(w:460) |
-| Decay Height of 1.0 (64 km) and Max Scattering of 4 | Decay Height of 1.0 (64 km) and Max Scattering of 1 |
+| Density Decay Height of 1.0 (64 km) and Max Scattering of 4 | Density Decay Height of 1.0 (64 km) and Max Scattering of 1 |
 [/REGION]
 [/EXCERPT:examples]
 <!---Includes
@@ -111,7 +108,7 @@ The Atmospheric Fog Actor contains the following properties:
 [REGION:simpletable]
 	|**Property** | **Description** |
 	| --- | --- |
-	| *Decay Height* |This value controls fog density decay height. Lower values cause the fog to be denser, while higher values thin the fog, causing less scatter. Values range from [0..1] |
+	| *Density Height* |This value controls fog density decay height. Lower values cause the fog to be denser, while higher values thin the fog, causing less scatter. Values range from [0..1] |
 	| *Max Scattering Num* | This sets a limit on the number of scattering calculations, from 1 (single scattering) to 4 (quadruple scattering). |
 	| *Inscatter Altitude Sample* | Number of different altitudes at which to sample inscatter color (Size of 3d texture dimension; 1-32). |
 [/REGION]
