#UE4doc #docs Added tabs for compiling with Windows and Mac.

diff --git a/Engine/Documentation/Source/Programming/QuickStart/3/QuickStart_3.INT.udn b/Engine/Documentation/Source/Programming/QuickStart/3/QuickStart_3.INT.udn
--- a/Engine/Documentation/Source/Programming/QuickStart/3/QuickStart_3.INT.udn
+++ b/Engine/Documentation/Source/Programming/QuickStart/3/QuickStart_3.INT.udn
@@ -29,11 +29,11 @@ Version: 4.8
 
 	[REGION:note]We're about to write code. All code used in this tutorial will be found at bottom of the page in its current state as of completing that page's instructions.[/REGION]
 
-2. In `FloatingActor.h`, we'll add the following code just before the closing brace and semicolon at the end of the file:
+1. In `FloatingActor.h`, we'll add the following code just before the closing brace and semicolon at the end of the file:
 
 		float RunningTime;
 
-3. Switching to `FloatingActor.cpp`, we'll add the following code just before the closing brace at the bottom of **AFloatingActor::Tick**:
+1. Switching to `FloatingActor.cpp`, we'll add the following code just before the closing brace at the bottom of **AFloatingActor::Tick**:
 
 		FVector NewLocation = GetActorLocation();
 		float DeltaHeight = (FMath::Sin(RunningTime + DeltaTime) - FMath::Sin(RunningTime));
@@ -43,15 +43,67 @@ Version: 4.8
 
 	[REGION:tip] The code we've just written will cause **FloatingActors** to bob up and down smoothly, using the _RunningTime_ variable we created to keep track of our movement over time. [/REGION] 
 
-4. Now that we're done coding, we can compile by right-clicking our project in the **Solution Browser** and selecting the **Build** command, or by pressing the **Compile** button in the **Unreal Editor**. Once the compile succeeds, **Unreal Engine** will automatically load our changes.
-
-	[![](CompileFromVS.png)](CompileFromVS.png)
+%Globals:OSSelect%
+
+[OBJECT:ToggleButtonContent]
+	[PARAMLITERAL:category]
+	OS
+	[/PARAMLITERAL]
+	[PARAMLITERAL:id]
+	windows
+	[/PARAMLITERAL]
+	[PARAMLITERAL:active]
+	active_button_content
+	[/PARAMLITERAL]
+	[PARAM:content]
+		*	Now that we're done coding, we can compile by right-clicking our project in the **Solution Browser** and selecting the **Build** command, or by pressing the **Compile** button in the **Unreal Editor**. Once the compile succeeds, **Unreal Engine** will automatically load our changes.
+
+			[![](CompileFromVS.png)](CompileFromVS.png)
 	
-	(in **Visual Studio**)
+			(compiling from Visual Studio pictured above)
+
+			[![](CompileFromEditor.png)](CompileFromEditor.png)
 
-	[![](CompileFromEditor.png)](CompileFromEditor.png)
+			(compiling from the Unreal Editor pictured above)
+	[/PARAM]
+[/OBJECT]
 
-	(in the **Unreal Editor**)
+[OBJECT:ToggleButtonContent]
+	[PARAMLITERAL:category]
+	OS
+	[/PARAMLITERAL]
+	[PARAMLITERAL:id]
+	mac
+	[/PARAMLITERAL]
+	[PARAMLITERAL:active]
+	active_button_content
+	[/PARAMLITERAL]
+	[PARAM:content]
+		[REGION:tip]
+		When building in Xcode, you are only compiling the game project, not the Editor.
+		[/REGION]
+		
+		*	Now that we're done coding, we can compile our project by using XCode's Debug Configuration by clicking the **Product > Build** option.
+		
+			![](ProjectBuild_mac.png)
+			
+			(compiling using XCode's Debug Configuration pictured above)
+
+		*	We can also compile using XCode's Development Configuration by clicking the **Product > Build For > Profiling** option.
+
+			![](ProjectBuildProfile_mac.png)
+
+			(compiling using XCode's Development Configuration pictured above)
+
+			[![](CompileFromEditor.png)](CompileFromEditor.png)
+
+			(compiling from the **Unreal Editor** pictured above)
+
+		[REGION:note]
+		When running the binary Editor, it is important to add the -game flag if you rebuilt your project in any **Uncooked** configuration, and the -debug flag if you rebuilt your project in any **Debug** configuration.
+		[/REGION]
+	[/PARAM]
+[/OBJECT]
 
 ----
 
diff --git a/Engine/Documentation/Source/Programming/QuickStart/QuickStart.INT.udn b/Engine/Documentation/Source/Programming/QuickStart/QuickStart.INT.udn
--- a/Engine/Documentation/Source/Programming/QuickStart/QuickStart.INT.udn
+++ b/Engine/Documentation/Source/Programming/QuickStart/QuickStart.INT.udn
@@ -55,18 +55,18 @@ Version: 4.8
 
 ## Goals
 
-The primary goal of this tutorial is to introduce programmers to Unreal Engine 4's development environment. By the end of this quick start guide, developers should be comfortable setting up and developing **C++ Projects** in Unreal Engine 4 (UE4).
+The primary goal of this tutorial is to introduce programmers to Unreal Engine 4's development environment. By the end of this quick start guide, developers should be comfortable setting up and developing C++ Projects in Unreal Engine 4 (UE4).
 
 ## Objectives
 
 Upon completion, developers will have the skill to:
 
-* **Create a New Unreal Engine Project**
-* **Add a New C++ Class**
-* **Compile a Project** 
-* **Add an Instance of a New Class to a Level**
+* Create a New Unreal Engine Project
+* Add a New C++ Class
+* Compile a Project 
+* Add an Instance of a New Class to a Level
 
-_When we're finished, we'll be able to see an **Actor** we programmed in C++ moving around on screen._
+_When we're finished, we'll be able to see an Actor we programmed in C++ moving around on screen._
 
 [OBJECT:TopicButtonList]
 	[PARAM:icon]
