#UE4 Docs: Converting analytics docs from UEDOC-1614 + minor section re-org

diff --git a/Engine/Documentation/Source/Gameplay/Analytics/Analytics.INT.udn b/Engine/Documentation/Source/Gameplay/Analytics/Analytics.INT.udn
--- a/Engine/Documentation/Source/Gameplay/Analytics/Analytics.INT.udn
+++ b/Engine/Documentation/Source/Gameplay/Analytics/Analytics.INT.udn
@@ -3,46 +3,39 @@ Title:In-Game Analytics
 Crumbs: %ROOT%, Gameplay
 Description:Using in-game analytics to track player engagement and find balance issues.
 
-[TOC (start:2 end:2)]
 
-## Overview
 In order to get data on how your game is performing, you must use an analytics provider to capture and process the data. For those without a homegrown solution, there are plenty of 
 options available, from free services to paid ones. Unreal Engine provides an abstract interface for communicating with one or more analytics providers. Your game uses the interface, 
 and analytics providers offer a backing implementation. In some cases, Epic has built the backing provider already. Before version 4.5, Epic provided an implementation to multicast analytics
 events, relaying them to multiple providers. For providers, there is support of providers that support [Swrve](http://www.swrve.com), a paid service. Unreal Engine 4.5 and later provides
-support for [Apsalar](http://www.apsalar.com), a free service on iOS, and a file based system primarily used for debugging. More provider plugins will become 
-available over time, and it is also feasible to add your own provider if needed. 
-
-## Instrumenting Your Game
-The first step in capturing player retention data is to register an analytics provider for your game. This is done via your project's `DefaultEngine.ini` file. You must register a 
-default provider. Optionally, you can register different providers and account details for different build types of your game (development, testing, and production). The sections below 
-are an example of configuring the Apsalar plugin:
-
-	[Analytics]
-	ProviderModuleName=IOSApsalar
-	ApiKey=YourAnalyticsKey1
-	ApiSecret=YourAnalyticsSecret1
-	SendInterval=60
-
-	[AnalyticsDevelopment]
-	ApiKey=YourAnalyticsKey2
-	ApiSecret=YourAnalyticsSecret2
-	SendInterval=60
-
-	[AnalyticsTest]
-	ApiKey=YourAnalyticsKey3
-	ApiSecret=YourAnalyticsSecret4
-	SendInterval=60
-
-The `[Analytics]` section is the default one used and is where you should set the name of the default provider module. In the case above, it is set to the IOSApsalar plugin that is part 
-of the 4.5 release. The `ApiKey` and `ApiSecret` fields come from the Apsalar website. Once you create an account, they will give you a key and secret to use.
-
-Once you have it configured for your project, you are ready to start recording analytics events. To get just the basic player retention data, you need to create a session when the 
-game starts up and end it when it is no longer in the foreground. This can be done using the lines of code shown below, or the [Blueprint analytics plugin](Gameplay/Analytics/Blueprints)
-provided in versions 4.5 and later:
-
-	FAnalytics::Get().GetDefaultConfiguredProvider()->StartSession();
-	FAnalytics::Get().GetDefaultConfiguredProvider()->EndSession();
-
-With those calls as part of your game, you will automatically start gathering player retention data. After getting basic player retention data, you can start adding more events to 
-tell you even more about player behavior in your game. 
+support for Apsalar and a file based system primarily used for debugging. Support for Flurry was also added recently in version 4.8.
+
+More provider plugins will become available over time, and it is also feasible to add your own provider if needed. 
+
+[OBJECT:TopicButtonList]
+	[PARAM:icon]
+		![](%ROOT%/profile_icon.png)(convert:false)
+	[/PARAM]
+	[PARAM:title]
+		Implementing Game Analytics
+	[/PARAM]
+	[PARAM:links]
+		* [](Gameplay/Analytics/Instrumenting "%Gameplay/Analytics/Instrumenting:description%")
+		* [](Gameplay/Analytics/Blueprints "%Gameplay/Analytics/Blueprints:description%")
+	[/PARAM]
+[/OBJECT]
+
+[OBJECT:TopicButtonList]
+	[PARAM:icon]
+		![](%ROOT%/build_icon.png)(convert:false)
+	[/PARAM]
+	[PARAM:title]
+		Providers
+	[/PARAM]
+	[PARAM:links]
+		* [](Gameplay/Analytics/Apsalar "%Gameplay/Analytics/Apsalar:description%")
+		* [](Gameplay/Analytics/FileLogging "%Gameplay/Analytics/FileLogging:description%")
+		* [](Gameplay/Analytics/Flurry "%Gameplay/Analytics/Flurry:description%")
+		* [](Gameplay/Analytics/Multicast "%Gameplay/Analytics/Multicast:description%")
+	[/PARAM]
+[/OBJECT]
\ No newline at end of file
diff --git a/Engine/Documentation/Source/Gameplay/Analytics/Apsalar/Apsalar.INT.udn b/Engine/Documentation/Source/Gameplay/Analytics/Apsalar/Apsalar.INT.udn
new file mode 100644
--- /dev/null
+++ b/Engine/Documentation/Source/Gameplay/Analytics/Apsalar/Apsalar.INT.udn
@@ -0,0 +1,26 @@
+Availability:Public
+Title:Apsalar
+Crumbs: %ROOT%, Gameplay, Gameplay/Analytics
+Description:Apsalar is a a Facebook Mobile Ads analytics service supported in recent versions of Unreal Engine.
+
+![](image00.png)
+
+[Apsalar](http://www.apsalar.com) started as a free analytics service that can perform Facebook advertisement attribution. It recently switched to a fee based model. It's a popular service, 
+since you can measure the return on investment for Facebook Mobile Ads, which only a small number of analytics providers can do. To use their service, you must register at their site, get an
+ application key and password that uniquely identifies your app, and download the libraries which get compiled into the Apsalar plugin. See the plugin's corresponding `<PlatformAndName>.Build.cs`
+ file to see where the libraries and headers are expected to be placed. 
+ 
+## Configuration
+
+Once you have done the prerequisites and have successfully built the plugin for the target platform, you can configure the plugin for your game. As of 4.5, there are five configuration 
+options. The `ApiKey` and `ApiSecret` are used to uniquely identify your application. The `SendInterval` setting controls how often data is sent to the Apsalar servers. `MaxBufferSize` is the
+ size of the buffer it keeps before flushing to the service. `ManuallyReportRevenue` controls whether Apsalar listens to in-app purchase receipts or not. If set to false, Apsalar will 
+ automatically report revenue for your application. The only required properties are the `ApiKey` and `ApiSecret`. You can leave the others out of your configuration and Apsalar's defaults 
+ will be used. The snippet below shows a theoretical configuration for Apsalar. As with all analytics providers, the configuration data goes in your `DefaultEngine.ini` file.
+
+	[Analytics]
+	ApiKey=MyAppsAnalytics
+	ApiSecret=Ga35tYi
+	SendInterval=60
+	MaxBufferSize=32767
+	ManuallyReportRevenue=true
\ No newline at end of file
diff --git a/Engine/Documentation/Source/Gameplay/Analytics/FileLogging/FileLogging.INT.udn b/Engine/Documentation/Source/Gameplay/Analytics/FileLogging/FileLogging.INT.udn
new file mode 100644
--- /dev/null
+++ b/Engine/Documentation/Source/Gameplay/Analytics/FileLogging/FileLogging.INT.udn
@@ -0,0 +1,19 @@
+Availability:Public
+Title:File Logging Analytics Provider
+Crumbs: %ROOT%, Gameplay, Gameplay/Analytics
+Description:The file logging provider can be used to debug your analytics process during development.
+
+![](image00.png)
+
+This provider is used to write analytics API calls to disk in JSON format. The reason to use this provider is for debugging the analytics process. It writes the data to the 
+`Saved/Analytic`s folder giving each file a unique name that ends in `.analytics`. You can then compare the data saved in that file to the events on your analytics provider's dashboard to 
+make sure all of the data is being processed. There aren't any special configuration settings for this provider.
+
+[REGION:note]
+This provider writes each session's data to disk which will cause ever-growing data in a released game or app. We recommend that you only use this provider for development and 
+don't include it in your released product.
+[/REGION]
+
+The image below shows files created when testing the 4.8 API additions.
+
+![](image01.png)
\ No newline at end of file
diff --git a/Engine/Documentation/Source/Gameplay/Analytics/Flurry/Flurry.INT.udn b/Engine/Documentation/Source/Gameplay/Analytics/Flurry/Flurry.INT.udn
new file mode 100644
--- /dev/null
+++ b/Engine/Documentation/Source/Gameplay/Analytics/Flurry/Flurry.INT.udn
@@ -0,0 +1,20 @@
+Availability:Public
+Title:Blueprint Analytics Plugin
+Crumbs: %ROOT%, Gameplay, Gameplay/Analytics
+Description:A set of Blueprint nodes provided to allow you to communicate with analytics services
+
+![](image00.png)
+
+[Flurry](http://www.flurry.com) is a widely used free analytics service. Because it is so widely used, it can compare your app's data to other applications in the same category. This lets you quickly get a feeling of 
+how your game is doing and areas you might need to do some work. To use their service, you must register at their site, get an application key that uniquely identifies your app, and 
+download the libraries which get compiled into the Flurry plugin. See the plugin's corresponding `<PlatformAndName>.Build.cs` file to see where the libraries and headers are expected to be placed. 
+To get started, go to .
+
+## Configuration
+
+Once you have done the prerequisites and have successfully built the plugin for the target platform, you can configure the plugin for your game. As of 4.8, there is only one configuration 
+property to set: the key that uniquely identifies your game. The snippet below shows a theoretical configuration for Flurry. As with all analytics providers, the configuration data goes in 
+your `DefaultEngine.ini` file.
+
+	[Analytics]
+	FlurryApiKey=RANDOM34LETTERS4511
\ No newline at end of file
diff --git a/Engine/Documentation/Source/Gameplay/Analytics/Instrumenting/Instrumenting.INT.udn b/Engine/Documentation/Source/Gameplay/Analytics/Instrumenting/Instrumenting.INT.udn
new file mode 100644
--- /dev/null
+++ b/Engine/Documentation/Source/Gameplay/Analytics/Instrumenting/Instrumenting.INT.udn
@@ -0,0 +1,38 @@
+Availability:Public
+Title:Instrumenting Your Game
+Crumbs: %ROOT%, Gameplay, Gameplay/Analytics
+Description:Using in-game analytics to track player engagement and find balance issues.
+
+
+The first step in capturing player retention data is to register an analytics provider for your game. This is done via your project's `DefaultEngine.ini` file. You must register a 
+default provider. Optionally, you can register different providers and account details for different build types of your game (development, testing, and production). The sections below 
+are an example of configuring the Apsalar plugin:
+
+	[Analytics]
+	ProviderModuleName=IOSApsalar
+	ApiKey=YourAnalyticsKey1
+	ApiSecret=YourAnalyticsSecret1
+	SendInterval=60
+
+	[AnalyticsDevelopment]
+	ApiKey=YourAnalyticsKey2
+	ApiSecret=YourAnalyticsSecret2
+	SendInterval=60
+
+	[AnalyticsTest]
+	ApiKey=YourAnalyticsKey3
+	ApiSecret=YourAnalyticsSecret4
+	SendInterval=60
+
+The `[Analytics]` section is the default one used and is where you should set the name of the default provider module. In the case above, it is set to the IOSApsalar plugin that is part 
+of the 4.5 release. The `ApiKey` and `ApiSecret` fields come from the Apsalar website. Once you create an account, they will give you a key and secret to use.
+
+Once you have it configured for your project, you are ready to start recording analytics events. To get just the basic player retention data, you need to create a session when the 
+game starts up and end it when it is no longer in the foreground. This can be done using the lines of code shown below, or the [Blueprint analytics plugin](Gameplay/Analytics/Blueprints)
+provided in versions 4.5 and later:
+
+	FAnalytics::Get().GetDefaultConfiguredProvider()->StartSession();
+	FAnalytics::Get().GetDefaultConfiguredProvider()->EndSession();
+
+With those calls as part of your game, you will automatically start gathering player retention data. After getting basic player retention data, you can start adding more events to 
+tell you even more about player behavior in your game. 
diff --git a/Engine/Documentation/Source/Gameplay/Analytics/Multicast/Multicast.INT.udn b/Engine/Documentation/Source/Gameplay/Analytics/Multicast/Multicast.INT.udn
new file mode 100644
--- /dev/null
+++ b/Engine/Documentation/Source/Gameplay/Analytics/Multicast/Multicast.INT.udn
@@ -0,0 +1,21 @@
+Availability:Public
+Title:Multicast Analytics Provider
+Crumbs: %ROOT%, Gameplay, Gameplay/Analytics
+Description:A set of Blueprint nodes provided to allow you to communicate with analytics services
+
+![](image00.png)
+
+The multicast provider sends analytics events to multiple providers. This allows you to use multiple analytics providers easily, since it will take the calls and hand them to each 
+registered analytics providers in turn without you having to manually do this dispatching work. Having multiple analytics providers is desirable since each provider has different 
+strengths and weaknesses. By using multiple providers, you can get all of the features that you need to run your app business.
+
+
+## Configuration
+
+Configuring the provider is straightforward. It just needs to know the list of providers that you want data sent to. This list is provided as a comma-delimited list of the module 
+names for the providers. In the example below, the AnalyticsMulticast provider is specified as the default provider for the application. Then, the list of providers that it sends
+ the data to is specified using a comma-delimited list as mentioned before. As with all analytics providers, the configuration data goes in your `DefaultEngine.ini` file.
+
+	[Analytics]
+	ProviderModuleName=AnalyticsMulticast
+	ProviderModuleNames=FileLogging,IOSFlurry,IOSApsalar
\ No newline at end of file
