missed page checked in

diff --git a/Engine/Documentation/Source/Engine/Actors/PlayerStart/PlayerStart.INT.udn b/Engine/Documentation/Source/Engine/Actors/PlayerStart/PlayerStart.INT.udn
--- a/Engine/Documentation/Source/Engine/Actors/PlayerStart/PlayerStart.INT.udn
+++ b/Engine/Documentation/Source/Engine/Actors/PlayerStart/PlayerStart.INT.udn
@@ -5,7 +5,6 @@ Description:Guide to using Player Starts.
 
 ![](Player_Start_Actor.png)
 
-## Overview
 
 The ability to spawn a player into the world at any location is a very useful feature for any game. Unreal Engine 4 offers a special Actor that will allow you to do this called a **Player Start**. A Player Start is just that, a location in the game world that the player will start from.  
 
