Updated documentation for tiled landscape import

diff --git a/Engine/Documentation/Source/Engine/LevelStreaming/WorldBrowser/WorldBrowser.INT.udn b/Engine/Documentation/Source/Engine/LevelStreaming/WorldBrowser/WorldBrowser.INT.udn
--- a/Engine/Documentation/Source/Engine/LevelStreaming/WorldBrowser/WorldBrowser.INT.udn
+++ b/Engine/Documentation/Source/Engine/LevelStreaming/WorldBrowser/WorldBrowser.INT.udn
@@ -184,6 +184,7 @@ In the **Levels** panel toolbar you can find an option for importing tiled lands
 | Setting | Description |
 | ------- | ----------- |
 | Select Heightmap Tiles... | Allows you to select multiple heightmap tiles you want to import. |
+| Flip Tile Y Coordinate | Whether tiles should be imported with Y coordinate flipped. This option should be turned on when importing tiles created in World Machine.  |
 | Tile Coordinates Offset | Offset in tiles from origin. If set to zero, tile with coordinates x0_y0 will be created at world origin. |
 | Import Configuration | Generates all possible landscape configurations depending on resolution of selected tiles. |
 | Landscape Scale | How much to scale the tiled landscape by.|
