#UE4 Docs: Adding the first pass for the new section about the global menu.

diff --git a/Engine/Documentation/Source/Platforms/GearVR/GlobalMenu/1/GlobalMenu_HudCreation.INT.udn b/Engine/Documentation/Source/Platforms/GearVR/GlobalMenu/1/GlobalMenu_HudCreation.INT.udn
new file mode 100644
--- /dev/null
+++ b/Engine/Documentation/Source/Platforms/GearVR/GlobalMenu/1/GlobalMenu_HudCreation.INT.udn
@@ -0,0 +1,79 @@
+Availability:Docs
+Title:Samsung Gear VR Global Menu HUD Creation
+Crumbs:%ROOT%, Platforms/GearVR
+Description:Instructions on how to create a new HUD to use with your Gear VR project.
+Navigation:platform
+Platform:Gear VR
+
+[VAR:Steps]
+[OBJECT:Navigation]
+	[PARAM:previous]
+	[/PARAM]
+	[PARAM:current]
+	[/PARAM]
+	[PARAM:home]
+		[Quick Start Home](Platforms\GearVR\GlobalMenu)
+	[/PARAM]
+	[PARAM:next]
+		[Next Step](Platforms\GearVR\GlobalMenu\2 "%Platforms/GearVR/QuickStart/2:title%")
+	[/PARAM]
+[/OBJECT]
+[/VAR]
+
+
+%Steps%
+
+
+Before we can start to build the Global Menu HUD we first need to create a new HUD Actor so that the Blueprint code that we create has a location to be stored. 
+In the following section of this Quick Start we will give over how to create a new HUD to your Gear VR project.
+
+1. First need to create a new HUD by **Right - Clicking** in the **Content Browser** and then selecting the **Blueprint Class** option from the **Create Basic Asset** section.
+	
+	![](T_Create_New_BP.png)
+
+1. From the **Pick Parent Class** menu input **HUD** into the **Search** and the select the **HUD** class.
+
+	![](T_Find&Select_HUD.png)
+
+1. Once created give the HUD a name. For this example we will use **GearVR_GlobalMenu** as the name for this Blueprint HUD.
+
+	![](T_Named_HUD.png)
+
+
+%Steps%
+
+
+
+
+
+
+
+
+
+
+
+
+
+
+
+
+
+
+
+
+
+
+
+
+
+
+
+
+
+
+
+
+
+
+
+
diff --git a/Engine/Documentation/Source/Platforms/GearVR/GlobalMenu/2/GlobalMenu_HUDAssignment.INT.udn b/Engine/Documentation/Source/Platforms/GearVR/GlobalMenu/2/GlobalMenu_HUDAssignment.INT.udn
new file mode 100644
--- /dev/null
+++ b/Engine/Documentation/Source/Platforms/GearVR/GlobalMenu/2/GlobalMenu_HUDAssignment.INT.udn
@@ -0,0 +1,41 @@
+Availability:Docs
+Title:Samsung Gear VR Global Menu HUD Assignment
+Crumbs:%ROOT%, Platforms/GearVR
+Description: How to assign the new HUD to your project
+Navigation:platform
+Platform:Gear VR
+
+[VAR:Steps]
+[OBJECT:Navigation]
+	[PARAM:previous]
+	[Previous Step](Platforms/GearVR/GlobalMenu/1 "%Platforms/GearVR/GlobalMenu/1:title%")
+	[/PARAM]
+	[PARAM:current]
+	[/PARAM]
+	[PARAM:home]
+		[Quick Start Home](Platforms/GearVR/GlobalMenu)
+	[/PARAM]
+	[PARAM:next]
+		[Next Step](Platforms/GearVR/GlobalMenu/3 "%Platforms/GearVR/GlobalMenu/3:title%")
+	[/PARAM]
+[/OBJECT]
+[/VAR]
+
+%Steps%
+
+With the HUD now created it is now time to assign it so the project you are working with will use this HUD instead of the default one. 
+To accomplish this please do the following.
+
+1. On the Main Toolbar, click on the  **Blueprints** icon to expose the Blueprint options for this level.
+
+	![](T_BP_Options.png)
+
+1. From the Blueprints menu, move the mouse over the **Game Mode** settings in the **Project Settings** section to expose the **Game Mode Classes** section.
+
+	![](T_GameMode_Settings.png)
+
+1. In the **Game Mode Classes** move the mouse to the section that says **HUD:New** and then from the menu that is displayed, select the **Select HUD Class** option then select the **GearVR_GlobalMenu** HUD.
+
+	![](T_Assign_HUD.png)
+
+%Steps%
diff --git a/Engine/Documentation/Source/Platforms/GearVR/GlobalMenu/3/GlobalMenu_HUDBPSetup.INT.udn b/Engine/Documentation/Source/Platforms/GearVR/GlobalMenu/3/GlobalMenu_HUDBPSetup.INT.udn
new file mode 100644
--- /dev/null
+++ b/Engine/Documentation/Source/Platforms/GearVR/GlobalMenu/3/GlobalMenu_HUDBPSetup.INT.udn
@@ -0,0 +1,127 @@
+Availability:Docs
+Title:Samsung Gear VR Global Menu Blueprint Setup
+Crumbs:%ROOT%, Platforms/GearVR
+Description: How to setup a Blueprint to display the Gear VR Global Menu
+Navigation:platform
+Platform:Gear VR
+
+[VAR:Steps]
+[OBJECT:Navigation]
+	[PARAM:previous]
+	[Previous Step](Platforms/GearVR/GlobalMenu/2 "%Platforms/GearVR/GlobalMenu/2:title%")
+	[/PARAM]
+	[PARAM:current]
+	[/PARAM]
+	[PARAM:home]
+		[Quick Start Home](Platforms/GearVR/GlobalMenu)
+	[/PARAM]
+	[PARAM:next]
+	[/PARAM]
+[/OBJECT]
+[/VAR]
+
+%Steps%
+
+With the **GearVR_GlobalMenu** HUD now assigned as the game HUD, we can now start to build the Blueprint logic that will call the menu based on the requirements from Oculus.
+
+
+1. Open up the **GearVR_GlobalMenu** Blueprint by **Double - Clicking** on it inside of the **Content Browser**.
+
+1. Once the **GearVR_GlobalMenu** Blueprint is open, go the the **Event Graph** and then delete everything that is there leaving a blank graph to work with.
+
+	![](T_Empty_BP.png)
+
+1. In the **Variables** section create the following variables types with the following names.
+	
+	![](T_Variable_Names.png)
+
+	|Variable Type     |Name    |Default Value|
+	|------------------|--------|------------------| 
+	| Boolean| Invoked          | False|
+	| Float|IncrementGlobalMenuPressTimer|0.0|
+	| Float|GlobalMenuPressTimer|0.0|
+
+1. In order for the HUD to know that the back button is being pressed on the Gear VR Head Mounted Display (HMD) we have to enable the HUD to receive input signals by using the **Enable Input** node. 
+	Inside of your HUD Blueprint add and connect the nodes like in the following image. 
+
+	![](T_Enable_HUD_Input.png)
+
+	[region:note]
+	If you are having issues finding **Enable Input** try turning off **Context Sensitive** and searching for it again. 
+
+	![](T_Turn_Off_Context_Sensitive.png)
+	[/region]
+
+1. Now that the HUD can receive button inputs add and setup the following nodes so that when the **Android Back** button is pressed we can enable the Global Menu and when it's released we can reset everything for the next time it is pressed.
+
+	![](T_Key_Input.png)
+
+	[region:note]
+	To quickly test to see if your button logic is working, add a keyboard input key as an additional input to the back button. 
+	While this will not active the Global Menu while testing on PC, it will be able to help you see everything up till the menu is supposed to be called helping you see if everything is connected correctly.
+
+	![](T_Testing_Tip.png)
+	[/region]
+
+1. With the back button able to send and receive events we now need setup some logic to call the Global Menu after the back button has been held down for more than **0.75** seconds to meet the requirements set by Oculus. 
+	Setup up your HUD Blueprint to match the image below.
+
+	![](T_Timer_Logic.png)
+
+1. Input **0.75** in the **B** input of the **>= Node** then input **OVRGLOBALMENU** into the **Command** section of the **Execute Console Command** node.
+
+	![](T_Timer_Values.png)
+
+1. Finally as per the instructions from Oculus, we need to display something that will grow over time to notify the user the the Global Menu is being activated. 
+	Add and setup the nodes in your Blueprint to match the image below.
+
+	![](T_Draw_HUD.png)
+
+1. When everything has been completed your Blueprint graph should look like this.
+
+	![](T_Final_Graph.png)
+
+1. Now you can package your project up and deploy it to your Smartphone to use with the Gear VR HMD. 
+	To see the menu in action, launch your project and then hold down the back button that is on the side of Gear VR HMD.
+	As you do you will see a line being drawn across the screen and then when it finishes the Global Menu will be displayed.
+
+
+
+
+
+
+
+
+
+
+
+
+
+
+
+
+
+
+
+
+
+
+
+
+
+
+
+
+
+
+
+
+
+
+
+
+
+
+
+
+%Steps%
\ No newline at end of file
diff --git a/Engine/Documentation/Source/Platforms/GearVR/GlobalMenu/GearVRGlobalMenu.INT.udn b/Engine/Documentation/Source/Platforms/GearVR/GlobalMenu/GearVRGlobalMenu.INT.udn
new file mode 100644
--- /dev/null
+++ b/Engine/Documentation/Source/Platforms/GearVR/GlobalMenu/GearVRGlobalMenu.INT.udn
@@ -0,0 +1,53 @@
+Availability:Docs
+Title:Samsung Gear VR Global Menu
+Crumbs:%ROOT%, Platforms
+Description:How to setup your Gear VR project to access the global menu.
+Navigation:platform
+Platform:Gear VR
+
+In order for your Gear VR project to be accepted to the Oculus store there are certain requirements that your project must adhere to. 
+One of these requirements is that at any point while someone is using your project they need to be able to access the Gear VR Global Menu when they hold down the back button for more than 0.75 seconds. 
+In the following Quick Start we will cover everything you need to know so the user can enable the Gear VR Global menu in your projects.
+
+[region:note]
+Please note that the following Quick Start will work on any project that is setup for Gear VR. 
+If you have no done so already please make sure to check out the Gear [VR Quick Start](Platforms\GearVR\QuickStart) for detailed information on how to get a project up and running with Gear VR.
+[/region]
+
+[OBJECT:TopicButtonList]
+	[PARAM:icon]
+		![](%ROOT%/build_icon.png)(convert:false)
+	[/PARAM]
+	[PARAM:title]
+		Steps
+	[/PARAM]
+	[PARAM:description]
+	[/PARAM]
+	[PARAM:links]
+	* [](Platforms/GearVR/GlobalMenu/1 "%Platforms/GearVR/GlobalMenu/1:description%")
+	* [](Platforms/GearVR/GlobalMenu/2 "%Platforms/GearVR/GlobalMenu/2:description%")
+	* [](Platforms/GearVR/GlobalMenu/3 "%Platforms/GearVR/GlobalMenu/3:description%")
+	[/PARAM]
+[/OBJECT]
+
+
+
+
+
+
+
+
+
+
+
+
+
+
+
+
+
+
+
+
+
+
