#UE4doc #docs: Removed "Overview" and adjusted bold type according to Instructional Guide

diff --git a/Engine/Documentation/Source/Engine/Rendering/Materials/Functions/Reference/Misc/Misc.INT.udn b/Engine/Documentation/Source/Engine/Rendering/Materials/Functions/Reference/Misc/Misc.INT.udn
--- a/Engine/Documentation/Source/Engine/Rendering/Materials/Functions/Reference/Misc/Misc.INT.udn
+++ b/Engine/Documentation/Source/Engine/Rendering/Materials/Functions/Reference/Misc/Misc.INT.udn
@@ -5,7 +5,7 @@ Description:Miscellaneous functions that do not fall into existing categories.
 
 [TOC (start:2 end:3)]
 
-## Overview
+
 
 Miscellaneous functions include a variety of one-off functions that do not readily fit into any of the other existing categories. They are placed here to avoid too many categories with only one function inside.
 
diff --git a/Engine/Documentation/Source/Engine/Rendering/Materials/Functions/Reference/Opacity/Opacity.INT.udn b/Engine/Documentation/Source/Engine/Rendering/Materials/Functions/Reference/Opacity/Opacity.INT.udn
--- a/Engine/Documentation/Source/Engine/Rendering/Materials/Functions/Reference/Opacity/Opacity.INT.udn
+++ b/Engine/Documentation/Source/Engine/Rendering/Materials/Functions/Reference/Opacity/Opacity.INT.udn
@@ -5,7 +5,7 @@ Description:Functions to handle opacity values within a material network.
 
 [TOC (start:2 end:3)]
 
-## Overview
+
 
 Opacity functions exist to speed up handling of complex opacity calculations.
 
diff --git a/Engine/Documentation/Source/Engine/Rendering/Materials/Functions/Reference/PivotPainter/PivotPainter.INT.udn b/Engine/Documentation/Source/Engine/Rendering/Materials/Functions/Reference/PivotPainter/PivotPainter.INT.udn
--- a/Engine/Documentation/Source/Engine/Rendering/Materials/Functions/Reference/PivotPainter/PivotPainter.INT.udn
+++ b/Engine/Documentation/Source/Engine/Rendering/Materials/Functions/Reference/PivotPainter/PivotPainter.INT.udn
@@ -5,7 +5,7 @@ Description:Functions designed to allow utilization of the Pivot Painter MAXScri
 
 [TOC (start:2 end:3)]
 
-## Overview
+
 
 [REGION:header_img]
 [![splash_demo.jpg](Engine/Content/Tools/PivotPainter/splash_demo.jpg)](Engine/Content/Tools/PivotPainter/PivotPainterFoliageExamples.mp4)
diff --git a/Engine/Documentation/Source/Engine/Rendering/Materials/Functions/Reference/Procedurals/Procedurals.INT.udn b/Engine/Documentation/Source/Engine/Rendering/Materials/Functions/Reference/Procedurals/Procedurals.INT.udn
--- a/Engine/Documentation/Source/Engine/Rendering/Materials/Functions/Reference/Procedurals/Procedurals.INT.udn
+++ b/Engine/Documentation/Source/Engine/Rendering/Materials/Functions/Reference/Procedurals/Procedurals.INT.udn
@@ -5,7 +5,7 @@ Description:Procedurally generated textures and operations, such as the creation
 
 [TOC (start:2 end:3)]
 
-## Overview
+
 
 The Procedural functions provide a quick way to make simple procedurally generated textures and masks. This saves memory over having to use imported textures.
 
diff --git a/Engine/Documentation/Source/Engine/Rendering/Materials/Functions/Reference/Reflections/Reflections.INT.udn b/Engine/Documentation/Source/Engine/Rendering/Materials/Functions/Reference/Reflections/Reflections.INT.udn
--- a/Engine/Documentation/Source/Engine/Rendering/Materials/Functions/Reference/Reflections/Reflections.INT.udn
+++ b/Engine/Documentation/Source/Engine/Rendering/Materials/Functions/Reference/Reflections/Reflections.INT.udn
@@ -5,7 +5,7 @@ Description:Functions for aiding in the calculation of values for a variety of r
 
 [TOC (start:2 end:3)]
 
-## Overview
+
 
 The Reflections functions give you a variety of ways to speed up the handling of reflection calculations in your materials.
 
diff --git a/Engine/Documentation/Source/Engine/Rendering/Materials/Functions/Reference/Shading/Shading.INT.udn b/Engine/Documentation/Source/Engine/Rendering/Materials/Functions/Reference/Shading/Shading.INT.udn
--- a/Engine/Documentation/Source/Engine/Rendering/Materials/Functions/Reference/Shading/Shading.INT.udn
+++ b/Engine/Documentation/Source/Engine/Rendering/Materials/Functions/Reference/Shading/Shading.INT.udn
@@ -5,7 +5,7 @@ Description:Functions for handling special shading types, such as Fuzzy Shading
 
 [TOC (start:2 end:3)]
 
-## Overview
+
 
 The Shading functions provide for specialized shading operations, such as fuzzy shading and adjusting the shape of a specular highlight.
 
diff --git a/Engine/Documentation/Source/Engine/Rendering/Materials/Functions/Reference/WorldPositionOffset/WorldPositionOffset.INT.udn b/Engine/Documentation/Source/Engine/Rendering/Materials/Functions/Reference/WorldPositionOffset/WorldPositionOffset.INT.udn
--- a/Engine/Documentation/Source/Engine/Rendering/Materials/Functions/Reference/WorldPositionOffset/WorldPositionOffset.INT.udn
+++ b/Engine/Documentation/Source/Engine/Rendering/Materials/Functions/Reference/WorldPositionOffset/WorldPositionOffset.INT.udn
@@ -5,7 +5,7 @@ Description:Functions for handling vertex manipulation using world position offs
 
 [TOC (start:2 end:3)]
 
-## Overview
+
 
 The WorldPositionOffset category contains special functions for applying manipulating vertices of a mesh via world position offset. These functions can be chained together for additive effects. 
 
diff --git a/Engine/Documentation/Source/Engine/Rendering/Materials/HowTo/BumpOffset/BumpOffset.INT.udn b/Engine/Documentation/Source/Engine/Rendering/Materials/HowTo/BumpOffset/BumpOffset.INT.udn
--- a/Engine/Documentation/Source/Engine/Rendering/Materials/HowTo/BumpOffset/BumpOffset.INT.udn
+++ b/Engine/Documentation/Source/Engine/Rendering/Materials/HowTo/BumpOffset/BumpOffset.INT.udn
@@ -8,7 +8,7 @@ Related: Engine/Rendering/Materials
 
 [TOC(start:2)]
 
-## Overview
+
 
 **Bump Offset** mapping give a surface the illusion of depth by modifying the UV coordinates in a creative way to help further displace the texels from the surface of the object, giving the illusion that the surface has more details than it really does. 
 In this How To example, we will be covering not only how you can find the Bump Offset Material Expression, we will also be covering how you can utilize the Bump Offset node in your Materials.
diff --git a/Engine/Documentation/Source/Engine/Rendering/Materials/HowTo/MaterialHowTo.INT.udn b/Engine/Documentation/Source/Engine/Rendering/Materials/HowTo/MaterialHowTo.INT.udn
--- a/Engine/Documentation/Source/Engine/Rendering/Materials/HowTo/MaterialHowTo.INT.udn
+++ b/Engine/Documentation/Source/Engine/Rendering/Materials/HowTo/MaterialHowTo.INT.udn
@@ -28,7 +28,7 @@ Related: Engine/Rendering/Materials/MaterialProperties
 [/OBJECT]
 [/VAR]
 
-## Overview
+
 
 The following How-To tutorials will demonstrate how you can utilize the full potential of the Material Editor in Unreal Engine 4 (UE4).
 From tips on How - To correctly set up Human Skin to How - To use Layered Materials, the How - To tutorials are a must read for anyone who wants to know more about UE4's Material Editor.
