#UE4doc #docs: Removed "Overview" from these pages

diff --git a/Engine/Documentation/Source/Engine/Rendering/LightingAndShadows/DistanceFieldAmbientOcclusion/DistanceFieldAmbientOcclusion.INT.udn b/Engine/Documentation/Source/Engine/Rendering/LightingAndShadows/DistanceFieldAmbientOcclusion/DistanceFieldAmbientOcclusion.INT.udn
--- a/Engine/Documentation/Source/Engine/Rendering/LightingAndShadows/DistanceFieldAmbientOcclusion/DistanceFieldAmbientOcclusion.INT.udn
+++ b/Engine/Documentation/Source/Engine/Rendering/LightingAndShadows/DistanceFieldAmbientOcclusion/DistanceFieldAmbientOcclusion.INT.udn
@@ -48,7 +48,7 @@ Description:Dynamic Ambient Occlusion for Movable SkyLights.
 
 [TOC(start:2 end:2)]
 
-## Overview
+
 
 Shadowing for Movable Sky Lights is provided by a new feature which generates medium scale Ambient Occlusion from Signed Distance Field Volumes precomputed around each rigid object.  It supports dynamic scene changes - the rigid meshes can be moved or hidden and it will affect the occlusion.  Unlike Screen Space Ambient Occlusion, occlusion is computed from world space occluders so there are no artifacts from missing data off-screen.
 
diff --git a/Engine/Documentation/Source/Engine/Rendering/LightingAndShadows/IESLightProfiles/IESLightProfiles.INT.udn b/Engine/Documentation/Source/Engine/Rendering/LightingAndShadows/IESLightProfiles/IESLightProfiles.INT.udn
--- a/Engine/Documentation/Source/Engine/Rendering/LightingAndShadows/IESLightProfiles/IESLightProfiles.INT.udn
+++ b/Engine/Documentation/Source/Engine/Rendering/LightingAndShadows/IESLightProfiles/IESLightProfiles.INT.udn
@@ -49,7 +49,7 @@ Description:Using IES Light Profiles in Unreal Engine 4.
 
 [TOC(start:2)]
 
-## Overview
+
 
 [EXCERPT:IES001]
 **IES Light Profiles** are a curve that defines light intensity in an arc that Unreal Engine 4 will "sweep" around an axis to make **Point Lights** and **Spot Lights** appear to be casting light more realistically, as if accounting for the reflective surfaces in the light fixture, the shape of the light bulb, and any lensing effects. The curve itself can be seen as a multiplier for the overall brightness of the light, as if you were projecting a texture from the light, but with out the texture overhead to do so or the errors that can occur when projecting at some angles.
diff --git a/Engine/Documentation/Source/Engine/Rendering/LightingAndShadows/LightPropagationVolumes/LightPropagationVolumes.INT.udn b/Engine/Documentation/Source/Engine/Rendering/LightingAndShadows/LightPropagationVolumes/LightPropagationVolumes.INT.udn
--- a/Engine/Documentation/Source/Engine/Rendering/LightingAndShadows/LightPropagationVolumes/LightPropagationVolumes.INT.udn
+++ b/Engine/Documentation/Source/Engine/Rendering/LightingAndShadows/LightPropagationVolumes/LightPropagationVolumes.INT.udn
@@ -49,7 +49,7 @@ Light Propagation Volumes are a feature in development and not ready for product
 
 [TOC]
 
-## Overview
+
 
 [REGION:fullwidth]
 ![](lpv_scene.png)
diff --git a/Engine/Documentation/Source/Engine/Rendering/LightingAndShadows/LightTypes/SkyLight/SkyLight.INT.udn b/Engine/Documentation/Source/Engine/Rendering/LightingAndShadows/LightTypes/SkyLight/SkyLight.INT.udn
--- a/Engine/Documentation/Source/Engine/Rendering/LightingAndShadows/LightTypes/SkyLight/SkyLight.INT.udn
+++ b/Engine/Documentation/Source/Engine/Rendering/LightingAndShadows/LightTypes/SkyLight/SkyLight.INT.udn
@@ -48,7 +48,7 @@ Description:The basics of Sky Lights.
 
 [TOC(start:2 end:2)]
 
-## Overview
+
 [EXCERPT:SL001]
 The **Sky Light** captures the distant parts of your level (everything further than `SkyDistanceThreshold`) and applies that to the scene as a light.  That means the sky's appearance and its lighting / reflections will match, even if your sky is coming from atmosphere, or layered clouds on top of a skybox, or distant mountains.  You can also manually specify a cubemap to use.
 
diff --git a/Engine/Documentation/Source/Engine/Rendering/LightingAndShadows/LightsAndPostProcess/LightingAndPostProcess.INT.udn b/Engine/Documentation/Source/Engine/Rendering/LightingAndShadows/LightsAndPostProcess/LightingAndPostProcess.INT.udn
--- a/Engine/Documentation/Source/Engine/Rendering/LightingAndShadows/LightsAndPostProcess/LightingAndPostProcess.INT.udn
+++ b/Engine/Documentation/Source/Engine/Rendering/LightingAndShadows/LightsAndPostProcess/LightingAndPostProcess.INT.udn
@@ -48,7 +48,7 @@ Description:How lights interact with the Post Process Chain.
 
 [TOC(start:2 end:2)]
 
-## Overview
+
 
 The post process pipeline takes the lighting in the scene and uses it in various ways, from simulating an organic eye adapting to changes in light or a camera's numerous refracting lenses. While the lighting of a scene may be the first step, the post process steps really drive many effects.
 
diff --git a/Engine/Documentation/Source/Engine/Rendering/LightingAndShadows/RayTracedDistanceFieldShadowing/RayTracedDistanceFieldShadowing.INT.udn b/Engine/Documentation/Source/Engine/Rendering/LightingAndShadows/RayTracedDistanceFieldShadowing/RayTracedDistanceFieldShadowing.INT.udn
--- a/Engine/Documentation/Source/Engine/Rendering/LightingAndShadows/RayTracedDistanceFieldShadowing/RayTracedDistanceFieldShadowing.INT.udn
+++ b/Engine/Documentation/Source/Engine/Rendering/LightingAndShadows/RayTracedDistanceFieldShadowing/RayTracedDistanceFieldShadowing.INT.udn
@@ -48,7 +48,7 @@ Description:Dynamic area shadows
 
 [TOC(start:2 end:2)]
 
-## Overview
+
 
 Ray Traced Distance Field shadows leverage the properties of a distance field representation of the scene to compute efficient area shadowing from dynamic meshes.  This uses the same data as [Distance Field Ambient Occlusion](Engine/Rendering/LightingAndShadows/DistanceFieldAmbientOcclusion), and therefore has many of the same limitations.
 
diff --git a/Engine/Documentation/Source/Engine/Rendering/LightingAndShadows/Shadows/Shadows.INT.udn b/Engine/Documentation/Source/Engine/Rendering/LightingAndShadows/Shadows/Shadows.INT.udn
--- a/Engine/Documentation/Source/Engine/Rendering/LightingAndShadows/Shadows/Shadows.INT.udn
+++ b/Engine/Documentation/Source/Engine/Rendering/LightingAndShadows/Shadows/Shadows.INT.udn
@@ -49,7 +49,7 @@ Related: Engine/Rendering/LightingAndShadows/RayTracedDistanceFieldShadowing
 
 [TOC(start:2 end:2)]
 
-## Overview
+
 
 Shadows make objects feel grounded in the world and give the viewer a sense of depth and space. Static shadows are free as far as rendering goes, but dynamic shadows can be one of the biggest drains on performance. This document covers the basics of the 4 types of shadow casting you can have in Unreal Engine 4.
 
