#UE4 Docs: Updating this with the new page for the Global Menu.

diff --git a/Engine/Documentation/Source/Platforms/GearVR/GearVR.INT.udn b/Engine/Documentation/Source/Platforms/GearVR/GearVR.INT.udn
--- a/Engine/Documentation/Source/Platforms/GearVR/GearVR.INT.udn
+++ b/Engine/Documentation/Source/Platforms/GearVR/GearVR.INT.udn
@@ -16,6 +16,7 @@ Platform:Gear VR
 	[PARAM:links]
 		* [](Platforms/GearVR/Prerequisites "%Platforms/GearVR/Prerequisites:description%")
 		* [](Platforms/GearVR/QuickStart "%Platforms/GearVR/QuickStart:description%")
+		* [](Platforms/GearVR/GlobalMenu "%Platforms/GearVR/GlobalMenu:description%")
 		* [](Platforms/GearVR/BestPractices "%Platforms/GearVR/BestPractices:description%")
 		* [](Platforms/GearVR/Debugging "%Platforms/GearVR/BestPractices:description%")
 	[/PARAM]
