#UE4doc #docs: Removed Overview header
Changed multiple green borders in images to yellow
Bolded multiple editor terms
Removed a couple redundant hyperlinks

diff --git a/Engine/Documentation/Source/Engine/Actors/SkeletalMeshActors/SkeletalMeshActors.INT.udn b/Engine/Documentation/Source/Engine/Actors/SkeletalMeshActors/SkeletalMeshActors.INT.udn
--- a/Engine/Documentation/Source/Engine/Actors/SkeletalMeshActors/SkeletalMeshActors.INT.udn
+++ b/Engine/Documentation/Source/Engine/Actors/SkeletalMeshActors/SkeletalMeshActors.INT.udn
@@ -11,8 +11,8 @@ Related: Engine/Animation
 
 ![](SkeletalMeshActor_Matinee_Example.png)
 
-## Overview 
-When playing a video game that uses Unreal Engine 4 (UE4), the character that you are controlling on screen is often made up of a special kind of mesh called a Skeletal Mesh Actor. Skeletal Mesh Actors are special because they can be used to display complex animation data that was made inside of an external 3D application and then imported into UE4. The following guide aims to show you the ins and outs of how to use Skeletal Mesh actors inside UE4.
+
+When playing a video game that uses Unreal Engine 4 (UE4), the character that you are controlling on screen is often made up of a special kind of mesh called a **Skeletal Mesh Actor**. Skeletal Mesh Actors are special because they can be used to display complex animation data that was made inside of an external 3D application and then imported into UE4. The following guide aims to show you the ins and outs of how to use Skeletal Mesh actors inside UE4.
 
 [region:tip]
 All content for this tutorial was taken from the **Matinee** demo which you can download for free from the Learn Tab. 
@@ -26,7 +26,7 @@ Skeletal Mesh Actors are placed in maps using standard Actor placement methods;
 
 **Drag & Drop**
 
-You can drag & drop Skeletal Mesh Actor�s right from the **Content Browser** right into the level you are working with by following these few simple steps. 
+You can drag & drop Skeletal Mesh Actors right from the **Content Browser** right into the level you are working with by following these few simple steps.
 1. In the **Content Browser**, locate the Skeletal Mesh you want to add to the map as a Skeletal Mesh Actor.
 
 1. **Left-click** on the Skeletal Mesh in the **Content Browser** and drag the mouse (while holding the **Left Mouse Button** down) from the **Content Browser** to the location in the viewport where you want to place the mesh. The location does not need to be exact. You can always reposition and rotate the mesh afterwards.   
@@ -48,7 +48,7 @@ You can also place Skeletal Mesh Actors in the world using the **Right-click** c
 ![](Skeletal_Mesh_Context_Menu.gif.gif)
 
 ## Playing Animation on a Skeletal Mesh Actor 
-Getting a Skeletal Mesh Actor so that it will animate while the game is running can be done two different ways inside of UE4. You can use an [Animation Blueprint]( Engine/Animation/AnimBlueprints), which will allow you to play and blend multiple animations together. Or you can use and Animation Asset to play a single [Animation Sequence]( Engine/Animation/Sequences). In this tutorial, we will be focusing on using an [Animation Sequence]( Engine/Animation/Sequences) to play our animation while the game is running. 
+Getting a Skeletal Mesh Actor so that it will animate while the game is running can be done two different ways inside of UE4. You can use an [Animation Blueprint]( Engine/Animation/AnimBlueprints), which will allow you to play and blend multiple animations together. Or you can use and Animation Asset to play a single [Animation Sequence]( Engine/Animation/Sequences). In this tutorial, we will be focusing on using an Animation Sequence to play our animation while the game is running. 
 
 ### Skeletal Mesh Actor Setup
 The setup for getting a Skeletal Mesh Actor ready to play an animation is actually quite simple and can be done in the following steps.
@@ -86,7 +86,7 @@ Materials that are to be used with Skeletal Mesh Actors need to have the materia
 [/region]
 
 ## Collision
-Normal collision creation and detection does not work with Skeletal Mesh Actors. In order for your Skeletal Meshes to collide with objects in the level, your Skeletal Mesh Actors need to have a [Physics Asset]( Engine/Physics/PhAT/UserGuide) that is specially created for your object. 
+Normal collision creation and detection does not work with Skeletal Mesh Actors. In order for your Skeletal Meshes to collide with objects in the level, your Skeletal Mesh Actors need to have a Physics Asset that is specially created for your object. 
 
 ### Creating Collision for a Skeletal Mesh Actor 
 
