#UE4doc #docs Minor grammatical edit.

diff --git a/Engine/Documentation/Source/Programming/QuickStart/3/QuickStart_3.INT.udn b/Engine/Documentation/Source/Programming/QuickStart/3/QuickStart_3.INT.udn
--- a/Engine/Documentation/Source/Programming/QuickStart/3/QuickStart_3.INT.udn
+++ b/Engine/Documentation/Source/Programming/QuickStart/3/QuickStart_3.INT.udn
@@ -83,7 +83,7 @@ Version: 4.8
 		When building in Xcode, you are only compiling the game project, not the Editor.
 		[/REGION]
 		
-		*	Now that we're done coding, we can compile our project by using XCode's Debug Configuration by clicking the **Product > Build** option.
+		*	Now that we're done coding, we can compile our project using XCode's Debug Configuration by clicking the **Product > Build** option.
 		
 			![](ProjectBuild_mac.png)
 			
