#UE4doc #docs: Didn't commit changes the first time

diff --git a/Engine/Documentation/Source/Engine/Actors/StaticMeshActor/UserGuide.INT.udn b/Engine/Documentation/Source/Engine/Actors/StaticMeshActor/UserGuide.INT.udn
--- a/Engine/Documentation/Source/Engine/Actors/StaticMeshActor/UserGuide.INT.udn
+++ b/Engine/Documentation/Source/Engine/Actors/StaticMeshActor/UserGuide.INT.udn
@@ -26,10 +26,9 @@ Description:Guide to placing and working with StaticMeshActors for creating worl
 
 [TOC(start:2)]
 
-## Overview
 
-StaticMeshs are one of the foundational types of renderable geometry in Unreal Engine. In order to use these meshes 
-to populate your worlds, the StaticMeshActor is used. StaticMeshes dragged into the level from the **Content Browser** are 
+**Static Meshes** are one of the foundational types of renderable geometry in Unreal Engine. In order to use these meshes 
+to populate your worlds, the **StaticMeshActor** is used. StaticMeshes dragged into the level from the **Content Browser** are 
 automatically converted to StaticMeshActors.
 
 [REGION:note]
@@ -54,7 +53,7 @@ menu of the viewports or drag and drop from the [Content Browser](Engine/Content
 
 1. In the **Content Browser**, locate the StaticMesh you want to add to the map as a StaticMeshActor.
 
-1. **Left-Click** on the StaticMesh in the **Content Browser** and drag the mouse (while holding the **Left Mouse Button** down) from the **Content Browser** to the location in the viewport where you want to place the mesh. The location does not need to be exact. You can always reposition, rotate, scale the mesh afterwards.   
+1. **Left-Click** on the StaticMesh in the **Content Browser** and drag the mouse (while holding the **Left Mouse Button** down) from the **Content Browser** to the location in the **Viewport** where you want to place the mesh. The location does not need to be exact. You can always reposition, rotate, scale the mesh afterwards.   
 
 1. Release the **Left Mouse Button** to place the mesh in the map as a StaticMeshActor, as seen in the Property Window.   
 
@@ -75,7 +74,7 @@ menu of the viewports or drag and drop from the [Content Browser](Engine/Content
 
 ## Making a StaticMeshActor Movable
 
-In order to move, rotate, or scale a StaticMeshActor during play, you must first make it movable. 
+In order to move, rotate, or scale a StaticMeshActor during play, you must first make it movable.
 
 
 ## Enabling a StaticMeshActor to Simulate Physics
@@ -89,14 +88,14 @@ The Material(s) applied to a StaticMesh can be overridden on a per-Actor basis.
 asset to be used multiple times in a single map or throughout your game while giving each a unique appearance. 
 The **Materials** property located in the **Rendering** category of the _StaticMeshComponent_ of the StaticMeshActor 
 is an array of Materials that directly maps to those applied to the StaticMesh asset through the [Static Mesh Editor](Engine/Content/Types/StaticMeshes/Editor). 
-Materials can be assigned to the array manually or they can be applied directly to the mesh in the viewport by 
+Materials can be assigned to the array manually or they can be applied directly to the mesh in the **Viewport** by 
 dragging and dropping from the **Content Browser**.
 
 **Manual Assignment**
 
 ![](SMA_MaterialSingle.png)
 
-1. In the viewport, select the StaticMeshActor you want to assign.
+1. In the **Viewport**, select the StaticMeshActor you want to assign.
 
 1. In the **Details** Pane, under the **Materials** category, you will see all the materials assigned to the StaticMeshActor.
 
