#UE4doc #docs: Minor edits made to Global Menu pages

diff --git a/Engine/Documentation/Source/Platforms/GearVR/GlobalMenu/1/GlobalMenu_HudCreation.INT.udn b/Engine/Documentation/Source/Platforms/GearVR/GlobalMenu/1/GlobalMenu_HudCreation.INT.udn
--- a/Engine/Documentation/Source/Platforms/GearVR/GlobalMenu/1/GlobalMenu_HudCreation.INT.udn
+++ b/Engine/Documentation/Source/Platforms/GearVR/GlobalMenu/1/GlobalMenu_HudCreation.INT.udn
@@ -25,17 +25,17 @@ Platform:Gear VR
 
 
 Before we can start to build the Global Menu HUD we first need to create a new HUD Actor so that the Blueprint code that we create has a location to be stored. 
-In the following section of this Quick Start we will give over how to create a new HUD to your Gear VR project.
+In the following section of this Quick Start we will go over how to create a new HUD to your Gear VR project.
 
 1. First need to create a new HUD by **Right - Clicking** in the **Content Browser** and then selecting the **Blueprint Class** option from the **Create Basic Asset** section.
 	
 	![](T_Create_New_BP.png)
 
-1. From the **Pick Parent Class** menu input **HUD** into the **Search** and the select the **HUD** class.
+1. From the **Pick Parent Class** menu, input **HUD** into the **Search** and the select the **HUD** class.
 
 	![](T_Find&Select_HUD.png)
 
-1. Once created give the HUD a name. For this example we will use **GearVR_GlobalMenu** as the name for this Blueprint HUD.
+1. Once created, give the HUD a name. For this example we will use **GearVR_GlobalMenu** as the name for this Blueprint HUD.
 
 	![](T_Named_HUD.png)
 
diff --git a/Engine/Documentation/Source/Platforms/GearVR/GlobalMenu/2/GlobalMenu_HUDAssignment.INT.udn b/Engine/Documentation/Source/Platforms/GearVR/GlobalMenu/2/GlobalMenu_HUDAssignment.INT.udn
--- a/Engine/Documentation/Source/Platforms/GearVR/GlobalMenu/2/GlobalMenu_HUDAssignment.INT.udn
+++ b/Engine/Documentation/Source/Platforms/GearVR/GlobalMenu/2/GlobalMenu_HUDAssignment.INT.udn
@@ -23,8 +23,8 @@ Platform:Gear VR
 
 %Steps%
 
-With the HUD now created it is now time to assign it so the project you are working with will use this HUD instead of the default one. 
-To accomplish this please do the following.
+With the HUD now created, it is now time to assign it so the project you are working with will use this HUD instead of the default one. 
+To accomplish this please do the following:
 
 1. On the Main Toolbar, click on the  **Blueprints** icon to expose the Blueprint options for this level.
 
@@ -34,7 +34,7 @@ To accomplish this please do the following.
 
 	![](T_GameMode_Settings.png)
 
-1. In the **Game Mode Classes** move the mouse to the section that says **HUD:New** and then from the menu that is displayed, select the **Select HUD Class** option then select the **GearVR_GlobalMenu** HUD.
+1. In the **Game Mode Classes** move the mouse to the section that says **HUD:New** and then from the menu that is displayed, select the **Select HUD Class** option, then select the **GearVR_GlobalMenu** HUD.
 
 	![](T_Assign_HUD.png)
 
diff --git a/Engine/Documentation/Source/Platforms/GearVR/GlobalMenu/3/GlobalMenu_HUDBPSetup.INT.udn b/Engine/Documentation/Source/Platforms/GearVR/GlobalMenu/3/GlobalMenu_HUDBPSetup.INT.udn
--- a/Engine/Documentation/Source/Platforms/GearVR/GlobalMenu/3/GlobalMenu_HUDBPSetup.INT.udn
+++ b/Engine/Documentation/Source/Platforms/GearVR/GlobalMenu/3/GlobalMenu_HUDBPSetup.INT.udn
@@ -27,7 +27,7 @@ With the **GearVR_GlobalMenu** HUD now assigned as the game HUD, we can now star
 
 1. Open up the **GearVR_GlobalMenu** Blueprint by **Double - Clicking** on it inside of the **Content Browser**.
 
-1. Once the **GearVR_GlobalMenu** Blueprint is open, go the the **Event Graph** and then delete everything that is there leaving a blank graph to work with.
+1. Once the **GearVR_GlobalMenu** Blueprint is open, go the **Event Graph** and then delete everything that is there leaving a blank graph to work with.
 
 	![](T_Empty_BP.png)
 
@@ -52,7 +52,7 @@ With the **GearVR_GlobalMenu** HUD now assigned as the game HUD, we can now star
 	![](T_Turn_Off_Context_Sensitive.png)
 	[/region]
 
-1. Now that the HUD can receive button inputs add and setup the following nodes so that when the **Android Back** button is pressed we can enable the Global Menu and when it's released we can reset everything for the next time it is pressed.
+1. Now that the HUD can receive button inputs, add and setup the following nodes so that when the **Android Back** button is pressed we can enable the Global Menu. Then, when it's released, we can reset everything for the next time it is pressed.
 
 	![](T_Key_Input.png)
 
@@ -64,7 +64,7 @@ With the **GearVR_GlobalMenu** HUD now assigned as the game HUD, we can now star
 	[/region]
 
 1. With the back button able to send and receive events we now need setup some logic to call the Global Menu after the back button has been held down for more than **0.75** seconds to meet the requirements set by Oculus. 
-	Setup up your HUD Blueprint to match the image below.
+	Setup up your HUD Blueprint to match the image below:
 
 	![](T_Timer_Logic.png)
 
@@ -72,12 +72,12 @@ With the **GearVR_GlobalMenu** HUD now assigned as the game HUD, we can now star
 
 	![](T_Timer_Values.png)
 
-1. Finally as per the instructions from Oculus, we need to display something that will grow over time to notify the user the the Global Menu is being activated. 
-	Add and setup the nodes in your Blueprint to match the image below.
+1. Finally as per the instructions from Oculus, we need to display something that will grow over time to notify the user the Global Menu is being activated. 
+	Add and setup the nodes in your Blueprint to match the image below:
 
 	![](T_Draw_HUD.png)
 
-1. When everything has been completed your Blueprint graph should look like this.
+1. When everything has been completed your Blueprint graph should look like this:
 
 	![](T_Final_Graph.png)
 
