#UE4doc #docs: Didn't commit changes the first time

diff --git a/Engine/Documentation/Source/Engine/Actors/Brushes/Brushes.INT.udn b/Engine/Documentation/Source/Engine/Actors/Brushes/Brushes.INT.udn
--- a/Engine/Documentation/Source/Engine/Actors/Brushes/Brushes.INT.udn
+++ b/Engine/Documentation/Source/Engine/Actors/Brushes/Brushes.INT.udn
@@ -70,11 +70,9 @@ Related: Engine/Actors/Volumes
 
 
 
-## Overview
-
 ![](BSPHouse.png)
 
-This document breaks down the use of Geometry Brushes in level creation.
+This document breaks down the use of **Geometry Brushes** in level creation.
 
 Geometry Brushes are the most basic tool for level construction in Unreal. Conceptually, it is best to think of a Geometry Brush as filling in and carving out volumes of space in your level. Long ago, Geometry Brushes were used as the primary building block in level design. Now, however, that role has been passed on to Static Meshes, which are far more efficient. However, Geometry Brushes can still be useful in the early stages of a product for rapid prototyping of levels and objects, as well as for level construction by those who do not have access to 3D modeling tools. This document goes over the use of Geometry Brushes and how they can utilized in your levels.
 
@@ -128,11 +126,11 @@ Brushes are created using the **Geometry** tab in the **Place Mode** panel:
 
 	![](add_subtract.png)
 
-1. Drag and drop one of the primitive types from the list into a Viewport:
+1. Drag and drop one of the primitive types from the list into a **Viewport**:
 
 	![](BoxBrush.png)(w:640)
 
-1. Modify the **Brush Settings**, use the transform tools, or activate Geometry Editing Mode to reshape and resize the Brush. See [Modifying Brushes](#modifyingbrushes) for more details.
+1. Modify the **Brush Settings**, use the transform tools, or activate **Geometry Editing Mode** to reshape and resize the Brush. See [Modifying Brushes](#modifyingbrushes) for more details.
 
 [/EXCERPT:CreatingBSP]
 
@@ -254,9 +252,9 @@ You can modify Brushes using several methods, each of which is suited to particu
 
 ### Geometry Editing Mode
 
-The best way to change the actual shape of a Brush is to use Geometry Editing Mode. This editor mode allows the direct manipulation of the vertices, edges, and faces of the Brush. It is very similar to working in a very simplified 3D modeling application.
+The best way to change the actual shape of a Brush is to use **Geometry Editing Mode**. This editor mode allows the direct manipulation of the vertices, edges, and faces of the Brush. It is very similar to working in a very simplified 3D modeling application.
 
-For more information about Geometry Editing Mode and how to use it to modify Brushes, see the [Level Editor Modes](Engine/UI/LevelEditor/Modes) page.
+For more information about **Geometry Editing Mode** and how to use it to modify Brushes, see the [Level Editor Modes](Engine/UI/LevelEditor/Modes) page.
 
 
 ### Transform Widgets
@@ -269,13 +267,13 @@ For more information on the Transform widgets and how to use them, please see [T
 
 ## Brush Properties
 
-Existing Brushes can be edited by selecting the Brush and then looking in the **Details** panel. If you have the entire Brush selected, you will see the Brush Settings category:
+Existing Brushes can be edited by selecting the Brush and then looking in the **Details** panel. If you have the entire Brush selected, you will see the **Brush Settings** category:
 
 ![](BrushDetails.png)
 
 ### Brush Types
 
-The Brush Type dropdown includes the following:
+The **Brush Type** dropdown includes the following:
 
 | Brush Type Dropdown ||
 |---|---|
@@ -287,13 +285,13 @@ There is a variety of different Brush types you may use throughout the level cre
 #### Additive
 
 [EXCERPT:Additive]
-Additive Brushes are like solid, filled-in space. This is the type you will use for any Brush geometry you wish to add to the level. A good way to visualize an additive Brush is to imagine the four walls, the floor, and the ceiling of a room. Each of these would be a separate box-like additive brush in your map with their corners matching up to form an exclosed space.
+**Additive** Brushes are like solid, filled-in space. This is the type you will use for any Brush geometry you wish to add to the level. A good way to visualize an additive Brush is to imagine the four walls, the floor, and the ceiling of a room. Each of these would be a separate box-like additive brush in your map with their corners matching up to form an exclosed space.
 [/EXCERPT:Additive]
 
 #### Subtractive
 
 [EXCERPT:Subtractive]
-A subtractive Brush is a hollow, carved-out space. This is the type of Brush you would use to remove solid space, such as to create doors, windows, etc, from previously created additive Brushes. Subtractive space is the only area that players can freely move around in.
+A **Subtractive** Brush is a hollow, carved-out space. This is the type of Brush you would use to remove solid space, such as to create doors, windows, etc, from previously created additive Brushes. Subtractive space is the only area that players can freely move around in.
 [/EXCERPT:Subtractive]
 
 ### Advanced Properties
@@ -304,7 +302,7 @@ Clicking the ![Advanced](button_advanced.png) button at the bottom of the Brush
 
 #### Polygons
 
-The Polygons dropdown offers the following options:
+The **Polygons** dropdown offers the following options:
 
 ![](PolygonsDropdown.png)
 
@@ -315,7 +313,7 @@ The Polygons dropdown offers the following options:
 
 #### Solidity
 
-The Solidity dropdown includes the following:
+The **Solidity** dropdown includes the following:
 
 Be sure to read the [Brush Solidity section](#BrushSolidity) for more details.
 
@@ -329,7 +327,7 @@ Be sure to read the [Brush Solidity section](#BrushSolidity) for more details.
 
 #### Order
 
-The Order dropdown includes the following:
+The **Order** dropdown includes the following:
 
 Be sure to read the [Brush Order](#BrushOrder) section for more details:
 
@@ -343,7 +341,7 @@ Be sure to read the [Brush Order](#BrushOrder) section for more details:
 
 ### Align and Static Mesh Buttons
 
-If you expand the properties under the Brush Category, the following buttons appear:
+If you expand the properties under the **Brush Settings** Category, the following buttons appear:
 
 | Brush Type Dropdown ||
 |---|---|
@@ -371,7 +369,7 @@ If you select a Brush surface (use **Ctrl + Shift + Left Click** to select the s
 
 ### Geometry Category
 
-The Geometry category contains a few tools for helping you manage Material application across Brush surfaces.
+The **Geometry** category contains a few tools for helping you manage Material application across Brush surfaces.
 
 |Geometry Category Buttons||
 |---|---|
@@ -384,7 +382,7 @@ The Geometry category contains a few tools for helping you manage Material appli
 ### Surface Properties
 
 [EXCERPT:SurfaceProperties]
-The Surface Properties area contains a variety of tools to help you control how textures are placed across your surfaces, as well as lightmap resolution.
+The **Surface Properties** area contains a variety of tools to help you control how textures are placed across your surfaces, as well as lightmap resolution.
 
 |Surface Property Categories||
 |---|---|
@@ -396,7 +394,7 @@ The Surface Properties area contains a variety of tools to help you control how
 
 #### Lighting
 
-The Lighting section allows you to change a variety of important light-centric properties for your Brush surfaces.
+The **Lighting** section allows you to change a variety of important light-centric properties for your Brush surfaces.
 
 |Lighting Properties||
 |---|---|
@@ -412,7 +410,7 @@ The Lighting section allows you to change a variety of important light-centric p
 ## Brush Solidity
 
 [EXCERPT:Solidity]
-Brushes can be either solid, semi-solid, or non-solid. The solidity of a Brush refers to whether it collides with objects in the world and whether the Brush causes BSP cuts to be created in the surrounding Brushes when building geometry. 
+Brushes can be either **Solid**, **Semi-solid**, or **Non-solid**. The solidity of a Brush refers to whether it collides with objects in the world and whether the Brush causes BSP cuts to be created in the surrounding Brushes when building geometry. 
 
 The solidity of Brushes can be changed in the **Details** panel after the Brush is created:
 
@@ -424,7 +422,7 @@ The three types of solidity are explained below.
 
 ### Solid
 
-Solid Brushes are the default type of Brush. These are what you get when you create a new additive or subtractive Brush. They have the following attributes:
+**Solid** Brushes are the default type of Brush. These are what you get when you create a new additive or subtractive Brush. They have the following attributes:
 
 
 * Block players and projectiles in the game.
@@ -435,7 +433,7 @@ Solid Brushes are the default type of Brush. These are what you get when you cre
 
 ### Semi-Solid
 
-Semi-solid Brushes are colliding Brushes that can be placed in a level without creating cuts to the surrounding world geometry. These can be used to create things such as pillars and support beams, though such objects should normally be created through the use of Static Meshes. They have the following attributes:
+**Semi-solid** Brushes are colliding Brushes that can be placed in a level without creating cuts to the surrounding world geometry. These can be used to create things such as pillars and support beams, though such objects should normally be created through the use of Static Meshes. They have the following attributes:
 
 
 * Block players and projectiles, just as Solid Brushes do.
@@ -446,7 +444,7 @@ Semi-solid Brushes are colliding Brushes that can be placed in a level without c
 
 ### Non-Solid
 
-Non-Solid Brushes are non-colliding Brushes that also do not create cuts in the surrounding world geometry. These have the effect of being visible but unable to be interacted with in any way. They have the following attributes:
+**Non-Solid** Brushes are non-colliding Brushes that also do not create cuts in the surrounding world geometry. These have the effect of being visible but unable to be interacted with in any way. They have the following attributes:
 
 
 * Do not block players or projectiles.
