Removed "Overview" from the top of these pages.

diff --git a/Engine/Documentation/Source/Engine/Rendering/LightingAndShadows/Basics/PlacingLights.INT.udn b/Engine/Documentation/Source/Engine/Rendering/LightingAndShadows/Basics/PlacingLights.INT.udn
--- a/Engine/Documentation/Source/Engine/Rendering/LightingAndShadows/Basics/PlacingLights.INT.udn
+++ b/Engine/Documentation/Source/Engine/Rendering/LightingAndShadows/Basics/PlacingLights.INT.udn
@@ -66,7 +66,7 @@ Description:The basics of placing lights into levels and setting them up.
 
 [TOC(start:2)]
 
-## Overview
+
 
 In Unreal Engine 4, there are a number of ways to add lights to a scene and there are a few key properties that have the greatest impact on lighting in the world.
 
diff --git a/Engine/Documentation/Source/Engine/Rendering/LightingAndShadows/LightMobility/LightMobility.INT.udn b/Engine/Documentation/Source/Engine/Rendering/LightingAndShadows/LightMobility/LightMobility.INT.udn
--- a/Engine/Documentation/Source/Engine/Rendering/LightingAndShadows/LightMobility/LightMobility.INT.udn
+++ b/Engine/Documentation/Source/Engine/Rendering/LightingAndShadows/LightMobility/LightMobility.INT.udn
@@ -46,7 +46,7 @@ Description:Light Mobility and how it affects lighting and shadowing.
 ![](Engine/Rendering/LightingAndShadows/lighting_overview_banner_001.png)(convert:false)
 [/REGION]
 
-## Overview
+
 
 On every light, within the **Transform** category, you will find a property called **Mobility**.
 
diff --git a/Engine/Documentation/Source/Engine/Rendering/LightingAndShadows/LightMobility/StationaryLights/StationaryLights.INT.udn b/Engine/Documentation/Source/Engine/Rendering/LightingAndShadows/LightMobility/StationaryLights/StationaryLights.INT.udn
--- a/Engine/Documentation/Source/Engine/Rendering/LightingAndShadows/LightMobility/StationaryLights/StationaryLights.INT.udn
+++ b/Engine/Documentation/Source/Engine/Rendering/LightingAndShadows/LightMobility/StationaryLights/StationaryLights.INT.udn
@@ -48,7 +48,7 @@ Description:Light type that can change its color and brightness at runtime but c
 
 [TOC(start:2 end:2)]
 
-## Overview
+
 
 [INCLUDE:Resources/ContentExamples\Lighting\2_2#StationaryLights01]
 
