#UE4doc #docs: Removed Overview header
Bolded a term
Fixed some spelling errors.

diff --git a/Engine/Documentation/Source/Engine/Actors/TargetPoint/TargetPoint.INT.udn b/Engine/Documentation/Source/Engine/Actors/TargetPoint/TargetPoint.INT.udn
--- a/Engine/Documentation/Source/Engine/Actors/TargetPoint/TargetPoint.INT.udn
+++ b/Engine/Documentation/Source/Engine/Actors/TargetPoint/TargetPoint.INT.udn
@@ -7,15 +7,15 @@ Related: Gameplay/HowTo/SpawnAndDestroyActors/Blueprints
 ![](Target_Point_Actors.png)
 
 ## Overview
-When creating worlds for games you sometimes need the ability to place and spawn items in the world for the player to interact with. Target Point Actors do just this; they give you a generic point in the world that you can spawn items from. If you are familiar with other 3D applications like 3Ds Max or Maya, Target Point Actors are very similar to dummy Actors that you find in those programs. 
+When creating worlds for games you sometimes need the ability to place and spawn items in the world for the player to interact with. **Target Point Actors** do just this; they give you a generic point in the world that you can spawn items from. If you are familiar with other 3D applications like 3Ds Max or Maya, Target Point Actors are very similar to dummy Actors that you find in those programs. 
 
 ## Placing Target Point Actors
-Target Point Actors can be found in the **Modes** panel under the **Basic** category. To add one to the world, simply select it from the Modes panel and then drag it into the world. 
+Target Point Actors can be found in the **Modes** panel under the **All Classes** category. To add one to the world, simply select it from the **Modes** panel and then drag it into the world. 
 
 ![](Using_Target_Point_Actor.gif)
 
 ## Using Target Points
-Target Point Actors can be used for all kids of things inside Unreal Engine 4.0. Here is a list of some of the things that you can use Target Point Actors for.
+Target Point Actors can be used for all kinds of things inside Unreal Engine 4. Here is a list of some of the things that you can use Target Point Actors for.
 
 * Look-at targets for cameras during cinematic sequences.
 * Path points for AI agents.
