#UE4doc #docs: Removed Overview header
Changed color of a green border to yellow
Replaced image with updated visuals
Removed image / method of creating a trigger event as it has been deprecated

diff --git a/Engine/Documentation/Source/Engine/Actors/Triggers/TriggerActors.INT.udn b/Engine/Documentation/Source/Engine/Actors/Triggers/TriggerActors.INT.udn
--- a/Engine/Documentation/Source/Engine/Actors/Triggers/TriggerActors.INT.udn
+++ b/Engine/Documentation/Source/Engine/Actors/Triggers/TriggerActors.INT.udn
@@ -25,13 +25,11 @@ Description:Actor that can be activated and cause events to occur in the level.
 
 [TOC(start:2)]
 
-## Overview
-
 [EXCERPT:Intro]
 **Triggers** are Actors that are used to cause an event to occur when they are interacted with by some other object in the 
 level. In other words, they are used to _trigger_ events in response to some other action in the level. All of the default 
 Triggers are generally the same, differing only in the shape of the area of influence - box, capsule, and sphere - used by 
-the Trigger to detect if another object has activated it.
+the Trigger to detect if another object has activated it. 
 [/EXCERPT:Intro]
 
 [REGION:imagetable]
@@ -53,11 +51,7 @@ Triggers are used to activate events placed inside of the [Level Blueprint](Engi
 There are several different types of events that a Trigger can activate. The main ones happen in response to a some type of 
 collision with another object, such as something hitting or overlapping with the Trigger, or in response to input from the player.
 
-There are a few different methods of adding an event to the Level Blueprint for a Trigger. With the Trigger selected in the viewport:
-
-* Choose an event from the **Add Level Events for [Trigger Actor Name]** dropdown in the **Details** panel.
-
-	![](trigger_event_details.png)
+With the Trigger selected in the **Viewport**:
 
 * **Right-click** in the Level Blueprint Event Graph and choose one of the events under **Add Event for [Trigger Actor Name]**.
 
