#UE4doc #docs:
Removed Overview header
Bolded a few editor terms

diff --git a/Engine/Documentation/Source/Engine/Actors/Volumes/Volumes.INT.udn b/Engine/Documentation/Source/Engine/Actors/Volumes/Volumes.INT.udn
--- a/Engine/Documentation/Source/Engine/Actors/Volumes/Volumes.INT.udn
+++ b/Engine/Documentation/Source/Engine/Actors/Volumes/Volumes.INT.udn
@@ -27,7 +27,6 @@ Description:Guide to placing and working with Volumes to augment the behavior an
 
 [TOC(start:2 end:2)]
 
-## Overview
 
 **Volumes** are 3-dimensional areas within your levels, each of which serve a specific purpose. Some of these purposes include:
 
@@ -43,7 +42,7 @@ There are many types of volumes that perform a variety of tasks. This document w
 One aspect about volumes that is critical to new level designers is the importance of using them wisely in a level. This boils down to two key points: 
 
 1. Volumes are invisible during gameplay; Players cannot see them!
-1. Due to the first point, **very rarely will you use a volume by itself for any effect.** 
+1. Due to the first point, very rarely will you use a volume by itself for any effect. 
 
 Because volumes are invisible, Players will generally not know that they have entered one; not unless you, as the level designer, help them a bit. Volumes should generally be thought of as components of a greater effect, with some other level element acting as a visual cue to lend understanding and complete the effect.
 
@@ -211,7 +210,7 @@ Several properties can be adjusted from the **Details** panel on this volume to
 | Fade Time | This is the time (in seconds) to fade from the current reverb settings into the volumes setting.  |
 | Enabled | Determines whether the volume is currently enabled and is able to affect sounds. |
 
-Ambient Zones define how Sound Actors located inside the associated Audio Volume will be altered by the Player's location. The settings for the Ambient Zone can also be adjusted from the **Details** panel. 
+**Ambient Zones** define how Sound Actors located inside the associated Audio Volume will be altered by the Player's location. The settings for the Ambient Zone can also be adjusted from the **Details** panel. 
 
 |Property| Description |
 | ------- | ------------ |
