#UE4doc #docs:  Removed "Overview" heading

diff --git a/Engine/Documentation/Source/Engine/Rendering/Materials/HowTo/AnimatingUVCoords/AnimatingUVCoordinates.INT.udn b/Engine/Documentation/Source/Engine/Rendering/Materials/HowTo/AnimatingUVCoords/AnimatingUVCoordinates.INT.udn
--- a/Engine/Documentation/Source/Engine/Rendering/Materials/HowTo/AnimatingUVCoords/AnimatingUVCoordinates.INT.udn
+++ b/Engine/Documentation/Source/Engine/Rendering/Materials/HowTo/AnimatingUVCoords/AnimatingUVCoordinates.INT.udn
@@ -9,7 +9,7 @@ Related: Engine/Rendering/LightingAndShadows/LightFunctions
 
 [TOC(start:2 end:3)]
 
-## Overview
+
 
 The ability to add motion to your Materials is a must, especially when you try to recreate effects such as fire, water, or smoke. 
 A very cheap and effective way to do this inside Unreal Engine 4 (UE4) is to use the **Panner Material Expression** node. 
diff --git a/Engine/Documentation/Source/Engine/Rendering/Materials/HowTo/ColoredTransluscentShadows/ColoredTransluscentShadows.INT.udn b/Engine/Documentation/Source/Engine/Rendering/Materials/HowTo/ColoredTransluscentShadows/ColoredTransluscentShadows.INT.udn
--- a/Engine/Documentation/Source/Engine/Rendering/Materials/HowTo/ColoredTransluscentShadows/ColoredTransluscentShadows.INT.udn
+++ b/Engine/Documentation/Source/Engine/Rendering/Materials/HowTo/ColoredTransluscentShadows/ColoredTransluscentShadows.INT.udn
@@ -9,7 +9,7 @@ Related: Engine/Rendering/LightingAndShadows/LitTranslucency
 
 [TOC(start:2 end:3)]
 
-## Overview
+
 
 In this "How To," we intend to show you how you can setup and use **Translucent Shadow Color** to create a material that can cast a colored shadow. 
 There are many applications where this could be useful to inject some color into your scene, but a common use would be for something like a stained glass window.
diff --git a/Engine/Documentation/Source/Engine/Rendering/Materials/HowTo/DetailTexturing/DetailTexturing.INT.udn b/Engine/Documentation/Source/Engine/Rendering/Materials/HowTo/DetailTexturing/DetailTexturing.INT.udn
--- a/Engine/Documentation/Source/Engine/Rendering/Materials/HowTo/DetailTexturing/DetailTexturing.INT.udn
+++ b/Engine/Documentation/Source/Engine/Rendering/Materials/HowTo/DetailTexturing/DetailTexturing.INT.udn
@@ -10,7 +10,7 @@ Related: Engine/Rendering/Materials/ExpressionReference/Textures
 
 [TOC(start:2 end:3)]
 
-## Overview
+
 
 When viewing a Material very closely, you will notice that the Textures used in the Material break down and become very pixilated, resulting in something that does not look very good. 
 To address this issue, you can use what is known as a Detail Texture to help hide a Textures pixilation when a Material is viewed very closely. 
diff --git a/Engine/Documentation/Source/Engine/Rendering/Materials/HowTo/EmissiveGlow/EmissiveGlow.INT.udn b/Engine/Documentation/Source/Engine/Rendering/Materials/HowTo/EmissiveGlow/EmissiveGlow.INT.udn
--- a/Engine/Documentation/Source/Engine/Rendering/Materials/HowTo/EmissiveGlow/EmissiveGlow.INT.udn
+++ b/Engine/Documentation/Source/Engine/Rendering/Materials/HowTo/EmissiveGlow/EmissiveGlow.INT.udn
@@ -10,7 +10,7 @@ Related: Resources/ContentExamples/MaterialNodes/1_5
 
 [TOC(start:2 end:3)]
 
-## Overview 
+
 
 During the course of project development, the need might arise for a Self-Illuminated Material or glowing Material. 
 Inside of Unreal Engine 4(UE4), Self-Illuminated Materials are known as Emissive Materials. 
diff --git a/Engine/Documentation/Source/Engine/Rendering/Materials/HowTo/Fresnel/Fresnel.INT.udn b/Engine/Documentation/Source/Engine/Rendering/Materials/HowTo/Fresnel/Fresnel.INT.udn
--- a/Engine/Documentation/Source/Engine/Rendering/Materials/HowTo/Fresnel/Fresnel.INT.udn
+++ b/Engine/Documentation/Source/Engine/Rendering/Materials/HowTo/Fresnel/Fresnel.INT.udn
@@ -11,7 +11,7 @@ Related: Engine/Rendering/Materials/ExpressionReference/Utility
 
 ![](FN_Fresnel_Sample.png)
 
-## Overview 
+
 
 Often times when creating a film or cinematic, you need a way to better highlight the silhouette of a character. This lighting technique is often referred to as rim lighting, or edge lighting, and 
 involves adding additional lights to help better distinguish the Characters silhouette from the background. While this works quite nicely for film, adding additional lights to any 3D scene adds 
diff --git a/Engine/Documentation/Source/Engine/Rendering/Materials/HowTo/Human_Skin/HumanSkin.INT.udn b/Engine/Documentation/Source/Engine/Rendering/Materials/HowTo/Human_Skin/HumanSkin.INT.udn
--- a/Engine/Documentation/Source/Engine/Rendering/Materials/HowTo/Human_Skin/HumanSkin.INT.udn
+++ b/Engine/Documentation/Source/Engine/Rendering/Materials/HowTo/Human_Skin/HumanSkin.INT.udn
@@ -11,7 +11,7 @@ Related: Resources/ContentExamples/MaterialNodes/1_7
 
 [TOC(start:2)]
 
-## Overview
+
 
 [REGION:lightbox]
 [![](Engine\Rendering\Materials\LightingModels\SubSurfaceProfile\Results.png)(w:920 convert:false)](Engine\Rendering\Materials\LightingModels\SubSurfaceProfile\Results.png)
diff --git a/Engine/Documentation/Source/Engine/Rendering/Materials/HowTo/Instancing/Instancing.INT.udn b/Engine/Documentation/Source/Engine/Rendering/Materials/HowTo/Instancing/Instancing.INT.udn
--- a/Engine/Documentation/Source/Engine/Rendering/Materials/HowTo/Instancing/Instancing.INT.udn
+++ b/Engine/Documentation/Source/Engine/Rendering/Materials/HowTo/Instancing/Instancing.INT.udn
@@ -11,7 +11,7 @@ Related: Engine/Rendering/Materials/MaterialProperties
 
 [TOC(start:2 end:3)]
 
-## Overview 
+
 
 Having to create, setup, and tweak standard Materials can be a very time consuming process in UE4. 
 To help speed up and streamline this process, Unreal Engine 4 offers a special type of Material called a **Material Instance**. 
diff --git a/Engine/Documentation/Source/Engine/Rendering/Materials/HowTo/Main_Material_Node/Main_Material_Node.INT.udn b/Engine/Documentation/Source/Engine/Rendering/Materials/HowTo/Main_Material_Node/Main_Material_Node.INT.udn
--- a/Engine/Documentation/Source/Engine/Rendering/Materials/HowTo/Main_Material_Node/Main_Material_Node.INT.udn
+++ b/Engine/Documentation/Source/Engine/Rendering/Materials/HowTo/Main_Material_Node/Main_Material_Node.INT.udn
@@ -10,7 +10,7 @@ Related: Engine/Rendering/Materials/MaterialProperties
 
 [TOC(start:2 end:3)]
 
-## Overview
+
 
 The Main Material node of the Material Graph is where all of the Material Expression nodes that are added to the Material Graph will end up being connected to. 
 The results of what is input into the Main Material node is what will be displayed when the Material is compiled and applied to objects in game. 
diff --git a/Engine/Documentation/Source/Engine/Rendering/Materials/HowTo/Making_Functions/Making_Functions.INT.udn b/Engine/Documentation/Source/Engine/Rendering/Materials/HowTo/Making_Functions/Making_Functions.INT.udn
--- a/Engine/Documentation/Source/Engine/Rendering/Materials/HowTo/Making_Functions/Making_Functions.INT.udn
+++ b/Engine/Documentation/Source/Engine/Rendering/Materials/HowTo/Making_Functions/Making_Functions.INT.udn
@@ -11,7 +11,7 @@ Related: Engine/Rendering/Materials
 
 [TOC(start:2)]
 
-## Overview
+
 The ability to reuse part of one Material Graph inside another Material Graph can greatly speed up both Material creation time and Material workflow. 
 Unreal Engine 4 (UE4) offers users the ability to do this using what is known as a Material Function.
 In this How To example, we will be re-creating the Fresnel style shading effect that you often see inside of Video Games using a Material Function. 
diff --git a/Engine/Documentation/Source/Engine/Rendering/Materials/HowTo/Making_Parameters/Making_Parameters.INT.udn b/Engine/Documentation/Source/Engine/Rendering/Materials/HowTo/Making_Parameters/Making_Parameters.INT.udn
--- a/Engine/Documentation/Source/Engine/Rendering/Materials/HowTo/Making_Parameters/Making_Parameters.INT.udn
+++ b/Engine/Documentation/Source/Engine/Rendering/Materials/HowTo/Making_Parameters/Making_Parameters.INT.udn
@@ -9,7 +9,7 @@ Related: Engine/Matinee/HowTo/MHT_6
 
 [TOC(start:2)]
 
-## Overview
+
 
 To interact with Material Instances, you need to use a special kind of Material Expression node called a **Material Parameter**. 
 In this How To, we will be looking at how you can add or convert Material Expression nodes into Material Parameters nodes and also how you can use those Material Parameters nodes inside of your Materials.
diff --git a/Engine/Documentation/Source/Engine/Rendering/Materials/HowTo/Masking/Masking.INT.udn b/Engine/Documentation/Source/Engine/Rendering/Materials/HowTo/Masking/Masking.INT.udn
--- a/Engine/Documentation/Source/Engine/Rendering/Materials/HowTo/Masking/Masking.INT.udn
+++ b/Engine/Documentation/Source/Engine/Rendering/Materials/HowTo/Masking/Masking.INT.udn
@@ -12,7 +12,7 @@ Related: Engine/Rendering/Materials/HowTo/Instancing
 
 [TOC(start:2 end:3)]
 
-## Overview 
+
 
 When creating 3D assets, you might find that you need the ability to define different surfaces types within the same Material. 
 A simple and cheap way to achieve this is to use a Texture Mask that defines which surface should be affected by which section of the Material. 
diff --git a/Engine/Documentation/Source/Engine/Rendering/Materials/HowTo/Placing_Expressions_Functions/Placing_Expressions_Functions.INT.udn b/Engine/Documentation/Source/Engine/Rendering/Materials/HowTo/Placing_Expressions_Functions/Placing_Expressions_Functions.INT.udn
--- a/Engine/Documentation/Source/Engine/Rendering/Materials/HowTo/Placing_Expressions_Functions/Placing_Expressions_Functions.INT.udn
+++ b/Engine/Documentation/Source/Engine/Rendering/Materials/HowTo/Placing_Expressions_Functions/Placing_Expressions_Functions.INT.udn
@@ -8,7 +8,7 @@ Related: Engine/UI/KeyBindings
 
 [TOC(start:2)]
 
-## Overview
+
 
 Material Expressions and Material Functions are the building blocks used to create fully functional Materials inside Unreal Engine 4.
 Each expression or function is a self-contained black box that can run little snippets of code on their inputs and output the results. 
diff --git a/Engine/Documentation/Source/Engine/Rendering/Materials/HowTo/Previewing_Applying/Previewing_Applying.INT.udn b/Engine/Documentation/Source/Engine/Rendering/Materials/HowTo/Previewing_Applying/Previewing_Applying.INT.udn
--- a/Engine/Documentation/Source/Engine/Rendering/Materials/HowTo/Previewing_Applying/Previewing_Applying.INT.udn
+++ b/Engine/Documentation/Source/Engine/Rendering/Materials/HowTo/Previewing_Applying/Previewing_Applying.INT.udn
@@ -10,7 +10,7 @@ Related: Engine/Rendering/Materials/MaterialProperties
 
 [TOC(start:2 end:3)]
 
-## Overview 
+
 
 Previewing and applying your Materials is something that you will have to do quite often to make sure that you are getting the exact results you want. 
 Unreal Engine 4 (UE4) makes both previewing and applying materials a very straight forward process.
diff --git a/Engine/Documentation/Source/Engine/Rendering/Materials/HowTo/Refraction/Refraction.INT.udn b/Engine/Documentation/Source/Engine/Rendering/Materials/HowTo/Refraction/Refraction.INT.udn
--- a/Engine/Documentation/Source/Engine/Rendering/Materials/HowTo/Refraction/Refraction.INT.udn
+++ b/Engine/Documentation/Source/Engine/Rendering/Materials/HowTo/Refraction/Refraction.INT.udn
@@ -9,7 +9,7 @@ Related: Engine/Rendering/Materials/MaterialInputs
 
 [TOC(start:2 end:3)]
 
-## Overview
+
 
 When light passes from one medium to another, say from air to water, the direction that the light was traveling will be altered where the two mediums meet. 
 This lighting phenomenon is referred to as Refraction and happens because the speed at which light travels can be altered by the different types of materials light come in contact with. 
diff --git a/Engine/Documentation/Source/Engine/Rendering/Materials/HowTo/Subsurface_Profile/Subsurface_Profile.INT.udn b/Engine/Documentation/Source/Engine/Rendering/Materials/HowTo/Subsurface_Profile/Subsurface_Profile.INT.udn
--- a/Engine/Documentation/Source/Engine/Rendering/Materials/HowTo/Subsurface_Profile/Subsurface_Profile.INT.udn
+++ b/Engine/Documentation/Source/Engine/Rendering/Materials/HowTo/Subsurface_Profile/Subsurface_Profile.INT.udn
@@ -9,7 +9,7 @@ Related: Resources/ContentExamples/MaterialNodes/1_7
 
 [TOC(start:2)]
 
-## Overview 
+
 The ability to rendering realistic looking human skin is a must for any modern Video Game engine. To fulfill this need, Unreal Engine 4 (UE4) now offers a shading method specifically for skin or wax like surfaces called **Subsurface Profile**. 
 While the Subsurface Profile shading model has similar properties to the **Subsurface** shading model, its key difference is in how it renders. 
 SubsurfaceProfile uses a rendering approach that is based in Screen Space because this helps to better display the subtle subsurface effects seen in human skin where backscattering is a secondary effect only seen in few cases like ears for example.
diff --git a/Engine/Documentation/Source/Engine/Rendering/Materials/HowTo/Subsurface_Scattering/Subsurface_Scattering.INT.udn b/Engine/Documentation/Source/Engine/Rendering/Materials/HowTo/Subsurface_Scattering/Subsurface_Scattering.INT.udn
--- a/Engine/Documentation/Source/Engine/Rendering/Materials/HowTo/Subsurface_Scattering/Subsurface_Scattering.INT.udn
+++ b/Engine/Documentation/Source/Engine/Rendering/Materials/HowTo/Subsurface_Scattering/Subsurface_Scattering.INT.udn
@@ -9,7 +9,7 @@ Related: Resources/ContentExamples/MaterialNodes/1_7
 
 [TOC(start:2)]
 
-## Overview 
+
 Subsurface scattering is the term used to describe the lighting phenomenon where light scatters as it passes through a translucent / semi translucent surface.
  Unreal Engine 4 (UE4) offers a special [Shading Model ](Engine/Rendering/Materials/MaterialProperties/LightingModels) called [Subsurface](Engine/Rendering/Materials/LightingModels/SubSurface) that is used specifically for materials that need this interaction like skin or wax.
  The following How To guide will walk you through all you need to know about how to use Subsurface scattering in your Materials.
diff --git a/Engine/Documentation/Source/Engine/Rendering/Materials/HowTo/Transparency/Transparency.INT.udn b/Engine/Documentation/Source/Engine/Rendering/Materials/HowTo/Transparency/Transparency.INT.udn
--- a/Engine/Documentation/Source/Engine/Rendering/Materials/HowTo/Transparency/Transparency.INT.udn
+++ b/Engine/Documentation/Source/Engine/Rendering/Materials/HowTo/Transparency/Transparency.INT.udn
@@ -8,7 +8,7 @@ Related: Engine/Rendering/LightingAndShadows/LitTranslucency
 
 [TOC(start:2 end:3)]
 
-## Overview
+
 
 When creating certain surface types such as water or glass, you need the ability to make the surface not only see through, but also give the surface a sense of depth and color. 
 In the real world, these properties are often referred to as **Transparency** or **Opacity** and are often used interchangeably to describe the same thing. 
