#UE4doc #docs:
Bolded some editor terms
Removed unnecessary quotation marks around editor terms.

diff --git a/Engine/Documentation/Source/Engine/AI/BehaviorTrees/HowUE4BehaviorTreesDiffer/HowUE4BehaviorTreesDiffer.INT.udn b/Engine/Documentation/Source/Engine/AI/BehaviorTrees/HowUE4BehaviorTreesDiffer/HowUE4BehaviorTreesDiffer.INT.udn
--- a/Engine/Documentation/Source/Engine/AI/BehaviorTrees/HowUE4BehaviorTreesDiffer/HowUE4BehaviorTreesDiffer.INT.udn
+++ b/Engine/Documentation/Source/Engine/AI/BehaviorTrees/HowUE4BehaviorTreesDiffer/HowUE4BehaviorTreesDiffer.INT.udn
@@ -15,9 +15,9 @@ There are three critical ways in which the UE4 implementation of Behavior Trees
 
 Having an event-driven architecture grants improvements to both performance and debugging.  However, to take the most advantage of these improvements, you will need to understand the other differences to our trees and structure your behavior trees appropriately.
 
-**Performance Improvements** - Since the code does not have to iterate through the entire tree every tick, performance is much better!  Conceptually, instead of constantly asking "Are we there yet?", we can just rest until we are prodded and told "We are there!"
+Since the code does not have to iterate through the entire tree every tick, performance is much better!  Conceptually, instead of constantly asking "Are we there yet?", we can just rest until we are prodded and told "We are there!"
 
-**Better debugging** - When stepping forward and backward through the behavior tree's execution history to visually debug the behavior, it is ideal to make the history show relevant changes and NOT show irrelevant ones.  In the event-driven implementation, it is not necessary to filter out irrelevant steps that iterated over the tree and chose the same behavior as before, because that additional iteration never had to happen in the first place!  Instead, only changes to the location of execution in the tree or to blackboard values matter, and it is easy to just show those differences.  <!-- (See the debugging section for more information on how to step through the execution history and other debugging tricks.) -->
+When stepping forward and backward through the behavior tree's execution history to visually debug the behavior, it is ideal to make the history show relevant changes and NOT show irrelevant ones.  In the event-driven implementation, it is not necessary to filter out irrelevant steps that iterated over the tree and chose the same behavior as before, because that additional iteration never had to happen in the first place!  Instead, only changes to the location of execution in the tree or to blackboard values matter, and it is easy to just show those differences.  <!-- (See the debugging section for more information on how to step through the execution history and other debugging tricks.) -->
 
 ## Conditionals Are Not Leaf Nodes
 
@@ -35,34 +35,34 @@ Another advantage of conditional decorators is that it is easy to make those dec
 
 ## Special Handling for Concurrent Behaviors
 
- "Standard" behavior trees often use a "Parallel" composite node to handle concurrent behaviors.  The Parallel node begins execution on all of its children simultaneously.  Special rules determine how to act if one or more of those child trees finish (depending on the desired behavior).
+ Standard behavior trees often use a **Parallel** composite node to handle concurrent behaviors.  The Parallel node begins execution on all of its children simultaneously.  Special rules determine how to act if one or more of those child trees finish (depending on the desired behavior).
 
 [REGION:note]
-**Parallel** nodes are not necessarily multi-threading (executing tasks at "truly" the same time).  They are just a way to conceptually perform several tasks at once.  Often they still run on the same thread and begin in some sequence.  That sequence should be irrelevant since they will all happen in the same frame, but it is still sometimes important.
+Parallel nodes are not necessarily multi-threading (executing tasks at "truly" the same time).  They are just a way to conceptually perform several tasks at once.  Often they still run on the same thread and begin in some sequence.  That sequence should be irrelevant since they will all happen in the same frame, but it is still sometimes important.
 [/REGION]
 
-Instead of complex "Parallel" nodes, UE4 Behavior Trees use "Simple Parallel" nodes and our own special node type which we call "Services" to accomplish the same sorts of behaviors.
+Instead of complex Parallel nodes, UE4 Behavior Trees use **Simple Parallel** nodes and our own special node type which we call Services to accomplish the same sorts of behaviors.
 
-### Why not use "Parallel" nodes?
+### Why not use Parallel nodes?
 
 1.	Parallel nodes can be very confusing, even for relatively simple behaviors. 
 
-	Effectively Parallel nodes are simultaneously running a bunch of separate sub-trees, but any or all of those sub-trees may need to abort if one of them fails.  Or they may "succeed" when the others finish (whether successful or failing).  Or... etc.!  Parallel behaviors can be confusing even in simple cases, and with the number of options potentially available, it can become highly confusing.
+	Effectively Parallel nodes are simultaneously running a bunch of separate sub-trees, but any or all of those sub-trees may need to abort if one of them fails.  Or they may succeed when the others finish (whether successful or failing).  Or... etc.!  Parallel behaviors can be confusing even in simple cases, and with the number of options potentially available, it can become highly confusing.
 
 1.	Parallel nodes make it harder to optimize performance, especially in terms of making event-driven trees.
 
 
-### What does UE4 use instead of "Parallel" nodes?
+### What does UE4 use instead of Parallel nodes?
 
 There are three types of nodes which provide the functionality that would normally come from Parallel nodes:
 
-#### "Simple Parallel" nodes
+#### Simple Parallel nodes
 
 Simple Parallel nodes allow only two children: one which must be a single task node (with optional decorators), and the other of which can be a complete subtree.
 
 You can think of the Simple Parallel node as "While doing A, do B as well."  For example, "While attacking the enemy, move toward the enemy."  Basically, A is a primary task, and B is a secondary or filler task while waiting for A to complete.
 
-While there are some options as to how to handle the secondary "meanwhile" task (Task B), the node is relatively simple in concept compared to traditional Parallel nodes.  Nonetheless, it supports much of the most common usage of Parallel nodes.
+While there are some options as to how to handle the secondary meanwhile task (Task B), the node is relatively simple in concept compared to traditional Parallel nodes.  Nonetheless, it supports much of the most common usage of Parallel nodes.
 
 Simple Parallel nodes allow easy usage of some of our event-driven optimizations.  Full Parallel nodes would be much more complex to optimize.
 
@@ -79,7 +79,7 @@ Services are active only as long as execution remains in the subtree rooted at t
 
 #### Decorator "Observer Aborts" Property
 
-One common usage case for standard Parallel nodes is to constantly check conditions so that a task can abort if the conditions it requires becomes _false_.  For example, if you have a cat that performs a sequence "Shake Rear End", "Pounce", you may want to give up immediately if the mouse escapes into its mouse hole.  With "Parallel" nodes, you would have a child that checks if the mouse can be pounced on, and then another child that's the sequence to perform.  Since our Behavior Trees are event-driven, we instead handle this by having our conditional decorators "observe" their values and abort when necessary.  (In this example, you would just have the "Mouse Can Be Pounced On?" decorator on the sequence itself, with "Observer Aborts" set to "Self". )
+One common usage case for standard Parallel nodes is to constantly check conditions so that a task can abort if the conditions it requires becomes _false_.  For example, if you have a cat that performs a sequence "Shake Rear End", "Pounce", you may want to give up immediately if the mouse escapes into its mouse hole.  With Parallel nodes, you would have a child that checks if the mouse can be pounced on, and then another child that's the sequence to perform.  Since our Behavior Trees are event-driven, we instead handle this by having our conditional decorators observe their values and abort when necessary.  (In this example, you would just have the "Mouse Can Be Pounced On?" decorator on the sequence itself, with "Observer Aborts" set to "Self". )
 
 <!-- For more information, see the <LINK>"Observer Aborts" explanation page. -->
 
@@ -94,7 +94,7 @@ Using Services and Simple Parallel nodes creates simple trees that are easier to
 
 ### Ease of Debugging
 
-Clearer graphs are easier to debug.  In addition, having fewer "simultaneous" execution paths is a huge boon to watching what is actually happening in the graph.
+Clearer graphs are easier to debug.  In addition, having fewer simultaneous execution paths is a huge boon to watching what is actually happening in the graph.
 
 ### Easier Optimizations
 
