#UE4doc #docs:
Unbolded and bolded multiple hyperlinks / editor terms

diff --git a/Engine/Documentation/Source/Engine/AI/BehaviorTrees/NodeReference/Composites/Composites.INT.udn b/Engine/Documentation/Source/Engine/AI/BehaviorTrees/NodeReference/Composites/Composites.INT.udn
--- a/Engine/Documentation/Source/Engine/AI/BehaviorTrees/NodeReference/Composites/Composites.INT.udn
+++ b/Engine/Documentation/Source/Engine/AI/BehaviorTrees/NodeReference/Composites/Composites.INT.udn
@@ -5,13 +5,13 @@ Availability:Public
 
 [TOC(start:2 end:2)]
 
-**Composite Nodes** define the root of a branch and the base rules for how that branch is executed. They can have [**Decorators**](Engine/AI/BehaviorTrees/NodeReference/Decorators) applied to them to modify entry into their branch or even cancel out mid execution. Also, they can have [**Services**](Engine/AI/BehaviorTrees/NodeReference/Services) attached to them that will only be active if the children of the Composite are being executed.
+**Composite Nodes** define the root of a branch and the base rules for how that branch is executed. They can have [Decorators](Engine/AI/BehaviorTrees/NodeReference/Decorators) applied to them to modify entry into their branch or even cancel out mid execution. Also, they can have [Services](Engine/AI/BehaviorTrees/NodeReference/Services) attached to them that will only be active if the children of the Composite are being executed.
 
 ## Selector
 
 ![](Selector.png)
 
-Selector Nodes execute their children from left to right, and will stop executing its children when one of their children **Succeeds**. If a Selector's child succeeds, the Selector succeeds. If all the Selector's children fail, the Selector fails.
+**Selector** Nodes execute their children from left to right, and will stop executing its children when one of their children **Succeeds**. If a Selector's child succeeds, the Selector succeeds. If all the Selector's children fail, the Selector fails.
 
 | Property | Description |
 | -------- | ----------- |
@@ -21,7 +21,7 @@ Selector Nodes execute their children from left to right, and will stop executin
 
 ![](Sequence.png)
 
-Sequence Nodes execute their children from left to right, and will stop executing its children when one of their children **Fails**. If a child fails, then the Sequence fails. If all the Sequence's children succeed, then the Sequence succeeds.
+**Sequence** Nodes execute their children from left to right, and will stop executing its children when one of their children **Fails**. If a child fails, then the Sequence fails. If all the Sequence's children succeed, then the Sequence succeeds.
 
 | Property | Description |
 | -------- | ----------- |
