#UE4doc #docs Minor grammatical update and removed bolded items (Compliance per Instructional Guide v.1).

diff --git a/Engine/Documentation/Source/Engine/Blueprints/BP_HowTo/BlueprintHowTo.INT.udn b/Engine/Documentation/Source/Engine/Blueprints/BP_HowTo/BlueprintHowTo.INT.udn
--- a/Engine/Documentation/Source/Engine/Blueprints/BP_HowTo/BlueprintHowTo.INT.udn
+++ b/Engine/Documentation/Source/Engine/Blueprints/BP_HowTo/BlueprintHowTo.INT.udn
@@ -7,9 +7,9 @@ Related: Engine/Blueprints/UserGuide
 Related: Gameplay
 
 
-If you are new to [Blueprints Visual Scripting](Engine/Blueprints) in Unreal Engine 4, the Blueprint How To's will provide information on the core aspects needed to working with Blueprints including Creating Blueprints, Adding and Connecting Nodes, Creating Variables and others as seen in the tables below. 
+If you are new to [](Engine/Blueprints) in Unreal Engine 4, these Blueprint How To's provide information on the core aspects needed to work with Blueprints including Creating Blueprints, Adding and Connecting Nodes, Creating Variables and others as seen in the tables below. 
 
-When you are ready and have an understanding of the core aspects behind Blueprints, **Gameplay** oriented Blueprint How To's can be found in the [](Gameplay) which illustrate gameplay related elements like Spawning/Destroying Actors, Respawning (OnDestroyed), Using Timers, and more. 
+When you are ready and have an understanding of the core aspects behind Blueprints, Gameplay oriented Blueprint How To's can be found in the [](Gameplay) which illustrate gameplay related elements like Spawning/Destroying Actors, Respawning (OnDestroyed), Using Timers, and more. 
 
 Refer to the tables below for more information and click each link to view its contents. 
 
@@ -18,32 +18,32 @@ Refer to the tables below for more information and click each link to view its c
 
 |How To... | Description |
 | ------- | ------------ |
-| **[](Engine/Blueprints/UserGuide/Types/ClassBlueprint/Creation)** | %Engine/Blueprints/UserGuide/Types/ClassBlueprint/Creation/:description% |
-| **[](Engine/Blueprints/BP_HowTo/PlacingNodes)** | %Engine/Blueprints/BP_HowTo/PlacingNodes:description% |
-| **[](Engine/Blueprints/BP_HowTo/ConnectingNodes)** | %Engine/Blueprints/BP_HowTo/ConnectingNodes:description% |
-| **[](Engine/Blueprints/BP_HowTo/NavigatingGraphs)** | %Engine/Blueprints/BP_HowTo/NavigatingGraphs:description% |
-| **[](Engine/Blueprints/BP_HowTo/CollapsingGraphs)** | %Engine/Blueprints/BP_HowTo/CollapsingGraphs:description% |
-| **[](Engine/Blueprints/BP_HowTo/Debugging)** | %Engine/Blueprints/BP_HowTo/Debugging:description% |
+| [](Engine/Blueprints/UserGuide/Types/ClassBlueprint/Creation) | %Engine/Blueprints/UserGuide/Types/ClassBlueprint/Creation/:description% |
+| [](Engine/Blueprints/BP_HowTo/PlacingNodes) | %Engine/Blueprints/BP_HowTo/PlacingNodes:description% |
+| [](Engine/Blueprints/BP_HowTo/ConnectingNodes) | %Engine/Blueprints/BP_HowTo/ConnectingNodes:description% |
+| [](Engine/Blueprints/BP_HowTo/NavigatingGraphs) | %Engine/Blueprints/BP_HowTo/NavigatingGraphs:description% |
+| [](Engine/Blueprints/BP_HowTo/CollapsingGraphs) | %Engine/Blueprints/BP_HowTo/CollapsingGraphs:description% |
+| [](Engine/Blueprints/BP_HowTo/Debugging) | %Engine/Blueprints/BP_HowTo/Debugging:description% |
 
 
 ## Intermediate
 
 |How To... | Description |
 | ------- | ------------ |
-| **[](Engine/Blueprints/BP_HowTo/ActorReference)** | %Engine/Blueprints/BP_HowTo/ActorReference:description% |
-| **[](Engine/Blueprints/BP_HowTo/MakingFunctions)** | %Engine/Blueprints/BP_HowTo/MakingFunctions:description% |
-| **[](Engine/Blueprints/BP_HowTo/WorkingWithArrays)** | %Engine/Blueprints/BP_HowTo/WorkingWithArrays:description% |
-| **[](Engine/Blueprints/UserGuide/Variables)** | The entire user guide for Variables. |
+| [](Engine/Blueprints/BP_HowTo/ActorReference) | %Engine/Blueprints/BP_HowTo/ActorReference:description% |
+| [](Engine/Blueprints/BP_HowTo/MakingFunctions) | %Engine/Blueprints/BP_HowTo/MakingFunctions:description% |
+| [](Engine/Blueprints/BP_HowTo/WorkingWithArrays) | %Engine/Blueprints/BP_HowTo/WorkingWithArrays:description% |
+| [](Engine/Blueprints/UserGuide/Variables) | The entire user guide for Variables. |
 
 ## Advanced
 
 |How To... | Description |
 | ------- | ------------ |
-| **[](Engine\Blueprints\BP_HowTo\DirectBlueprintComs)** | %Engine\Blueprints\BP_HowTo\DirectBlueprintComs:description% |
-| **[](Engine\Blueprints\BP_HowTo\Casting)** | %Engine\Blueprints\BP_HowTo\Casting:description% |
-| **[](Engine\Blueprints\BP_HowTo\Interfaces)** | %Engine\Blueprints\BP_HowTo\Interfaces:description% |
-| **[](Engine\Blueprints\BP_HowTo\EventDispatcher)** | %Engine\Blueprints\BP_HowTo\EventDispatcher:description% |
-| **[](Engine/Blueprints/BP_HowTo/MakingMacros)** | %Engine/Blueprints/BP_HowTo/MakingMacros:description% |
+| [](Engine\Blueprints\BP_HowTo\DirectBlueprintComs) | %Engine\Blueprints\BP_HowTo\DirectBlueprintComs:description% |
+| [](Engine\Blueprints\BP_HowTo\Casting) | %Engine\Blueprints\BP_HowTo\Casting:description% |
+| [](Engine\Blueprints\BP_HowTo\Interfaces) | %Engine\Blueprints\BP_HowTo\Interfaces:description% |
+| [](Engine\Blueprints\BP_HowTo\EventDispatcher) | %Engine\Blueprints\BP_HowTo\EventDispatcher:description% |
+| [](Engine/Blueprints/BP_HowTo/MakingMacros) | %Engine/Blueprints/BP_HowTo/MakingMacros:description% |
 
 
 <!-- To remove since information is redundant. 
