#UE4doc #docs: Adjusted bold text in accordance with the Instructional Guide

diff --git a/Engine/Documentation/Source/Engine/Rendering/Materials/HowTo/DetailTexturing/DetailTexturing.INT.udn b/Engine/Documentation/Source/Engine/Rendering/Materials/HowTo/DetailTexturing/DetailTexturing.INT.udn
--- a/Engine/Documentation/Source/Engine/Rendering/Materials/HowTo/DetailTexturing/DetailTexturing.INT.udn
+++ b/Engine/Documentation/Source/Engine/Rendering/Materials/HowTo/DetailTexturing/DetailTexturing.INT.udn
@@ -48,7 +48,7 @@ While this is not the only way that you can apply Detail Textures to Materials,
 
 ### Converting Texture Sample to Texture Objects
 
-In order for the **Detail Texturing** Material Function to work correctly, you will need to convert the Textures you want to use as the Detail Textures from regular Textures into Texture Objects. 
+In order for the Detail Texturing Material Function to work correctly, you will need to convert the Textures you want to use as the Detail Textures from regular Textures into Texture Objects. 
 To change Textures into Texture Objects, all you have to do is the following.
 
 1. First find the Texture you want to use as a Detail Texture. 
@@ -111,15 +111,15 @@ Once selected, a new **Material** will be created. Name the Material **Detail_Ma
 	Remember that the **T_Ground_Gravel** Diffuse and Normal maps Texture Samplers need to be converted into **Texture Objects** or the Material will not work. For more information on how to do this, please refer to the **Converting Texture Sample to Texture Objects** section above. 
 	[/region]
 
-1. Now that the Textures have been hooked up, it is time to add some **Constant** Material Expressions so that things like the scale and intensity of the Detail Textures can be tweaked. To find the Constant Material Expression, search for **Constant** in the Material Palette and then drag the **Constant** from the **Palette** to the Material Graph. Once you have added one **Constant** Material Expression, duplicate it by first selecting it and then press **CTRL + W** on the keyboard two times. In the end, you want to have three **Constant** Material Expressions. 
+1. Now that the Textures have been hooked up, it is time to add some Constant Material Expressions so that things like the scale and intensity of the Detail Textures can be tweaked. To find the Constant Material Expression, search for Constant in the Material Palette and then drag the Constant from the **Palette** to the Material Graph. Once you have added one Constant Material Expression, duplicate it by first selecting it and then press CTRL + W on the keyboard two times. In the end, you want to have three Constant Material Expressions. 
 
 	![](DT_Adding_Constant_Nodes.png)
 
-1. With the appropriate amount of **Constant** Material Expression nodes added to the Material Graph, it is now time to hook them up and fill them in with data. To do this, all you need to do is connect the output of the **Constant** Material Expression nodes to the inputs on the **Detail Texturing** Material Function. Then input a value of 1 in each of the **Constant** Material Expression nodes. When completed, your Material Graph should look like this. 
+1. With the appropriate amount of Constant Material Expression nodes added to the Material Graph, it is now time to hook them up and fill them in with data. To do this, all you need to do is connect the output of the **Constant** Material Expression nodes to the inputs on the **Detail Texturing** Material Function. Then input a value of 1 in each of the **Constant** Material Expression nodes. When completed, your Material Graph should look like this. 
 	
 	![](DT_All_Connections_Made.png)
 
-1. By adjusting the values of the **Constant** Material Expression nodes, you can affect how the Detail Textures will look. Here is an example of what happens when you set the **Scale** to a value of 1 then 5 and finally 10. 
+1. By adjusting the values of the Constant Material Expression nodes, you can affect how the Detail Textures will look. Here is an example of what happens when you set the Scale to a value of 1 then 5 and finally 10. 
 
 	![](DT_Different_Scale_Settings.png)
 
@@ -127,15 +127,15 @@ Once selected, a new **Material** will be created. Name the Material **Detail_Ma
 
 If for whatever reason you cannot use the **Detail Texturing** Material Function, you can always build this functionality inside the Material Graph using Material Expression nodes and the directions below.
 
-1. First make a copy of **Detail_Material** by **Right-Clicking** on the Material inside the **Content Browser** and then selecting **Duplicate** from the Menu. 
+1. First make a copy of Detail_Material by **Right-Clicking** on the Material inside the **Content Browser** and then selecting **Duplicate** from the Menu. 
 
 	![](DT_Duplicate_Material.png)
 
-1. Once created, label the Material **Detail_Material_2**. When completed, the **Content Browser** should look like this.
+1. Once created, label the Material Detail_Material_2. When completed, the **Content Browser** should look like this.
 
 	![](DT_Duplicate_Materia_Namedl.png)
 
-1. Open up the new created Material by **Double-Clicking** on it with the **Left Mouse Button**. When opened, delete everything except the Textures and make sure to convert the **Texture Objects** back to **Texture Samples** by first selecting all the **Texture Objects** and **Right-Clicking** on them and then selecting **Convert To Texture Sampler** from the menu. When completed, the Material Graph should look like this. 
+1. Open up the new created Material by **Double-Clicking** on it with the Left Mouse Button. When opened, delete everything except the Textures and make sure to convert the Texture Objects back to Texture Samples by first selecting all the Texture Objects and **Right-Clicking** on them and then selecting **Convert To Texture Sampler** from the menu. When completed, the Material Graph should look like this. 
 
 	![](DT_New_Material_Setup.png)
 
@@ -172,15 +172,15 @@ To fix this, we can build upon the Detail Material that was created previously t
 This is often referred to as Distance Based Texture Blending and can be set up following the instructions below. 
 
 
-1. First make a copy of the **Detail_Material_2** and name it **Distance_Fade**. Once named, open up the Material and search in the **Palette** for the following Material Expression nodes and add them to the Material Graph.
+1. First make a copy of the Detail_Material_2 and name it Distance_Fade. Once named, open up the Material and search in the **Palette** for the following Material Expression nodes and add them to the Material Graph.
 
-	* **Absolute World Position** x 1
-	* **Camera Position** x 1
-	* **Distance** x 1
-	* **Divide**  x 1
-	* **Power** x 1
-	* **Clamp** x 1
-	* **Constant** x 2
+	* Absolute World Position x 1
+	* Camera Position x 1
+	* Distance x 1
+	* Divide  x 1
+	* Power x 1
+	* Clamp x 1
+	* Constant x 2
 
 1. With the all the required Material Expression nodes added to the Material Graph, we can begin hooking everything up. The distance at which the Texture change happens is controlled by the first **Constant** that is plugged into the **Divide** Material Expression, currently set to 512 in the example image. When completed, you should have something that looks like this. 
 
@@ -190,7 +190,7 @@ This is often referred to as Distance Based Texture Blending and can be set up f
 
 	![](DT_Distance_Based_Completed_Material.png)
 
-1. To see the effect in action, all you have to do is hold down the **Right Mouse Button** while inside the Material Preview window and then move the mouse forward and backwards. When you get far enough away from the preview mesh, you should see the Material transition from one texture to the other one. If you cannot see this happening, adjust the first **Constant** that is plugged into the **Divide** Material Expression, currently set to 512, to something smaller as smaller numbers will allow the effect to happen closer to the object. 
+1. To see the effect in action, all you have to do is **hold down the Right Mouse Button** while inside the Material Preview window and then move the mouse forward and backwards. When you get far enough away from the preview mesh, you should see the Material transition from one texture to the other one. If you cannot see this happening, adjust the first **Constant** that is plugged into the **Divide** Material Expression, currently set to 512, to something smaller as smaller numbers will allow the effect to happen closer to the object. 
 	[OBJECT:EmbeddedVideo]
 	[PARAMLITERAL:width]
 	640
