#UE4doc #docs: Adjusted bold text in accordance with the Instructional Guide, corrected arrow color

diff --git a/Engine/Documentation/Source/Engine/Rendering/Materials/HowTo/Masking/Masking.INT.udn b/Engine/Documentation/Source/Engine/Rendering/Materials/HowTo/Masking/Masking.INT.udn
--- a/Engine/Documentation/Source/Engine/Rendering/Materials/HowTo/Masking/Masking.INT.udn
+++ b/Engine/Documentation/Source/Engine/Rendering/Materials/HowTo/Masking/Masking.INT.udn
@@ -25,25 +25,25 @@ Quite often, a mask will be contained within a single channel of another texture
 This is a good way to make use of unused channels and keep the number of textures being sampled in the Material to a minimum. 
 Technically, any channel of any texture can be thought of and used as a texture mask. 
 
-Here is an example of what the Texture Mask for the **SM_Chair** Static Mesh from the Starter Content looks like. 
+Here is an example of what the Texture Mask for the SM_Chair Static Mesh from the Starter Content looks like. 
 
 ![](TM_Chair_Mask_Texture.png)
 
 ## Creating a Texture Mask
 
 Creating a Texture Mask is something that can be accomplished in any 2D image manipulation program. 
-To accomplish this, all you have to do is paint the areas that you want to Mask with the color **White** leaving all the other areas **Black**. 
+To accomplish this, all you have to do is paint the areas that you want to Mask with the color White leaving all the other areas Black. 
 In the following image, we can see this in action. 
 
 ![](TM_Mask_Break_Down.png)
 
 On the right hand side, outlined in green, is what the Mask Texture looks like. 
-Notice how only a few areas are painted **White** while all other areas are painted **Black**. 
-This is because the we only want to affect the areas that are in **White**.
+Notice how only a few areas are painted White while all other areas are painted Black. 
+This is because the we only want to affect the areas that are in White.
 On the left hand side, you can see what areas this Mask Texture will mask out when applied to the chair mesh.
 
 [region:note]
-When creating a Texture Mask you should always use either **Black** or **White** and never color information. Using any other type of color could result in strange artifacts when using the Mask inside UE4.  
+When creating a Texture Mask you should always use either Black or White and never color information. Using any other type of color could result in strange artifacts when using the Mask inside UE4.  
 [/region]
 
 ## Exporting a Texture Mask
@@ -65,7 +65,7 @@ In the following section, we will cover some of the most common ways to use Text
 ### Emissive Mask
 
 One of the most common things to do with a Mask Texture is to use it to control the Emissive sections of a Material. This is usually done by first creating a Emissive Mask
- texture that uses **White** to define what sections of the Material should be Emissive and then multiplying that Mask texture by a color inside the Material Editor.
+ texture that uses White to define what sections of the Material should be Emissive and then multiplying that Mask texture by a color inside the Material Editor.
  This is generally done so that you have more control over the look and intensity of the Emissive effect. 
  In the following example, we can see Emissive Texture Masking in action. 
 
@@ -80,24 +80,24 @@ change the parameters inside of the Material instead of updating and re-importin
 ### Material Masking 
 
 The other most common thing to do with a Mask Texture is to store different kinds of Texture information into the individual R, G, B, and Alpha channels. 
-A perfect example of this technique can be found by looking at the **SM_Chair** Static Mesh that comes with the Starter Content. 
+A perfect example of this technique can be found by looking at the SM_Chair Static Mesh that comes with the Starter Content. 
 
 [region:tip]
-You can find **SM_Chair** and all the content that goes with it by first selecting the **Starter Content** folder in the **Content Browser** and then entering **chair** into the search box. 
+You can find SM_Chair and all the content that goes with it by first selecting the **Starter Content** folder in the **Content Browser** and then entering chair into the search box. 
 This will display all the content that is related to the chair. If you do not see the chair, it probably means that you did not include the Starter Content with your project.
 To fix this, you will need to either create a new project or try the [Migrating Assets](Engine/Content/Browser/UserGuide/Migrate) tool to move the chair content from another project 
 into this one.
 [/region]
 
-Opening up the chair Material, **M_Chair**, we can see a perfect example of Texture Mapping in action. 
+Opening up the chair Material, M_Chair, we can see a perfect example of Texture Mapping in action. 
 
 ![](TM_Chair_Material.png)
 
 In this Material, the Mask Texture,**T_Chair_M**, is used to control every aspect of how the chair looks. 
 From what sections should be metal or non-metal to what color those sections should have, the Mask Texture helps to define all of this.
 In the following image, you can see a breakdown of how each channel of the Mask Texture for SM_Chair works.
-On the **Left-Hand side** is what the Mask looks like as a Texture.
-The smaller images that you see on the **Right-Hand side** show each of the individual channels and what area those channels will effect on the chair Static Mesh. 
+On the Left-Hand side is what the Mask looks like as a Texture.
+The smaller images that you see on the Right-Hand side show each of the individual channels and what area those channels will effect on the chair Static Mesh. 
 
 ![](TM_Mask_Example.png)
 
@@ -124,7 +124,7 @@ When all put together, the chair looks like the following inside a UE4 level:
 Texture Masking can become a very powerful tool especially when combined with other tools inside UE4. In the following section, we will be going over some Tips & Trick for unleashing 
 the full power of Texture Masking in your work. 
 
-### Texture Masking and Material Instancing
+### Texture Masking and Material separate
 
 Using Texture Masking and Material Instancing together is a great way to add an almost endless amount of variety to any assets. 
 For example, you can use the Texture Mask to define which areas should have certain properties, like color, and then use different Material Instances to change each of these properties giving you an almost endless variety of options.
@@ -136,7 +136,7 @@ For example, if we make a Material Instance out of the SM_Chair Material, we can
 
 ### Masking and Channel Artifacts 
 
-Because of a quirk with DirectX, the **Green Channel** of a texture will often offer the best compression. If any of your masks suffer greatly from compression artifacts, first try placing the information into the **Green Channel** to see if that helps. If that does not fix the issue, then try using the Alpha channel to store the information. 
+Because of a quirk with DirectX, the Green Channel of a texture will often offer the best compression. If any of your masks suffer greatly from compression artifacts, first try placing the information into the Green Channel to see if that helps. If that does not fix the issue, then try using the Alpha channel to store the information. 
 
 [region:tip]
 Be careful when trying to use the Alpha channel of a texture to store Mask Textures. Adding an Alpha channel to a texture will greatly increase the memory footprint of that texture and if 
@@ -146,17 +146,17 @@ you do this enough, you could lose all the saving you gained by packing differen
 
 ### sRGB and Mask Textures 
 
-When packing multiple Mask Textures into a single texture, you should **Disable** sRGB as your masks should not be Gamma corrected.
+When packing multiple Mask Textures into a single texture, you should Disable sRGB as your masks should not be Gamma corrected.
 
 ![](TM_Disable_sRGB.png)
 
 
-However when you do this, you need to make sure that you adjust your **2D Texture Samplers** node type as they expect the textures you are inputting to have sRGB enabled. 
+However when you do this, you need to make sure that you adjust your 2D Texture Samplers node type as they expect the textures you are inputting to have sRGB enabled. 
 If you do not change the Sampler Type, your Material will fail to compile and the following message will be displayed in the Stats log.
 
 ![](TM_sRGB_Sampler_Node_Error.png)
 
-To fix this error, all you need to do is to change the **Sampler Type** from the default of **Color** to **Linear Color** and the warning will go away. 
+To fix this error, all you need to do is to change the Sampler Type from the default of Color to Linear Color and the warning will go away. 
 However, for good measure, re-compile the Material to make sure that the changes took. Once completed, the warning will go away and the Material will now be ready to be used.  
 
 [OBJECT:EmbeddedVideo]
