#UE4doc #docs: Adjusted bold text in accordance with the Instructional Guide

diff --git a/Engine/Documentation/Source/Engine/Rendering/Materials/HowTo/Human_Skin/HumanSkin.INT.udn b/Engine/Documentation/Source/Engine/Rendering/Materials/HowTo/Human_Skin/HumanSkin.INT.udn
--- a/Engine/Documentation/Source/Engine/Rendering/Materials/HowTo/Human_Skin/HumanSkin.INT.udn
+++ b/Engine/Documentation/Source/Engine/Rendering/Materials/HowTo/Human_Skin/HumanSkin.INT.udn
@@ -18,7 +18,7 @@ Related: Resources/ContentExamples/MaterialNodes/1_7
 [/REGION]
 
 Getting Human skin to look correct when rendered is a very complex topic with no single perfect solution for getting correct results. 
-Because of this, please treat the following information _**as a starting point**_ for your work as many things are needed in order to achieve a realistic looking image. 
+Because of this, please treat the following information _as a starting point_ for your work as many things are needed in order to achieve a realistic looking image. 
 This is because our eyes are very well trained at detecting subtle details in human faces. 
 Getting it right is challenging but once the tools are mastered, they can be applied in many other cases.
 
@@ -49,11 +49,11 @@ Here is a list of some things that you need to be aware of before you begin to c
   Even having Eye Adaptation enabled can make a huge difference in how things look in the end. 
 * To avoid shadow artifacts will require extra effort because you will need to tweak things like shadow bias or use a more focused spot light cone for better lighting resolution.
 * When tweaking, it might be useful to zoom in by adjusting the editor FOV as there will be no clipping through geometry and less distortions. 
-* The editor Field Of View(FOV) can be adjusted by clicking on the **Drop Down Arrow** that is next to the icon that says **Perspective** and then adjusting the FOC value using the mouse. 
+* The editor Field Of View(FOV) can be adjusted by clicking on the Drop Down Arrow that is next to the icon that says Perspective and then adjusting the FOC value using the mouse. 
 	
 	![](HS_Change_FOV.png)
 
-* When viewing your work, you might notice that the streaming system does not take the zoom into account and could blur your texture when you do not want to. This was an intentional design decision made for 3D games but this could change at some point in the future. To prevent this from happening, you can change the texture **Streaming Distance Multiplier**  on the mesh instance.
+* When viewing your work, you might notice that the streaming system does not take the zoom into account and could blur your texture when you do not want to. This was an intentional design decision made for 3D games but this could change at some point in the future. To prevent this from happening, you can change the texture Streaming Distance Multiplier  on the mesh instance.
 
 ## Material 
 
@@ -100,7 +100,7 @@ The number 0.35 was used as it was derived from real world Specular values that
 
 In order for your Normal Map to work correctly for skin, you need to make sure that your Normal map meets the following criteria. 
 * The Normal Map that is used for the face should only contain details about wrinkles. Pore details should be done with a tiling texture as you can get better results this way.
-*  When combining the object's Normal map with the Detail Normal map, make sure that you use the Material Function **BlendAngleCorrectedNormals** to combine as you will yield better results. 
+*  When combining the object's Normal map with the Detail Normal map, make sure that you use the Material Function BlendAngleCorrectedNormals to combine as you will yield better results. 
    It is also very useful to expose parameters to adjust the mixing of the Normal maps. You can take the detail normal map value, add some scalar value to the z component and normalize. 
    This way you can transition to a flat detail Normal Map from your regular one. 
 * As skin has different shading properties in different areas, it is best to combine the Roughness with an object mapped (non-tiled) texture. 
