#UE4 Docs: Adding skill level, version and prereq meta data to all Gear VR files.

diff --git a/Engine/Documentation/Source/Platforms/GearVR/BestPractices/GearVRBestPractices.INT.udn b/Engine/Documentation/Source/Platforms/GearVR/BestPractices/GearVRBestPractices.INT.udn
--- a/Engine/Documentation/Source/Platforms/GearVR/BestPractices/GearVRBestPractices.INT.udn
+++ b/Engine/Documentation/Source/Platforms/GearVR/BestPractices/GearVRBestPractices.INT.udn
@@ -11,10 +11,11 @@ Related: Platforms/Mobile/Textures
 Related: Platforms/Mobile/Materials
 Related: Engine/Rendering/Materials/MaterialProperties/BlendModes
 Related: Engine/Rendering/Materials/CustomizedUVs
+Skilllevel: Intermediate
+Version: 4.8
 
 [TOC(start:2 end:3)]
 
-## Overview
 When developing content for Virtual Reality (VR) projects, certain steps must be taken to ensure that the user will have the best VR experience possible.
 In the following Best Practices guide we will cover all you need to know to help guarantee that the user has the best possible experience when viewing your project 
 using the Samsung Gear Virtual Reality(VR) Headset and Samsung Galaxy Note 4 Smartphone.
diff --git a/Engine/Documentation/Source/Platforms/GearVR/Debugging/GearVRDebugging.INT.udn b/Engine/Documentation/Source/Platforms/GearVR/Debugging/GearVRDebugging.INT.udn
--- a/Engine/Documentation/Source/Platforms/GearVR/Debugging/GearVRDebugging.INT.udn
+++ b/Engine/Documentation/Source/Platforms/GearVR/Debugging/GearVRDebugging.INT.udn
@@ -11,6 +11,8 @@ Related: Platforms/Mobile/Textures
 Related: Platforms/Mobile/Materials
 Related: Engine/Rendering/Materials/MaterialProperties/BlendModes
 Related: Engine/Rendering/Materials/CustomizedUVs
+Skilllevel: Intermediate
+Version: 4.8
 
 [TOC(start:2 end:3)]
 
diff --git a/Engine/Documentation/Source/Platforms/GearVR/GearVR.INT.udn b/Engine/Documentation/Source/Platforms/GearVR/GearVR.INT.udn
--- a/Engine/Documentation/Source/Platforms/GearVR/GearVR.INT.udn
+++ b/Engine/Documentation/Source/Platforms/GearVR/GearVR.INT.udn
@@ -3,6 +3,10 @@ Description:Instructions on how to get the Samsung Gear VR Headset to work with
 Availability:Public
 Navigation:platform
 Platform:Gear VR
+Skilllevel: Intermediate
+Version: 4.8
+prereq: Engine\Blueprints\GettingStarted
+
 
 [OBJECT:TopicButtonList]
 	[PARAM:icon]
diff --git a/Engine/Documentation/Source/Platforms/GearVR/GlobalMenu/1/GlobalMenu_HudCreation.INT.udn b/Engine/Documentation/Source/Platforms/GearVR/GlobalMenu/1/GlobalMenu_HudCreation.INT.udn
--- a/Engine/Documentation/Source/Platforms/GearVR/GlobalMenu/1/GlobalMenu_HudCreation.INT.udn
+++ b/Engine/Documentation/Source/Platforms/GearVR/GlobalMenu/1/GlobalMenu_HudCreation.INT.udn
@@ -4,6 +4,8 @@ Crumbs:%ROOT%, Platforms/GearVR
 Description:Instructions on how to create a new HUD to use with your Gear VR project.
 Navigation:platform
 Platform:Gear VR
+Skilllevel: Intermediate
+Version: 4.8
 
 [VAR:Steps]
 [OBJECT:Navigation]
diff --git a/Engine/Documentation/Source/Platforms/GearVR/GlobalMenu/2/GlobalMenu_HUDAssignment.INT.udn b/Engine/Documentation/Source/Platforms/GearVR/GlobalMenu/2/GlobalMenu_HUDAssignment.INT.udn
--- a/Engine/Documentation/Source/Platforms/GearVR/GlobalMenu/2/GlobalMenu_HUDAssignment.INT.udn
+++ b/Engine/Documentation/Source/Platforms/GearVR/GlobalMenu/2/GlobalMenu_HUDAssignment.INT.udn
@@ -4,6 +4,8 @@ Crumbs:%ROOT%, Platforms/GearVR
 Description: How to assign the new HUD to your project
 Navigation:platform
 Platform:Gear VR
+Skilllevel: Intermediate
+Version: 4.8
 
 [VAR:Steps]
 [OBJECT:Navigation]
diff --git a/Engine/Documentation/Source/Platforms/GearVR/GlobalMenu/3/GlobalMenu_HUDBPSetup.INT.udn b/Engine/Documentation/Source/Platforms/GearVR/GlobalMenu/3/GlobalMenu_HUDBPSetup.INT.udn
--- a/Engine/Documentation/Source/Platforms/GearVR/GlobalMenu/3/GlobalMenu_HUDBPSetup.INT.udn
+++ b/Engine/Documentation/Source/Platforms/GearVR/GlobalMenu/3/GlobalMenu_HUDBPSetup.INT.udn
@@ -4,6 +4,8 @@ Crumbs:%ROOT%, Platforms/GearVR
 Description: How to setup a Blueprint to display the Gear VR Global Menu
 Navigation:platform
 Platform:Gear VR
+Skilllevel: Intermediate
+Version: 4.8
 
 [VAR:Steps]
 [OBJECT:Navigation]
diff --git a/Engine/Documentation/Source/Platforms/GearVR/GlobalMenu/GearVRGlobalMenu.INT.udn b/Engine/Documentation/Source/Platforms/GearVR/GlobalMenu/GearVRGlobalMenu.INT.udn
--- a/Engine/Documentation/Source/Platforms/GearVR/GlobalMenu/GearVRGlobalMenu.INT.udn
+++ b/Engine/Documentation/Source/Platforms/GearVR/GlobalMenu/GearVRGlobalMenu.INT.udn
@@ -4,6 +4,9 @@ Crumbs:%ROOT%, Platforms
 Description:How to setup your Gear VR project to access the global menu.
 Navigation:platform
 Platform:Gear VR
+Skilllevel: Intermediate
+Version: 4.8
+prereq: Engine\Blueprints\GettingStarted
 
 In order for your Gear VR project to be accepted to the Oculus store there are certain requirements that your project must adhere to. 
 One of these requirements is that at any point while someone is using your project they need to be able to access the Gear VR Global Menu when they hold down the back button for more than 0.75 seconds. 
diff --git a/Engine/Documentation/Source/Platforms/GearVR/Prerequisites/GearVRPrerequisites.INT.udn b/Engine/Documentation/Source/Platforms/GearVR/Prerequisites/GearVRPrerequisites.INT.udn
--- a/Engine/Documentation/Source/Platforms/GearVR/Prerequisites/GearVRPrerequisites.INT.udn
+++ b/Engine/Documentation/Source/Platforms/GearVR/Prerequisites/GearVRPrerequisites.INT.udn
@@ -4,6 +4,9 @@ Availability:Public
 Crumbs:
 Navigation:platform
 Platform:Gear VR
+Skilllevel: Intermediate
+Version: 4.8
+prereq: Engine\Blueprints\GettingStarted
 
 ![](Platforms\GearVR\QuickStart\GVR_Headset_00.png)
 
@@ -11,8 +14,6 @@ Platform:Gear VR
 Developing for the Gear VR Headset requires **Unreal Engine 4.7.4** or later.
 [/REGION]
 
-## Overview 
-
 With the 4.7 release of Unreal Engine 4 (UE4) you can now build Virtual Reality(VR) projects for the Samsung Gear VR Headset that can transform your Samsung Galaxy Note 4, Galaxy S6, and Galaxy S6 Edge Smartphone into an immersive VR device. 
 In the following sections we will go over everything you need to know to get your Samsung Galaxy Note 4 Smartphone properly setup so that it can work with the Samsung Gear VR headset.
 
diff --git a/Engine/Documentation/Source/Platforms/GearVR/QuickStart/1/GearVRQuickStart_1.INT.udn b/Engine/Documentation/Source/Platforms/GearVR/QuickStart/1/GearVRQuickStart_1.INT.udn
--- a/Engine/Documentation/Source/Platforms/GearVR/QuickStart/1/GearVRQuickStart_1.INT.udn
+++ b/Engine/Documentation/Source/Platforms/GearVR/QuickStart/1/GearVRQuickStart_1.INT.udn
@@ -3,6 +3,8 @@ Title:1. Gear VR Project Creation
 Crumbs:%ROOT%, Platforms/GearVR
 Description:How to create a UE4 C++ based project that can be used with the Gear VR Headset.
 Navigation:platform
+Skilllevel: Intermediate
+Version: 4.8
 
 
 [VAR:Steps]
diff --git a/Engine/Documentation/Source/Platforms/GearVR/QuickStart/2/GearVRQuickStart_2.INT.udn b/Engine/Documentation/Source/Platforms/GearVR/QuickStart/2/GearVRQuickStart_2.INT.udn
--- a/Engine/Documentation/Source/Platforms/GearVR/QuickStart/2/GearVRQuickStart_2.INT.udn
+++ b/Engine/Documentation/Source/Platforms/GearVR/QuickStart/2/GearVRQuickStart_2.INT.udn
@@ -3,6 +3,8 @@ Title:2. Gear VR Project Setup
 Crumbs:%ROOT%, Platforms/GearVR
 Description:How to setup your C++ based UE4 project to be used with the Gear VR Headset.
 Navigation:platform
+Skilllevel: Intermediate
+Version: 4.8
 
 [VAR:Steps]
 [OBJECT:Navigation]
diff --git a/Engine/Documentation/Source/Platforms/GearVR/QuickStart/3/GearVRQuickStart_3.INT.udn b/Engine/Documentation/Source/Platforms/GearVR/QuickStart/3/GearVRQuickStart_3.INT.udn
--- a/Engine/Documentation/Source/Platforms/GearVR/QuickStart/3/GearVRQuickStart_3.INT.udn
+++ b/Engine/Documentation/Source/Platforms/GearVR/QuickStart/3/GearVRQuickStart_3.INT.udn
@@ -7,8 +7,8 @@ Platform:Gear VR
 Related: Platforms\GearVR\Prerequisites
 Related: Platforms\GearVR\BestPractices
 Related: Platforms\GearVR\Debugging
-
-
+Skilllevel: Intermediate
+Version: 4.8
 
 [VAR:Steps]
 [OBJECT:Navigation]
diff --git a/Engine/Documentation/Source/Platforms/GearVR/QuickStart/GearVRQuickStart.INT.udn b/Engine/Documentation/Source/Platforms/GearVR/QuickStart/GearVRQuickStart.INT.udn
--- a/Engine/Documentation/Source/Platforms/GearVR/QuickStart/GearVRQuickStart.INT.udn
+++ b/Engine/Documentation/Source/Platforms/GearVR/QuickStart/GearVRQuickStart.INT.udn
@@ -4,6 +4,8 @@ Crumbs:%ROOT%, Platforms/GearVR
 Description:Instructions on how to quickly setup and deploy projects for use with the Samsung Gear VR Headset.
 Navigation:platform
 Platform:Gear VR
+Skilllevel: Intermediate
+Version: 4.8
 
 The Gear VR Quick Start will walk you through how to setup your development PC and Samsung Galaxy Note 4 Smartphone to work with the Samsung Gear VR headset.
 
