#UE4 Docs: adding gameplay debugger page

diff --git a/Engine/Documentation/Source/Gameplay/Tools/GameplayDebugger/GameplayDebugger.INT.udn b/Engine/Documentation/Source/Gameplay/Tools/GameplayDebugger/GameplayDebugger.INT.udn
new file mode 100644
--- /dev/null
+++ b/Engine/Documentation/Source/Gameplay/Tools/GameplayDebugger/GameplayDebugger.INT.udn
@@ -0,0 +1,287 @@
+Title:Gameplay Debugger
+Availability:Public
+Crumbs:
+Description: Tool that enables analyzing realtime gameplay data at runtime.
+
+The Gameplay Debugger is useful for watching real-time data at runtime, even on clients in networked games using replication. It works in Play In Editor, Simulate In Editor, and standalone game sessions, and all of the data is displayed overlaid on the game viewport. The system provides a framework that can be extended to enable debugging of game-specific data.
+
+Engine implementation can show:
+
+* Basic data from Pawn
+* Basic data from AIController
+* Information about Behavior Tree and Blackboard data
+* Information about executed Environment Queries (EQS)
+* Information from perception system
+* Navmesh around player or selected Pawn with all details like links or areas
+
+Usually there is a lot of data so GDT use categories to limit amount of information on screen. Data from only enabled categories is replicated which saves replication channel's bandwidth. There are 5 default categories and 5 categories to use by projects:
+
+![image alt text](image_1.png)
+
+* Navmesh
+* Basic
+* Behavior Tree
+* EQS
+* Perception
+* and 5 categories to use by projects
+
+
+Existing categories can be extended to show more game specific data too.
+
+Below we have screen-shots captured on client with enabled few categories: Basic, EQS, NavMesh and Behavior Tree:
+
+![image alt text](image_2.png)
+
+Gameplay Debugger can be activated with ' (Apostrophe) key by default or with *EnableGDT* cheat. Key binding is set in file and can be easily changed. To select enemy for debugging press ' key while pointing that enemy on screen. Numeric keyboard is used to switch between visible categories. GameplayDebugger module has to be added to project's dependency modules to be activated and usable. 
+
+## Editor - Working with Gameplay Debugger 
+
+While working in editor GDT can be used in PIE or in Simulate. Binded key or *EnableGDT *cheat can be used to activate GDT in PIE. Simulate mode is little different from PIE so to activate this debug tool it's needed to enable "*Debug AI*" show flag. There is an option to change visible categories in Simulate too. *GameplayDebuggingReplicator *actor is used for that. This actor can be find in Scene Outliner and his properties are used to control GDT:
+
+![image alt text](image_3.jpg)
+
+
+## Basic extension
+
+Gameplay Debugger can be only extended with C++ code. For Blueprint projects it can only be used as is - to display basic debug information for now. It's quite easy to extend Gameplay Debugger, to collect and display game specific data. Custom classes inherited from *UGameplayDebuggingComponent *class and *AGameplayDebuggingHUDComponent *class are needed for that. FIrst class is used to collect and eventually replicated data and second class is used to display all collected data on screen.
+
+Below is a simple class which collects game specific data:
+
+
+GDTComponent.h
+
+	// Copyright 1998-2014 Epic Games, Inc. All Rights Reserved.
+	#pragma once
+	#include "GameplayDebuggingComponent.h"
+	#include "GDTComponent.generated.h"
+
+	UCLASS()
+	class UGDTComponent : public UGameplayDebuggingComponent
+	{
+	public:
+		GENERATED_UCLASS_BODY()
+		virtual void CollectBasicData() override;
+
+		UPROPERTY(Replicated)
+		float TestData; //custom data replicated to clients
+	};
+
+GDTComponent.cpp
+
+	// Copyright 1998-2014 Epic Games, Inc. All Rights Reserved.
+	#include "MyGameProject.h"
+	#include "GameplayDebuggingComponent.h"
+	#include "GDTComponent.h"
+
+	UGDTComponent::UGDTComponent(const class FPostConstructInitializeProperties& PCIP) : Super(PCIP) { }
+
+	void UGDTComponent::GetLifetimeReplicatedProps(TArray<FLifetimeProperty> &OutLifetimeProps) const
+	{
+		Super::GetLifetimeReplicatedProps( OutLifetimeProps );
+	#if !(UE_BUILD_SHIPPING || UE_BUILD_TEST)
+		DOREPLIFETIME( UGDTComponent, TestData);
+	#endif
+	}
+
+	void UGDTComponent::CollectBasicData()
+	{
+		Super::CollectBasicData();
+	#if !(UE_BUILD_SHIPPING || UE_BUILD_TEST)
+		TestData= FMath::RandRange(2.75, 8.25); //collect data and store it
+	#endif
+	}
+
+Next class is used to display new data on screen:
+
+
+GDTHUDComponent.h
+
+	// Copyright 1998-2014 Epic Games, Inc. All Rights Reserved.
+	#pragma once
+
+	#include "GameplayDebuggingHUDComponent.h"
+	#include "GDTHUDComponent.generated.h"
+
+	UCLASS(notplaceable)
+	class AGDTHUDComponent: public AGameplayDebuggingHUDComponent
+	{
+		GENERATED_UCLASS_BODY()
+	protected:
+		virtual void DrawBasicData(APlayerController* PC, class UGameplayDebuggingComponent *DebugComponent) override;
+	};
+
+GDTHUDComponent.cpp
+
+	// Copyright 1998-2014 Epic Games, Inc. All Rights Reserved.
+	#include "MyGameProject.h"
+	#include "GDTComponent.h"
+	#include "GDTHUDComponent.h"
+
+	AGDTHUDComponent::AGDTHUDComponent(const class FPostConstructInitializeProperties& PCIP)
+		: Super(PCIP)
+	{
+	}
+	void AGDTHUDComponent::DrawBasicData(APlayerController* PC, class UGameplayDebuggingComponent *DebugComponent)
+	{
+		Super::DrawBasicData(PC, DebugComponent);
+	#if !(UE_BUILD_SHIPPING || UE_BUILD_TEST)
+		const UGDTComponent* MyComp = Cast<UGDTComponent>(DebugComponent);
+		if (MyComp)
+		{
+			PrintString(DefaultContext, FString::Printf(TEXT("{white}Test data: {red}%f\n"), MyComp->TestData));
+		}
+	#endif
+	}
+
+Gameplay Debugger needs to know about new classes and that information can be set in DefaultEngine.ini configuration file:
+
+
+DefaultEngine.ini
+
+	[/Script/GameplayDebugger.GameplayDebuggingReplicator]
+	DebugComponentClassName="/Script/MyGameProject.GDTComponent"
+	DebugComponentHUDClassName="/Script/MyGameProject.GDTHUDComponent" |
+
+## Custom categories
+
+Few more steps are needed to add project specific category to Gameplay Debugger. 
+
+Let's extend GDTComponent class:
+
+
+GDTComponent.h
+
+	// Copyright 1998-2014 Epic Games, Inc. All Rights Reserved.
+	#pragma once
+	#include "GameplayDebuggingComponent.h"
+	#include "GDTComponent.generated.h"
+
+	UCLASS()
+	class UGDTComponent : public UGameplayDebuggingComponent
+	{
+	public:
+		GENERATED_UCLASS_BODY()
+	protected:
+		virtual void CollectDataToReplicate(bool bCollectExtendedData) override;
+		void CollectCustomData();
+	public:
+		UPROPERTY(Replicated)
+		float TestData; //custom data replicated to clients
+	};
+
+GDTComponent.cpp
+
+	// Copyright 1998-2014 Epic Games, Inc. All Rights Reserved.
+	#include "MyGameProject.h"
+	#include "GameplayDebuggingComponent.h"
+	#include "GDTComponent.h"
+
+	UGDTComponent::UGDTComponent(const class FPostConstructInitializeProperties& PCIP) : Super(PCIP) { }
+
+	void UGDTComponent::GetLifetimeReplicatedProps(TArray<FLifetimeProperty> &OutLifetimeProps) const
+	{
+		Super::GetLifetimeReplicatedProps( OutLifetimeProps );
+	#if !(UE_BUILD_SHIPPING || UE_BUILD_TEST)
+		DOREPLIFETIME( UGDTComponent, TestData);
+	#endif
+	}
+
+	void UGDTComponent::CollectCustomData()
+	{
+		Super::CollectBasicData();
+	#if !(UE_BUILD_SHIPPING || UE_BUILD_TEST)
+		TestData= FMath::RandRange(2.75, 8.25); //collect data and store it
+	#endif
+	}
+
+	void UGDTComponent::CollectDataToReplicate(bool bCollectExtendedData)
+	{
+		Super::CollectDataToReplicate(bCollectExtendedData);
+	#if !(UE_BUILD_SHIPPING || UE_BUILD_TEST)
+		if (ShouldReplicateData(EAIDebugDrawDataView::GameView1))
+		{ 
+			CollectCustomData();
+			if (bCollectExtendedData)
+			{
+				// collect additional data for selected Pawn/AIController 
+			}
+		}
+	#endif
+	}
+
+Extended HUD component, to display data in new view:
+
+
+GDTHUDComponent.h
+
+	// Copyright 1998-2014 Epic Games, Inc. All Rights Reserved.
+	#pragma once
+
+	#include "GameplayDebuggingHUDComponent.h"
+	#include "GDTHUDComponent.generated.h"
+
+	UCLASS(notplaceable)
+	class AGDTHUDComponent: public AGameplayDebuggingHUDComponent
+	{
+		GENERATED_UCLASS_BODY()
+	protected:
+		virtual void DrawGameSpecificView(APlayerController* PC, class UGameplayDebuggingComponent *DebugComponent) override;
+		virtual void GetKeyboardDesc(TArray<FDebugCategoryView>& Categories) override;
+		void DrawCustomData(APlayerController* PC, class UGameplayDebuggingComponent *DebugComponent);
+	};
+
+GDTHUDComponent.cpp
+
+	// Copyright 1998-2014 Epic Games, Inc. All Rights Reserved.
+	#include "MyGameProject.h"
+	#include "GDTComponent.h"
+	#include "GDTHUDComponent.h"
+
+	AGDTHUDComponent::AGDTHUDComponent(const class FPostConstructInitializeProperties& PCIP)
+		: Super(PCIP)
+	{
+	}
+	void AGDTHUDComponent::DrawCustomData(APlayerController* PC, class UGameplayDebuggingComponent *DebugComponent)
+	{
+	#if !(UE_BUILD_SHIPPING || UE_BUILD_TEST)
+		const UGDTComponent* MyComp = Cast<UGDTComponent>(DebugComponent);
+		if (MyComp)
+		{
+			PrintString(DefaultContext, FString::Printf(TEXT("{white}Test data: {red}%f\n"), MyComp->TestData));
+		}
+	#endif
+	}
+	void AGDTHUDComponent::GetKeyboardDesc(TArray<FDebugCategoryView>& Categories)
+	{
+		Super::GetKeyboardDesc(Categories);
+	#if !(UE_BUILD_SHIPPING || UE_BUILD_TEST)
+		Categories.Add(FDebugCategoryView(EAIDebugDrawDataView::GameView1, TEXT("MyData")));
+	#endif
+	}
+	void AGDTHUDComponent::DrawGameSpecificView(APlayerController* PC, class UGameplayDebuggingComponent *InDebugComponent)
+	{
+		 Super::DrawGameSpecificView(PC, InDebugComponent);
+	#if !(UE_BUILD_SHIPPING || UE_BUILD_TEST)
+		if (InDebugComponent && GameplayDebuggerSettings(GetDebuggingReplicator()).CheckFlag(EAIDebugDrawDataView::GameView1))
+		{
+			PrintString(DefaultContext, FColor::Green, TEXT("\nMY GAME DATA\n"));
+			 DrawCustomData(PC, InDebugComponent);
+		}
+	#endif
+	}
+
+New category is ready and can be used to debug project specific data:
+
+![image alt text](image_4.jpg)
+
+To draw debug information in colors *PrintString *function can use tags in string, to change active color. It makes easier to draw string with different colors:
+
+	void PrintString(FPrintContext& Context, const FString& InString );
+	void PrintString(FPrintContext& Context, const FColor& InColor, const FString& InString );
+	PrintString(DefaultContext, FColor::Green, TEXT("Whole text in green"));
+	PrintString(DefaultContext, TEXT("String {green}in green, {red}in red {white}or {R=0,G=0,B=255,A=255}in blue"));
+
+Last PrintString function generates string with 4 different colors:
+
+![image alt text](image_5.jpg)
+
