#UE4 Docs: adding Network Profiler page

diff --git a/Engine/Documentation/Source/Gameplay/Tools/NetworkProfiler/NetworkProfiler.INT.udn b/Engine/Documentation/Source/Gameplay/Tools/NetworkProfiler/NetworkProfiler.INT.udn
new file mode 100644
--- /dev/null
+++ b/Engine/Documentation/Source/Gameplay/Tools/NetworkProfiler/NetworkProfiler.INT.udn
@@ -0,0 +1,76 @@
+Title:Network Profiler
+Availability:Public
+Crumbs:
+Description:Tool for displaying network traffic and performance information captured at runtime.
+
+The Network Profiler is a standalone tool that can be used to display network traffic and performance information that can be recorded by the engine as a game runs. This is a good way to identify areas of a multiplayer game that use unusually high bandwidth, since you can see how much bandwidth individual actors, RPCs, and properties contribute to the total.
+
+## Recording a profiling session
+
+Before you can use the network profiler, you must record some data for it to analyze. To do this, you need to use a version of the engine with stat tracking enabled - usually this means either a debug build, or for non-debug configurations, an editor build (specifically, the engine should have been compiled with the STATS macro defined to a nonzero value).
+
+You can pass the "networkprofiler=true" command-line argument to record as soon as the engine starts up, but you can also use the following console commands at runtime to control the recording of network profiler data:
+
+* netprofile: toggles recording on and off
+
+* netprofile enable: start recording if not already recording
+
+* netprofile disable: stop recording if currently recording
+
+The data file will be saved to <project directory>/Saved/Profiling/<project name>-<timestamp>.nprof. If you see a file named NetworkProfiling.tmp in the same folder, this is just the temporary file the engine uses while collecting data. When the current profiling session is stopped, this file will be renamed according to the scheme above and is then ready to be opened in the standalone tool.
+
+## Viewing a profile session in the Network Profiler application
+
+The network profiler is a standalone application located at <UE install directory>/Engine/Binaries/DotNET/NetworkProfiler.exe.
+
+At the top of the window, click the "Open File" button to select a .nprof file and view the data contained within.![image alt text](image_0.png)
+
+### Chart, Filters, Details tab
+
+By default, the "Chart, Filters, Details" tab is selected. The different areas of this tab are described in the following sections.
+
+![image alt text](image_1.png)
+
+1. [Chart View](#chartview)
+1. [Stats List](#statslist)
+1. [Summary view](#summaryview)
+1. [Frame details](#framedetails)
+1. [Filters](#filters)
+1. [Performance view](#performanceview)
+
+#### Chart View
+
+This is the main chart view. Similar to charts in other profiling applications, it shows a graph of the enabled stats over time. You can click on a particular frame to populate the other views with the data associated with that frame, or click and drag across the graph to zoom in to that area and view a summary for that time span.
+
+#### Stats List
+
+This is a list of all the stats tracked by the network profiler. Toggling the checkboxes will show or hide that particular stat in the graph. For most stats, you can choose to show the raw count, the count per second, the number of bytes, or the number of bytes per second.
+
+#### Summary View
+
+This is a summary of the data for the time span that has been selected in the chart. It shows raw totals as well as per-second information. Note that the "Outgoing bandwidth" will be greater than the “Game socket send size” because “Outgoing bandwidth” includes an approximation of the IP and UDP header sizes as well.
+
+#### Frame Details
+
+This view breaks down the data that was sent on the currently selected frame. You can see which actors, properties, and remote functions were sent, as well as how many bytes each one contributed. Note that replicated actors and RPCs are both included in bunches, so the "NumBytes" field in “SEND BUNCH” sections includes the bytes from “SEND RPC” and “REPLICATE ACTOR” sections. Ultimately, the “SOCKET SEND TO” entries will tell you exactly how many bytes the engine sent to the outgoing socket. Also note that the standard UDP packet overhead is not included here.
+
+#### Filters
+
+You can enter a filter for actors, properties, and RPCs, and after clicking the "Apply Filters" button, the graph will update to only show data that is relevant for actors, properties, and RPCs that contain the text entered in the filter fields. Selecting a frame or range of frames in the graph after a filter has been applied will also filter the data in the summary and frame details views accordingly.
+
+#### Performance View
+
+While a single frame is selected in the graph, this view will show a list of actor types that were replicated during that frame, sorted by the CPU time it took to replicate them. If an actor has associated replicated properties, expanding the tree view for the actor will show each property. For actors, the first column of numbers is the CPU replication time for that actor, in milliseconds. The second column is the number of bytes used to represent the actor in the bunch, and the third column is the number of actors of the given type that were replicated during the current frame. For property rows, the first column is the number of bytes used for the property, and the second column shows how many properties with the given name were replicated for the actor.
+
+### Actors tab, Properties tab, and RPCs tab
+
+![image alt text](image_2.png)
+
+Each of these tabs show a summary of the respective replication data for the entire profiling session. Note that the data displayed in these tabs is not affected by the current frame or time span selected in the chart on the "Chart, Filters, Details" tab.
+
+The tabs show a list of actors, properties, or RPCs, respectively. The "Total Size" and “Average Size” columns show how much bandwidth was needed for a particular item, and the “Count” column represents how many times a particular item was replicated during the profiling session. You can toggle between sorting by a column in ascending or descending order by clicking on the column heading.
+
+## Server and client considerations
+
+You can record profiling data from both clients and servers. However, be aware that since only servers replicate actors and their properties, when viewing a profile recorded on a client, you will only see detailed data for RPCs that were sent from that client.
+
