#UE4 Docs: linked gameplay tools pages

 - updated sidebar

diff --git a/Engine/Documentation/Source/Gameplay/Gameplay.INT.udn b/Engine/Documentation/Source/Gameplay/Gameplay.INT.udn
--- a/Engine/Documentation/Source/Gameplay/Gameplay.INT.udn
+++ b/Engine/Documentation/Source/Gameplay/Gameplay.INT.udn
@@ -104,6 +104,22 @@ gameplay
 	[/PARAM]
 [/OBJECT]
 
+[OBJECT:TopicButtonList]
+	[PARAM:icon]
+		![](%ROOT%/profile_icon.png)(convert:false)
+	[/PARAM]
+	[PARAM:title]
+		Gameplay Tools
+	[/PARAM]
+	[PARAM:description]
+	[/PARAM]
+	[PARAM:links]
+		* [](Gameplay/Tools/GameplayDebugger "%Gameplay/Tools/GameplayDebugger:description%")
+		* [](Gameplay/Tools/NetworkProfiler "%Gameplay/Tools/NetworkProfiler:description%")
+		* [](Gameplay/Tools/VisualLogger "%Gameplay/Tools/VisualLogger:description%")
+	[/PARAM]
+[/OBJECT]
+
 [OBJECT:TopicButtonList]
 	[PARAM:icon]
 		![](%ROOT%/gameplay_icon.png)(convert:false)
