#UE4doc #docs: Removed extraneous bolding as according to the instructional guide. Altered images to match the instructional guide: green -> yellow.

diff --git a/Engine/Documentation/Source/Engine/Deployment/Launching/LaunchingToDevices.INT.udn b/Engine/Documentation/Source/Engine/Deployment/Launching/LaunchingToDevices.INT.udn
--- a/Engine/Documentation/Source/Engine/Deployment/Launching/LaunchingToDevices.INT.udn
+++ b/Engine/Documentation/Source/Engine/Deployment/Launching/LaunchingToDevices.INT.udn
@@ -24,7 +24,7 @@ This method will only install the current map on the device for fast iteration,
 [/REGION]
 
 ### Generic UE4Game App
-If you are making a content-only project, the executable that runs for **One Click Deploy** is actually the generic "UE4Game" game (since it can be used with any content-only project). When it runs on a 
+If you are making a content-only project, the executable that runs for One Click Deploy is actually the generic "UE4Game" game (since it can be used with any content-only project). When it runs on a 
 device, the icon that gets installed is just labeled as UE4Game. However, when we package a game, the name of your project is used in the final packaged build. When you install the packaged build,
  the icon will have your name (as well as your icons if you updated the default icons).
 
