#UE4doc #docs: Removed extraneous bolding as according to the instructional guide. Altered images to match the instructional guide: green -> yellow.

diff --git a/Engine/Documentation/Source/Engine/Deployment/UnrealFrontend/UnrealFrontend.INT.udn b/Engine/Documentation/Source/Engine/Deployment/UnrealFrontend/UnrealFrontend.INT.udn
--- a/Engine/Documentation/Source/Engine/Deployment/UnrealFrontend/UnrealFrontend.INT.udn
+++ b/Engine/Documentation/Source/Engine/Deployment/UnrealFrontend/UnrealFrontend.INT.udn
@@ -181,7 +181,7 @@ Console tab, the Automation tab and the upcoming Profiler.
 
 ![](session_browser.png)(convert:false)
 
-Like the **Device Browser**, the **Session Browser** is still limited in functionally. It currently consists of a split list 
+Like the Device Browser, the Session Browser is still limited in functionally. It currently consists of a split list 
 layout. The top list shows all discovered sessions and the bottom list shows the game instances that are part of 
 a selected session. A session does not necessarily consist of only games, but may also include auxiliary tools, 
 such as content cookers and file servers.
@@ -194,7 +194,7 @@ the Unreal Editor instance itself).
 ## Session Console
 
 The **Session Console** is the central hub for gathering the log output from all game instances that are currently selected 
-in the **Session Browser**. Which logs are being displayed depends on the currently selected instances in the **Session Browser**. 
+in the Session Browser. Which logs are being displayed depends on the currently selected instances in the Session Browser. 
 The list of log messages (4a) can then be filtered using a search string or category and verbosity filters (4b). It is 
 possible to send console commands to selected game instances (4c). Commands sent 
 to multiple instances will be executed by each instance. Last but not least, it is also possible to export all or selected 
@@ -221,9 +221,9 @@ for code and content verification at any given time.
 1. **Results**
 
 
-After selecting one or more instances in the **Session Browser**, the **Test List** will display all tests that 
+After selecting one or more instances in the Session Browser, the Test List will display all tests that 
 are available for those instances. You can then select which tests you actually want to run by selecting the checkbox 
-of one or more tests, or you can filter the list of tests using the **Filter Panel**. To execute the selected tests, 
+of one or more tests, or you can filter the list of tests using the Filter Panel. To execute the selected tests, 
 press the start button in the **Control Panel**. The results of the tests (including error messages) are then 
 aggregated in the **Results View**.
 
