#UEdoc #docs - removed overview header

diff --git a/Engine/Documentation/Source/Engine/Levels/Levels.INT.udn b/Engine/Documentation/Source/Engine/Levels/Levels.INT.udn
--- a/Engine/Documentation/Source/Engine/Levels/Levels.INT.udn
+++ b/Engine/Documentation/Source/Engine/Levels/Levels.INT.udn
@@ -10,7 +10,7 @@ Related: Resources/ContentExamples/LevelDesign
 [TOC(start:2 end:3)]
 
 ![](Level_Intro.png)
-## Overview
+
 When playing a Video Game, every object that you see or interact with in the world resides in what is known as a **Level**. In Unreal Engine 4 terms, a Level is made up of a collection of Static Meshes, Volumes, Lights, and Blueprints all working together to bring the desired experience to the player. Levels in UE4 can range in size from massive terrain based worlds to very small levels for just a single mini-game to take place in. The following guide will tell you all you need to know about creating and managing Levels inside of UE4. 
 
 ## How to Create Levels
