#UE4doc #docs: Removed extraneous bolding as according to the instructional guide. Fixed minor grammar errors.

diff --git a/Engine/Documentation/Source/Engine/Foliage/Foliage.INT.udn b/Engine/Documentation/Source/Engine/Foliage/Foliage.INT.udn
--- a/Engine/Documentation/Source/Engine/Foliage/Foliage.INT.udn
+++ b/Engine/Documentation/Source/Engine/Foliage/Foliage.INT.udn
@@ -49,8 +49,6 @@ Description:System for rendering instanced meshes on the surfaces of other geome
 
 [TOC(start:2 end:2)]
 
-## Overview
-
 [EXCERPT:Intro]
 
 The Foliage system allows you to quickly paint or erase sets of Static Meshes on Landscape Actors, other Static Mesh Actors, and legacy terrain. These meshes are automatically grouped together into batches that are rendered using hardware instancing, meaning that many instances can be rendered with only a single draw call.
@@ -73,7 +71,7 @@ The Foliage Edit Mode window appears in the **Modes** panel.
 
 ## Mesh List
 
-At first, there will be no foliage meshes listed. Before you can paint foliage, you must drag a Static Mesh from the **Content Browser** directly onto the black bar that says **Drag Static Meshes from the Content Browser into this area**.
+At first, there will be no foliage meshes listed. Before you can paint foliage, you must drag a Static Mesh from the **Content Browser** directly onto the black bar that says Drag Static Meshes from the Content Browser into this area.
 
 ![](Foliage_TextBar.png)
 
@@ -98,19 +96,19 @@ A Static Mesh in the mesh list can be displayed in one of three modes:
 * **Show Instance Settings** - Displays parameters that control the instances that are already placed in the level.
 	![Mesh Instance Settings](Foliage_ShowInstance.png)
 
-When the mesh is in either the **Paint Settings** or **Instance Settings** mode, three buttons also appear above the mesh icon.
+When the mesh is in either the Paint Settings or Instance Settings mode, three buttons also appear above the mesh icon.
 
 | Item | Description |
 | --- | --- |
 | ![](Shared/Icons/icon_assign_left_16x.png) | This button replaces all instances of the mesh already placed in the current level with the Static Mesh that is currently selected in the **Content Browser**. |
-| ![](Shared/Icons/icon_delete.png)(w:16) | This button removes the mesh from the mesh list as well as **all instances of this mesh already placed in the current level**. |
+| ![](Shared/Icons/icon_delete.png)(w:16) | This button removes the mesh from the mesh list as well as all instances of this mesh already placed in the current level. |
 | ![](Shared/Icons/icon_MeshPaint_Find_16x.png) | This button locates the Static Mesh in the **Content Browser**. |
 
 ### Selected Meshes in the Mesh List
 
-Meshes can be selected and deselected by clicking on them in the **Meshes** list (orange is selected). When creating or erasing meshes in the Level Editor viewport, only the meshes currently selected in the **Meshes** list will be affected. Foliage instances of other Static Meshes already in the level will not be affected in any way, and no instances of the deselected **Meshes** list items will be added to the level.
+Meshes can be selected and deselected by clicking on them in the **Meshes** list (orange is selected). When creating or erasing meshes in the Level Editor viewport, only the meshes currently selected in the Meshes list will be affected. Foliage instances of other Static Meshes already in the level will not be affected in any way, and no instances of the deselected Meshes list items will be added to the level.
 
-The following example shows two meshes selected. The **SM_Rock** mesh will not be painted or erased from the level as it is not selected.
+The following example shows two meshes selected. The SM_Rock mesh will not be painted or erased from the level as it is not selected.
 
 ![Two Meshes Selected](Foliage_MeshList.png)
 
@@ -128,7 +126,7 @@ Above the **Meshes** list, the Mesh Edit Tool displays properties related to the
 <!---
 [EXCERPT:PaintDensity]
 The target density at which you will add instances when using **Ctrl+click**.
-* This is a value from 0 to 1 where 1 paints mesh instances at the maximum density listed in the Mesh Properties for each mesh (see below).
+* This is a value from 0 to 1 where 1 paint mesh instances at the maximum density listed in the Mesh Properties for each mesh (see below).
 * If the density of meshes in the brush is already greater than this, no meshes will be added.
 [/EXCERPT:PaintDensity]
 --->
@@ -174,7 +172,7 @@ There are five tools available by selecting the appropriate button in the Toolba
 
 ![Paint Tool](ToolPaint.png)
 
-The Paint tool is Foliage Edit Mode's default tool. You can use it to paint instances of the Static Meshes in the **Meshes** list in your level.
+The Paint tool is Foliage Edit Mode's default tool. You can use it to paint instances of the Static Meshes in the Meshes list in your level.
 
 When Foliage Edit Mode is active, a transparent sphere brush is drawn in the level indicating where the foliage brush will operate.
 
@@ -249,7 +247,7 @@ Once instances are selected, a number of actions are available:
 
 ![Lasso Tool](ToolPaintSelect.png)
 
-The Lasso tool allows you to select a number of instances simultaneously by using the sphere brush that is also used with the Paint tool. The **Meshes** list selection and filter settings apply, just as they do for regular mesh painting. Holding down **Shift** while painting deselects instances. Once you have a selection, you can switch to the regular Select tool and perform some of the actions listed above.
+The Lasso tool allows you to select a number of instances simultaneously by using the sphere brush that is also used with the Paint tool. The Meshes list selection and filter settings apply, just as they do for regular mesh painting. Holding down **Shift** while painting deselects instances. Once you have a selection, you can switch to the regular Select tool and perform some of the actions listed above.
 
 ![Lasso Tool](Foliage_Lasso.png)(w:600)
 
@@ -300,7 +298,7 @@ The example Material below multiplies the Material mask by the fading factor to
 
 You can set foliage-related options for a Static Mesh in the **Static Mesh Editor** instead of in the Foliage tool.
 
-In the **Static Mesh Editor**, in the **Details** panel, in the **Static Mesh Settings** section, there is a **Foliage Default Settings** option. By default, it is set to **None**.  But if you click the drop-down arrow, you can choose **InstancedFoliageSettings**, which opens a whole sub-area of settings that apply to using this Static Mesh with the Foliage tool.  These settings mirror the settings in the Foliage tool's **Show Painting Settings** and **Show Instance Settings**.
+In the **Static Mesh Editor**, in the **Details** panel, in the **Static Mesh Settings** section, there is a **Foliage Default Settings** option. By default, it is set to **None**.  But if you click the drop-down arrow, you can choose **InstancedFoliageSettings**, which opens a whole sub-area of settings that apply to using this Static Mesh with the Foliage tool.  These settings mirror the settings in the Foliage tool's Show Painting Settings and Show Instance Settings.
 
 ![Static Mesh Editor Foliage Settings](Foliage_StaticMeshEditor.png)
 
@@ -308,7 +306,7 @@ In the **Static Mesh Editor**, in the **Details** panel, in the **Static Mesh Se
 
 Static Mesh LOD is supported on foliage. There are a few caveats:
 
-* You should ensure your Static Mesh has only one entry in the Elements array (visible under **LOD0**).
+* You should ensure your Static Mesh has only one entry in the Elements array (visible under LOD0).
 * Note that light and shadowmaps are shared between all LOD levels, so all LODs need to use the same unwrapping.
 * Currently the entire cluster of instances change LODs simultaneously. We may add support for distance-based fading per instance in the future.
 
