#UE4doc #docs: Removed extraneous bolding as according to the instructional guide. ##Overview removed as according to the instructional guide.

diff --git a/Engine/Documentation/Source/Engine/Landscape/Landscape.INT.udn b/Engine/Documentation/Source/Engine/Landscape/Landscape.INT.udn
--- a/Engine/Documentation/Source/Engine/Landscape/Landscape.INT.udn
+++ b/Engine/Documentation/Source/Engine/Landscape/Landscape.INT.udn
@@ -121,9 +121,6 @@ Description:Landscape system for creating terrain for large, open, outdoor envir
 
 [TOC(start:2 end:3)]
 
-
-## Overview
-
 [EXCERPT:Intro]
 The Landscape system enables you to create terrain for your world - mountains, valleys, uneven or sloped ground, even openings for caves - and easily modify both its shape and its appearance by using a range of tools.
 [/EXCERPT:Intro]
@@ -141,9 +138,9 @@ The Landscape tool has three modes, accessible by their icons at the top of the
 [EXCERPT:LandscapeToolModes2] 
 | Icon | Mode | Description |
 | ---- | ---- | ----------- |
-| ![](Engine/Landscape\Editing\Landscape_ManageMode.png) | **Manage mode** | Enables you to create new Landscapes, and to modify Landscape components. **Manage** mode is also where you work with **[Landscape Gizmos](Engine/Landscape\Editing\Gizmos)** to copy, paste, import, and export parts of your Landscape. <br/>For more information about **Manage** mode, see [](Engine/Landscape\Editing\ManageMode). |
-| ![](Engine/Landscape\Editing\Landscape_SculptMode.png)| **Sculpt mode** | Enables you to modify the shape of your Landscape, by selecting and using specific tools. <br/>For more information about **Sculpt** mode, see [](Engine/Landscape\Editing\SculptMode).|
-| ![](Engine/Landscape\Editing\Landscape_PaintMode.png)| **Paint mode** | Enables you to modify the appearance of parts of your Landscape by painting textures on it, based on the layers defined in the Landscape's Material. <br/>For more information about **Paint** mode, see [](Engine/Landscape\Editing\PaintMode).|
+| ![](Engine/Landscape\Editing\Landscape_ManageMode.png) | **Manage mode** | Enables you to create new Landscapes, and to modify Landscape components. Manage mode is also where you work with [Landscape Gizmos](Engine/Landscape\Editing\Gizmos) to copy, paste, import, and export parts of your Landscape. <br/>For more information about Manage mode, see [](Engine/Landscape\Editing\ManageMode). |
+| ![](Engine/Landscape\Editing\Landscape_SculptMode.png)| **Sculpt mode** | Enables you to modify the shape of your Landscape, by selecting and using specific tools. <br/>For more information about Sculpt mode, see [](Engine/Landscape\Editing\SculptMode).|
+| ![](Engine/Landscape\Editing\Landscape_PaintMode.png)| **Paint mode** | Enables you to modify the appearance of parts of your Landscape by painting textures on it, based on the layers defined in the Landscape's Material. <br/>For more information about Paint mode, see [](Engine/Landscape\Editing\PaintMode).|
 [/EXCERPT:LandscapeToolModes2]
 
 [REGION:tip]
