#UE4doc #docs: Removed extraneous bolding as according to the instructional guide. ##Overview removed as according to the instructional guide. Fixed multiple instances of the same word in a row: the the -> the.

diff --git a/Engine/Documentation/Source/Engine/Landscape/Creation/Creation.INT.udn b/Engine/Documentation/Source/Engine/Landscape/Creation/Creation.INT.udn
--- a/Engine/Documentation/Source/Engine/Landscape/Creation/Creation.INT.udn
+++ b/Engine/Documentation/Source/Engine/Landscape/Creation/Creation.INT.udn
@@ -25,9 +25,6 @@ Description:Guide to creating new Landscape terrains.
 
 [TOC(start:2)]
 
-
-## Overview
-
 Creating massive terrain based worlds for your projects is something that Unreal Engine 4 (UE4) can do quite well with its built in terrain editing tools. 
 The Landscape tool will allow you to create massive terrain based worlds that can still maintain playable frame rates across a multitude of different devices. 
 Because of the flexibility of the Landscape system, you can create your Landscape using either of the following methods. 
@@ -47,7 +44,7 @@ For more information about customizing the Level Editor's panels, see [Level Edi
 ## Opening the Landscape Tool
 
 Before you can create a Landscape, you must first open the Landscape tool which you can find on the **Modes toolbar** panel. 
-To open the Landscape tools click on the  the **Mountain icon** along the top of the **Modes** panel.
+To open the Landscape tools click on the **Mountain icon** along the top of the **Modes** panel.
 
 ![](Landscape_ModesPanel.png)
 
@@ -84,7 +81,7 @@ Creating a new Landscape from scratch is done through the **New Landscape** sect
 |**Rotation:**| Allows you to set the rotation of the Landscape in the world.
 |**Scale:**| Allows you to set the scale of the Landscape in the world.  
 |**Section Size:**| Section size is used for Landscape LOD and Culling. Smaller sections allow the Landscape to LOD sections more aggressively but with will have a higher CPU cost. Larger sizes will allow for fewer components and will be less costly on the CPU. If you want to have a large Landscape, you will need to use a larger section size as using a smaller section size and scaling the Landscape up could cost too much on the CPU.
-|**Section Per Component**|Sections Per Component helps with the Landscapes LOD.  Each section is the basic unit for the Landscape LOD. One component could have 2 x 2 sections which means that it is possible that one component could be rendering four different LOD�s at once. With a larger section size, you get the added benefit of reduced CPU calculation time. However you could run into issues with the Landscape rendering to many vertices at once. This is could be especially prevalent when using very large areas of Landscape.  These problems could be even worse on mobile devices because of the limited amount of draw calls you can have due to hardware limitations.|
+|**Section Per Component**|Sections Per Component helps with the Landscapes LOD.  Each section is the basic unit for the Landscape LOD. One component could have 2 x 2 sections which means that it is possible that one component could be rendering four different LOD's at once. With a larger section size, you get the added benefit of reduced CPU calculation time. However you could run into issues with the Landscape rendering to many vertices at once. This is could be especially prevalent when using very large areas of Landscape.  These problems could be even worse on mobile devices because of the limited amount of draw calls you can have due to hardware limitations.|
 |**Number of Components**|Along with section size, this sets the size of your landscape. This value is capped at 32 x 32 because each component does have a CPU cost associated with it and going over this cap could result in performance issues with the Landscape.|
 |**Overall Resolution:**| This is the number of vertices your Landscape is using.|
 |**Total Components**:| The total number of components that will be created for your landscape.|
