#UEdoc #docs - removed overview header and extraneous bold lettering

diff --git a/Engine/Documentation/Source/Engine/LevelStreaming/WorldBrowser/WorldBrowser.INT.udn b/Engine/Documentation/Source/Engine/LevelStreaming/WorldBrowser/WorldBrowser.INT.udn
--- a/Engine/Documentation/Source/Engine/LevelStreaming/WorldBrowser/WorldBrowser.INT.udn
+++ b/Engine/Documentation/Source/Engine/LevelStreaming/WorldBrowser/WorldBrowser.INT.udn
@@ -5,7 +5,6 @@ Description:System for managing large worlds including origin shifting technolog
 
 [TOC (start:2 end:3)]
 
-## Overview
 
 **World Composition** has been designed to simplify managing large worlds. One of the goals is to avoid using a 
 persistent level to store streaming information as it becomes a bottleneck when a team of level designers want to work 
@@ -21,16 +20,16 @@ allows you to create worlds which are not limited to the `WORLD_MAX` constant va
 
 World Composition managed worlds can be activated by switching on **Enable World Composition** flag in **World Settings**.
 
-You can also disable world origin shifting by switching off **Enable World Origin Rebasing** flag in **World Settings**. 
+You can also disable world origin shifting by switching off Enable World Origin Rebasing flag in World Settings. 
 
 ![](WorldBrowserSettings.png)
 
 ## Levels Hierarchy
 
 After you have activated World Composition, all the levels in your project will be visible in the **Levels** window. To open 
-this window, click on the **Windows** menu, and then select **Levels**.
+this window, click on the **Windows** menu, and then select Levels.
 
-The entries in the **Levels** window represent your world hierarchy. 
+The entries in the Levels window represent your world hierarchy. 
 
 ![](LevelsWindow.png)
 
@@ -39,11 +38,11 @@ The level with its name in blue text is the current level. Currently loaded leve
 To load a level in the world composition, just **double-click** its name.  Using drag-and-drop, you can establish relationships between individual levels. 
 Child levels store position relative to the parent level, so when you change the parent level's position all child levels will recalculate their position accordingly. 
 
-In the **Levels** window, you can lock and save levels, as well as open the **Level Blueprint** for each level, with buttons to the right of the level name.
+In the Levels window, you can lock and save levels, as well as open the **Level Blueprint** for each level, with buttons to the right of the level name.
  
 ## Levels Composition
 
-When **Enable World Composition** is on, there is a button to open **World Composition** at the top of the **Levels** window.
+When Enable World Composition is on, there is a button to open World Composition at the top of the Levels window.
 
 ![](WorldComposition.png)
 
@@ -85,7 +84,7 @@ lose its ability to be aligned with other landscape levels.
 [/REGION]
 
 
-The first time a level is loaded in **World Composition**, a new Level Bounds Actor is automatically created in the level.
+The first time a level is loaded in World Composition, a new Level Bounds Actor is automatically created in the level.
 The Level Bounds Actor is used to calculate the size of the level. 
 
 By default, the Level Bounds Actor automatically resizes to include all Actors found in the level.  However, some 
@@ -100,7 +99,7 @@ Currently, Realtime must be enabled in the viewports for the Level Bounds Actor
 [/REGION]
 
 The minimap has a yellow square indicating the safe edit area. The safe edit area's size is equal to the `WORLD_MAX` constant in the engine.
-Levels outside of this area will be hidden automatically, and their level images will reflect this.  **World Composition** tries to 
+Levels outside of this area will be hidden automatically, and their level images will reflect this.  World Composition tries to 
 keep levels you are working on close to the Editor world origin. This feature is called "focusing" and may shift the current 
 world origin to a new position if you begin editing a level outside of the safe edit area.  Shifted levels still preserve their Actors' original 
 transformations when they are saved, so saving a level that is currently shifted will not change the Actors' absolute positions.
@@ -121,7 +120,7 @@ transformations when they are saved, so saving a level that is currently shifted
  
 ## Level Details
  
-To see the details for any level in the **Level** window, select it and then click on the magnifying glass icon in the toolbar. This will summon the 
+To see the details for any level in the Level window, select it and then click on the magnifying glass icon in the toolbar. This will summon the 
 **Level Details** window, where you can also swap between available levels' information with a dropdown menu.
 
 ![](level_details.png)
