#UE4doc #docs: ##Overview removed as according to the instructional guide.

diff --git a/Engine/Documentation/Source/Engine/Rendering/ParticleSystems/Reference/Modules/Required/Required.INT.udn b/Engine/Documentation/Source/Engine/Rendering/ParticleSystems/Reference/Modules/Required/Required.INT.udn
--- a/Engine/Documentation/Source/Engine/Rendering/ParticleSystems/Reference/Modules/Required/Required.INT.udn
+++ b/Engine/Documentation/Source/Engine/Rendering/ParticleSystems/Reference/Modules/Required/Required.INT.udn
@@ -5,8 +5,6 @@ Description:Modules that are included along with an emitter and required for its
 
 [TOC(start:2 end:3)]
 
-## Overview
-
 The following classes and modules always exist when a new emitter is created. They will always be present and cannot be deleted.
 
 
