#UE4doc #docs: Fixed spelling: inti -> init

diff --git a/Engine/Documentation/Source/Engine/Rendering/ParticleSystems/Reference/Modules/Rotation/RotationModules.INT.udn b/Engine/Documentation/Source/Engine/Rendering/ParticleSystems/Reference/Modules/Rotation/RotationModules.INT.udn
--- a/Engine/Documentation/Source/Engine/Rendering/ParticleSystems/Reference/Modules/Rotation/RotationModules.INT.udn
+++ b/Engine/Documentation/Source/Engine/Rendering/ParticleSystems/Reference/Modules/Rotation/RotationModules.INT.udn
@@ -23,7 +23,7 @@ The Init Mesh Rotation module applies an initial rotation to mesh particles crea
 
 ### Init Mesh Rot (Seed)
 
-The Inti Mesh Rot (Seed) module applies an initial rotation to mesh particles created with a Mesh Data emitter, with the added benefit of randomization.
+The Init Mesh Rot (Seed) module applies an initial rotation to mesh particles created with a Mesh Data emitter, with the added benefit of randomization.
 
 | Properties | Description |
 | ---------- | ----------- |
