#UE4doc #docs: Fixed spelling error: subemitters -> submitters

diff --git a/Engine/Documentation/Source/Engine/Rendering/ParticleSystems/Reference/Modules/Spawn/SpawnModules.INT.udn b/Engine/Documentation/Source/Engine/Rendering/ParticleSystems/Reference/Modules/Spawn/SpawnModules.INT.udn
--- a/Engine/Documentation/Source/Engine/Rendering/ParticleSystems/Reference/Modules/Spawn/SpawnModules.INT.udn
+++ b/Engine/Documentation/Source/Engine/Rendering/ParticleSystems/Reference/Modules/Spawn/SpawnModules.INT.udn
@@ -32,5 +32,5 @@ Allows for spawning particles based on the distance the emitter has traveled. Us
 
 Using this module when spawning from other emitters may seem to not behave as expected when trying to spawn particles from a single lead particle.
 
-SpawnPerUnit actually uses the movement delta of the particle system itself and as such, does not understand what work to do when it is attached to subemitters inside the particle system. Even though its parents are moving through space, the actual whole system is still stationary, so the SpawnPerUnit module does nothing.
+SpawnPerUnit actually uses the movement delta of the particle system itself and as such, does not understand what work to do when it is attached to submitters inside the particle system. Even though its parents are moving through space, the actual whole system is still stationary, so the SpawnPerUnit module does nothing.
 
