#UE4doc #docs:  altered a sentence: This allows the vector field to be positionally offset within the emitter. -> This allows the vector fields position to be offset within the emitter.

diff --git a/Engine/Documentation/Source/Engine/Rendering/ParticleSystems/Reference/Modules/VectorField/VectorField.INT.udn b/Engine/Documentation/Source/Engine/Rendering/ParticleSystems/Reference/Modules/VectorField/VectorField.INT.udn
--- a/Engine/Documentation/Source/Engine/Rendering/ParticleSystems/Reference/Modules/VectorField/VectorField.INT.udn
+++ b/Engine/Documentation/Source/Engine/Rendering/ParticleSystems/Reference/Modules/VectorField/VectorField.INT.udn
@@ -35,7 +35,7 @@ The Local Vector Field module allows for a vector field to be associated directl
 |[REGION:tablesection]Vector Field[/REGION]||
 | ---- | ----- |
 | **Vector Field** | Property contains the assigned Vector Field from the **Content Browser**. Simply select the desired vector field and then click the arrow button to assign it. |
-| **Relative Translation** | This allows the vector field to be positionally offset within the emitter. [INCLUDE:#LocalTranslateNote]|
+| **Relative Translation** | This allows the vector fields position to be offset within the emitter. [INCLUDE:#LocalTranslateNote]|
 | **Relative Rotation** | This allows for a rotational offset for the local vector field. |
 | **Relative Scale 3D** | This allows the scale of the Local Vector Field to be edited, which is useful depending on the scale the vector field was created. |
 | **Intensity** | This affects the overall intensity of the local vector field. Think of it as a multiplier for the vector field's force. A value of 0 effectively deactivates the vector field. Values greater than 1 will overdrive the effect. |
