#UE4doc #docs: Altered images to match the instructional guide: green -> yellow.

diff --git a/Engine/Documentation/Source/Engine/Rendering/ParticleSystems/UserGuide/UserGuide.INT.udn b/Engine/Documentation/Source/Engine/Rendering/ParticleSystems/UserGuide/UserGuide.INT.udn
--- a/Engine/Documentation/Source/Engine/Rendering/ParticleSystems/UserGuide/UserGuide.INT.udn
+++ b/Engine/Documentation/Source/Engine/Rendering/ParticleSystems/UserGuide/UserGuide.INT.udn
@@ -5,8 +5,6 @@ Description:The general usage of Particle Systems and Cascade.
 
 [TOC(start:2 end:3)]
 
-## Overview
-
 This document will cover the most common workflows of working with Particle Systems and the Particle System editor: **Cascade**. Links to more extensive documentation are provided when available.
 
 
