#UEdoc #docs - updated grammar and functionality checks

diff --git a/Engine/Documentation/Source/Engine/Matinee/HowTo/MHT_1/MHT_1.INT.udn b/Engine/Documentation/Source/Engine/Matinee/HowTo/MHT_1/MHT_1.INT.udn
--- a/Engine/Documentation/Source/Engine/Matinee/HowTo/MHT_1/MHT_1.INT.udn
+++ b/Engine/Documentation/Source/Engine/Matinee/HowTo/MHT_1/MHT_1.INT.udn
@@ -47,11 +47,11 @@ In the following steps, we are going to create a simple door that opens/closes w
 	
 	![](MHT1_DragFrame.png)(w:640)
 
-1. Also in the **Props** folder, locate the **SM_Door** Static Mesh and drag-and-drop it into the **Level Viewport**. 
+1. Also in the Props folder, locate the **SM_Door** Static Mesh and drag-and-drop it into the Level Viewport. 
 
 	![](MHT1_DragDoor.png)(w:640)
 
-1. In the **Content Browser**, find the **SM_Door** Static Mesh again and **Double-Click** on it to open it up in the **Static Mesh** viewer. 
+1. In the Content Browser, find the SM_Door Static Mesh again and **Double-Click** on it to open it up in the **Static Mesh** viewer. 
 
 1. Inside the **Details** panel, in the search field type in **Collision** and set the **Collision Complexity** setting to **Use Complex Collision as Simple**. 
 
@@ -59,7 +59,7 @@ In the following steps, we are going to create a simple door that opens/closes w
 
 	_Note: The step above will ensure that the door has collision applied to it. _
 
-1. With the **SM_Door** selected in the **Level Viewport**, find the **Details** panel. Under **Transform** and the **Rotation** setting, set the **Z** value to **-180**. 
+1. With the **SM_Door** selected in the Level Viewport find the Details panel. Under **Transform** and the **Rotation** setting, set the **Z** value to **-180**. 
 
 	![](MHT1_RotateDoor.png)(w:640)
 
@@ -69,18 +69,18 @@ In the following steps, we are going to create a simple door that opens/closes w
 
 	![](MHT1_MoveDoor.png)(h:300) ![](MHT1_MoveDoor2.png)(w:310 h:300)
 	
-1. Inside the **Modes** menu and under the **Basic** tab, find the **Box Trigger** and drag-and-drop it into the **Level Viewport** inside the door frame.
+1. Inside the **Modes** menu and under the **Basic** tab, find the **Box Trigger** and drag-and-drop it into the Level Viewport inside the door frame.
 
 	![](MHT1_DragTrigger.png)(w:640)
 
-1. Select the **TriggerBox** and in the **Details** panel under the **Shape** section, locate and change the **X** value for **Box Extent** to **140**. 
+1. Select the **TriggerBox** and in the Details panel under the **Shape** section, locate and change the **X** value for **Box Extent** to **140**. 
 
 	![](MHT1_TriggerBoxSize.png)(w:760)
 
 	_Note: The **Trigger Box** extends past the door frame on the inside to allow it to open sooner (before the player approaches) as it will open inward._
 
 	[REGION:tip]
-	For an example of a door that opens in different directions based on which way the player enters the **TriggerBox** using **Blueprints**, refer to the [](Resources\ContentExamples\Blueprints_Advanced\2_4) Content Example. 
+	For an example of a door that opens in different directions based on which way the player enters the **Trigger Box** using **Blueprints**, refer to the [](Resources\ContentExamples\Blueprints_Advanced\2_4) Content Example. 
 	[/REGION]
 
 1. The door and trigger are now placed, next we will add the **Matinee** Actor to open/close it. 
@@ -98,7 +98,7 @@ In the following steps, we are going to create a simple door that opens/closes w
 
 	![](MHT1_AddMatinee.png)(w:760)
 
-1. The **Matinee** undo warning message will appear, click the **Continue** button. 
+1. The Matinee undo warning message will appear, click the **Continue** button. 
 
 	![](MHT1_MatineeWarning.png)
 
@@ -116,11 +116,11 @@ In the following steps, we are going to create a simple door that opens/closes w
 
 1. When the ![](MHT1_GroupName.png) appears, enter a **New Group Name** for the group (in this example, we are using the name **Door**). 
 
-1. Minimize the **Matinee** window, then in the **Main Editor** find the **World Outliner** window and select the **SM_Door** so that it is highlighted. 
+1. Minimize the Matinee window, then in the **Main Editor** find the **World Outliner** window and select the **SM_Door** so that it is highlighted. 
 
 	![](MHT1_SelectDoor.png)
 
-1. Return to the **Matinee** window and **Right-Click** on the Group (Door) to bring up the context menu, then select **Actors** then **Add Selected Actors**. 
+1. Return to the Matinee window and **Right-Click** on the Group (Door) to bring up the context menu, then select **Actors** then **Add Selected Actors**. 
 
 	![](MHT1_AddDoorToGroup.png)
 
@@ -140,7 +140,7 @@ In the following steps, we are going to create a simple door that opens/closes w
 
 	![](MHT1_AddKey.png)(w:780)
 
-1. Minimize the **Matinee** window and return to the **Level Viewport**. With the **SM_Door** still selected, press the **E** key to enter **Rotation** mode. 
+1. Minimize the Matinee window and return to the Level Viewport With the **SM_Door** still selected, press the **E** key to enter **Rotation** mode. 
 
 	![](MHT1_RotateWidget.png)
 
@@ -148,11 +148,11 @@ In the following steps, we are going to create a simple door that opens/closes w
 
 	![](MHT1_DoorOpen.png)
 
-1. Return to the **Matinee** window and **Left-Click** on the ![](MHT1_EndPoint.png) **End Marker** and move it to the left to **0.50**. 
+1. Return to the Matinee window and **Left-Click** on the ![](MHT1_EndPoint.png) **End Marker** and move it to the left to **0.50**. 
 
 	![](MHT1_EndPointMoved.png)
 
-1. The **Matinee** sequence set up is complete. In the next step, we will finish by hooking up the **Matinee** to the **Trigger**. 
+1. The Matinee sequence set up is complete. In the next step, we will finish by hooking up the Matinee to the **Trigger**. 
 
 
 [OBJECT:TopicList]
@@ -192,15 +192,15 @@ In the following steps, we are going to create a simple door that opens/closes w
 
 	![](MHT1_AddOnEndOverlap.png)
 
-1. In the **Event Graph** find the **MatineeActor** node and **Left-Click** on and drag-off the output pin (Blue Circle) to bring up the context menu, then enter **Play** in the search field and select the **Play** option. 
+1. In the **Event Graph** find the **MatineeActor** node and **Left-Click** on and **drag-off** the output pin (Blue Circle) to bring up the context menu, then enter **Play** in the search field and select the **Play** option. 
 
 	![](MHT1_DragOffPlay.png)(w:720)
 
-1. Drag-off the **MatineeActor** node again and in the context menu search field, enter **Reverse** and select the **Reverse** option. 
+1. **Drag-off** the **MatineeActor** node again and in the context menu search field, enter **Reverse** and select the **Reverse** option. 
 
 	![](MHT1_DragOffReverse.png)(w:720)
 
-1. Find the **OnActorBeginOverlap** node and drag-off the **Execute** output pin and connect it to the **Play** node's input pin. 
+1. Find the **OnActorBeginOverlap** node and **drag-off** the **Execute** output pin and connect it to the **Play** node's input pin. 
 
 	![](MHT1_ConnectOverlapToPlay.png)(w:720)
 
@@ -212,9 +212,9 @@ In the following steps, we are going to create a simple door that opens/closes w
 	
 	![](MHT1_ClickCompile.png)
 	
-1. The Blueprint is now complete and the **TriggerBox** is now hooked up to the **Matinee** Actor. 
+1. The Blueprint is now complete and the **TriggerBox** is now hooked up to the Matinee Actor. 
 
-1. Minimize the **Level Blueprint** and click the ![](MHT1_PlayButton.png) **Play** Button to play in the Editor. 
+1. Minimize the Level Blueprint and **click** the ![](MHT1_PlayButton.png) **Play** Button to play in the Editor. 
 
 	![](MHT1_SimpleDoor.gif)(w:480 convert:false)
 
diff --git a/Engine/Documentation/Source/Engine/Matinee/HowTo/MHT_2/MHT_2.INT.udn b/Engine/Documentation/Source/Engine/Matinee/HowTo/MHT_2/MHT_2.INT.udn
--- a/Engine/Documentation/Source/Engine/Matinee/HowTo/MHT_2/MHT_2.INT.udn
+++ b/Engine/Documentation/Source/Engine/Matinee/HowTo/MHT_2/MHT_2.INT.udn
@@ -38,7 +38,7 @@ In the following steps, we are going to create two lifts that move from one poin
 
 	![](MHT2_ThirdCubeAdded.png)
 
-1.  With the third cube selected, press **Control+W** to duplicate it and move the new cube across from, and leaving a space between, the third cube.
+1.  With the third cube selected, press **Control+W** to duplicate it and move the new cube across from and leaving a space between the third cube.
 
 	![](MHT2_FourthCubeAdded.png)
 
@@ -46,7 +46,7 @@ In the following steps, we are going to create two lifts that move from one poin
 
 	![](MHT2_FifthCubeAdded.png)
 
-1.  With the fifth cube selected, in the **Details** panel under **Transform**, set the **Scale** for the **Z** value to **5.0**.
+1.  With the fifth cube selected, in the Details panel under Transform set the **Scale** for the **Z** value to **5.0**.
 
 	![](MHT2_FifthCubeScaled.png)(w:720)
 
@@ -68,7 +68,7 @@ In the following steps, we are going to create two lifts that move from one poin
 
 	![](Engine\Matinee\HowTo\MHT_1\MHT1_AddMatinee.png)(w:720)
 
-1. The **Matinee** undo warning message will appear, click the **Continue** button. 
+1. The Matinee undo warning message will appear, click the **Continue** button. 
 
 	![](Engine\Matinee\HowTo\MHT_1\MHT1_MatineeWarning.png)
 	
@@ -86,11 +86,11 @@ In the following steps, we are going to create two lifts that move from one poin
 
 1. In the ![](MHT2_Lift1GroupName.png) **New Group Name** box, enter a name for the group (In this example, we named the group **Lift1**). 
 
-1. Minimize the **Matinee** window, then in the **Level Viewport** select the first cube created (the cube that is on the floor) so that it is highlighted. 
+1. Minimize the Matinee window, then in the **Level Viewport** select the first cube created (the cube that is on the floor) so that it is highlighted. 
 
 	![](MHT2_ReselectFirstCube.png)
 
-1. Return to the **Matinee** window and **Right-Click** on the Group (Lift1) to bring up the context menu, then select **Actors** then **Add Selected Actors**. 
+1. Return to the Matinee window and **Right-Click** on the Group (Lift1) to bring up the context menu, then select **Actors** then **Add Selected Actors**. 
 
 	![](MHT2_AddLift1.png)(w:640)
 	
@@ -110,13 +110,13 @@ In the following steps, we are going to create two lifts that move from one poin
 
 	![](MHT2_Lift1SecondKey.png)(w:640)
 
-1. Minimize the **Matinee Editor** and return to the **Level Viewport**. With the first cube still selected, move it up next to the second cube. 
+1. Minimize the Matinee Editor and return to the Level Viewport. With the first cube still selected, move it up next to the second cube. 
 
 	![](MHT2_Lift1MovedUp.png)
 
 	_Note: You should also see the yellow dotted line which indicates the movement the lift will take._
 
-1. In the **Matinee Editor**, **Left-Click** on the **Keyframe** at **2.00** and press **Control+C** to copy it, then move the **Time Bar** to **3.00** and press **Control+V** to paste it. 
+1. In the Matinee Editor, **Left-Click** on the **Keyframe** at **2.00** and press **Control+C** to copy it, then move the **Time Bar** to **3.00** and press **Control+V** to paste it. 
 
 	![](MHT2_Lift1ThirdKeyAdded.png)(w:640)
 
@@ -128,11 +128,11 @@ In the following steps, we are going to create two lifts that move from one poin
 
 	_Note: This will allow the lift to stay in the "lowered" position for a second before we loop the sequence and raise the lift again._
 
-1. Return to the **Main Editor** window and in the **World Outliner** select the **MatineeActor**. 
+1. Return to the **Main Editor** window and in the **World Outliner** select the **Matinee Actor**. 
 
 	![](MHT2_Matinee1Selected.png)(w:320)
 
-1. In the **Details** panel for the **Matinee** Actor, find the **Play** section and check both **Play on Level Load** and **Looping**. 
+1. In the Details panel for the Matinee Actor, find the **Play** section and check both **Play on Level Load** and **Looping**. 
 
 	![](MHT2_Lift1Details.png)(w:320)
 
@@ -150,15 +150,15 @@ In the following steps, we are going to create two lifts that move from one poin
 	[/PARAM]
 [/OBJECT]
 
-1. From the **World Outliner**, select the **MatineeActor** and in the **Details** panel, select the **Open Matinee** option. 
+1. From the **World Outliner**, select the **Matinee Actor** and in the Details panel, select the **Open Matinee** option. 
 
 	![](MHT2_ReopenMatinee.png)
 	
-1. Minimize the **Matinee** window, then in the **Level Viewport** select the third cube created so that it is highlighted. 
+1. Minimize the Matinee window, then in the Level Viewport select the third cube created so that it is highlighted. 
 
 	![](MHT2_SelectThirdPlatform.png)
 		
-1. Return to the **Matinee Editor** and **Right-Click** under the **Lift1** group, and in the context menu, select **Add New Empty Group**. 
+1. Return to the Matinee Editor and **Right-Click** under the **Lift1** group, and in the context menu, select **Add New Empty Group**. 
 
 	![](MHT2_AddSecondGroup.png)(w:640)
 
@@ -172,7 +172,7 @@ In the following steps, we are going to create two lifts that move from one poin
 
 	![](MHT2_Lift2SecondKey.png)(w:640)
 
-1. Minimize the **Matinee** window, then in the **Level Viewport** move the lift to the other side of the wall, next to the last platform as seen below. 
+1. Minimize the Matinee window, then in the Level Viewport move the lift to the other side of the wall, next to the last platform as seen below. 
 
 	![](MHT2_Lift2Moved.png)
 
@@ -196,7 +196,7 @@ In the following steps, we are going to create two lifts that move from one poin
 
 	![](MHT2_CurveKey1Break.png)(w:780)
 		
-1. In the **Curve Editor** window, **Left-Click** on the top **Keyframe** at the **2.00** mark. 
+1. In the Curve Editor window, **Left-Click** on the top **Keyframe** at the **2.00** mark. 
 
 	![](MHT2_CE_ClickTopKey.png)(w:780)
 
@@ -204,7 +204,7 @@ In the following steps, we are going to create two lifts that move from one poin
 
 	![](MHT2_CE_DragCurveDown1.png)(w:780)
 
-1. You should also see in the **Level Viewport** that the **Movement Path** has bent with the adjustment made in the last step. 
+1. You should also see in the Level Viewport that the **Movement Path** has bent with the adjustment made in the last step. 
 
 	![](MHT2_MovementBending.png)
 	
@@ -216,17 +216,17 @@ In the following steps, we are going to create two lifts that move from one poin
 
 	![](MHT2_EvenCurve.png)
 
-1. In the **Movement Track**, scrub the **Time Bar** to the **3.00** and **4.00** marks and press **Enter** at each to create **Keyframes**. 
+1. In the Movement Track scrub the **Time Bar** to the **3.00** and **4.00** marks and press **Enter** at each to create **Keyframes**. 
 
 	![](MHT2_Lift2FourthKeyAdded.png)
 	
-1. With the fourth **Keyframe** selected, in the **Level Viewport**, move the lift back to its original position. 
+1. With the fourth **Keyframe** selected, in the Level Viewport move the lift back to its original position. 
 
 	![](MHT2_Lift2FourthKeyView.png)
 
 	_Note: The movement track goes back through the wall, so we will need to adjust the curve for the remaining **Keyframes** as we did above._
 
-1. In the **Movement Track**, scrub the **Time Bar** to the **5.00** mark and press **Enter** to create the final **Keyframe**. 
+1. In the Movement Track scrub the **Time Bar** to the **5.00** mark and press **Enter** to create the final **Keyframe**. 
 
 	![](MHT2_FinalKeyFrameAdded.png)(w:780)
 
@@ -242,7 +242,7 @@ In the following steps, we are going to create two lifts that move from one poin
 
 	![](MHT2_ReverseCurveComplete.png)
 
-1. When you click the **Play** button from the **Main Editor's Toolbar**, you should have two lifts that move differently and one that uses a curved path. 
+1. When you click the **Play** button from the Main Editor's Toolbar, you should have two lifts that move differently and one that uses a curved path. 
 
 	![](MHT2_BothLiftsGif.gif)(w:420 convert:false)
 	
diff --git a/Engine/Documentation/Source/Engine/Matinee/HowTo/MHT_3/MHT_3.INT.udn b/Engine/Documentation/Source/Engine/Matinee/HowTo/MHT_3/MHT_3.INT.udn
--- a/Engine/Documentation/Source/Engine/Matinee/HowTo/MHT_3/MHT_3.INT.udn
+++ b/Engine/Documentation/Source/Engine/Matinee/HowTo/MHT_3/MHT_3.INT.udn
@@ -34,7 +34,7 @@ In the following steps, we are going to create a short sequence that uses multip
 
 	_Click for Full View._
 
-1.  Minimize the **Matinee** window, then from the **Modes** menu under the **Basic** tab, **Drag-and-drop** a **Camera** into the viewport.
+1.  Minimize the Matinee window, then from the **Modes** menu under the **Basic** tab, **Drag-and-drop** a **Camera** into the viewport.
 
 	![](MHT3_DragCamera.png)(w:580)
 
@@ -64,7 +64,7 @@ In the following steps, we are going to create a short sequence that uses multip
 	[/PARAM]
 [/OBJECT]
 
-1. Open the **Matinee** window (if you have closed it you can re-open it by selecting the **Matinee**, then from the **Details** panel, selecting **Open Matinee**.). 
+1. Open the Matinee window (if you have closed it you can re-open it by selecting the **Matinee**, then from the **Details** panel, selecting **Open Matinee**.). 
 
 	![](MHT3_OpenMatinee.png)
 
@@ -76,26 +76,26 @@ In the following steps, we are going to create a short sequence that uses multip
 
 	![](MHT3_Camera1Added.png)
 
-1. Repeat the process above and add the remaining two cameras to the **Matinee**. 
+1. Repeat the process above and add the remaining two cameras to the Matinee. 
 
 	![](MHT3_AllCamerasInMatinee.png)
 
-1. Next, we will increase the length of the **Matinee** by holding **Control** and dragging the red marker (in the yellow box below) to the **15.00** mark (yellow X). 
+1. Next, we will increase the length of the Matinee by holding **Control** and dragging the red marker (in the yellow box below) to the **15.00** mark (yellow X). 
 	[REGION:fullwidth]
 	![](MHT3_IncreaseMatineeTime.png)
 	[/REGION]
 
-1. Click on the first keyframe at **0.00** on the **Movement** track of **Camera1**, the camera should be selected in the viewport. 
+1. Click on the first keyframe at **0.00** on the **Movement** track of Camera1 the camera should be selected in the viewport. 
 	
 	![](MHT3_Camera1KeySelected.png)
 	
-1. In the **Viewport**, using the widget, move the camera up and angle it downward slightly (press **E** to enter **Rotation** mode). 
+1. In the Viewport, using the widget, move the camera up and angle it downward slightly (press **E** to enter **Rotation** mode). 
 
 	![](MHT3_Camera1Moved.png)
 
 	_We are going to pan the camera to the right across the template text, so try to angle the camera so that it will view the text._
 
-1. Return to the **Matinee** window and move the **Time Bar** to the **5.00** mark for the **Movement** track of **Camera1** and press **Enter** to add a new keyframe. 
+1. Return to the Matinee window and move the **Time Bar** to the **5.00** mark for the Movement track of Camera1 and press **Enter** to add a new keyframe. 
 
 	![](MHT3_Camera1Key2.png)
 	
@@ -114,7 +114,7 @@ In the following steps, we are going to create a short sequence that uses multip
 
 	![](MHT3_Camera2StartPosition.png)(h:420)
 
-1. Return to the **Matinee** window and move the **Time Bar** to the **7.50** mark for the **Movement** track of **Camera2** and press **Enter** to add a new keyframe. 
+1. Return to the Matinee window and move the **Time Bar** to the **7.50** mark for the Movement track of Camera2 and press **Enter** to add a new keyframe. 
 
 	![](MHT3_Camera2Key2.png)
 
@@ -122,15 +122,15 @@ In the following steps, we are going to create a short sequence that uses multip
 
 	![](MHT3_Camera2Finished.png)(h:420)
 
-1. Return to the **Matinee** window and move the keyframe at **0.00** for the **Movement** track of **Camera3** to the **7.50** mark. 
+1. Return to the Matinee window and move the keyframe at **0.00** for the Movement track of **Camera3** to the **7.50** mark. 
 
 	![](MHT3_Camera3Key1.png)
 
-1. With **Camera3's** keyframe still selected, return to the editor viewport and move the camera up and slightly hovering over the chair. 
+1. With Camera3's keyframe still selected, return to the editor viewport and move the camera up and slightly hovering over the chair. 
 
 	![](MHT3_Camera3Starting.png)(h:420)
 
-1. Return to the **Matinee** window and move the **Time Bar** to the **11.00** mark and press **Enter** to add a new keyframe. 
+1. Return to the Matinee window and move the **Time Bar** to the **11.00** mark and press **Enter** to add a new keyframe. 
 
 	![](MHT3_Camera3FinalKey.png)
 
@@ -138,7 +138,7 @@ In the following steps, we are going to create a short sequence that uses multip
 
 	![](MHT3_Camera3FinalShot.png)(h:420)
 
-1. Return to the **Matinee** window and select the **Movement** track of **Camera1**, then move the **Time Bar** to **11.00** and press **Enter** to add a new keyframe. 
+1. Return to the Matinee window and select the **Movement** track of Camera1 then move the **Time Bar** to **11.00** and press **Enter** to add a new keyframe. 
 	[REGION:fullwidth]
 	![](MHT3_Camera1LastKey1.png)
 	[/REGION]
@@ -157,11 +157,11 @@ In the following steps, we are going to create a short sequence that uses multip
 
 	![](MHT3_Camera3InPosition.png)
 
-1. Then return to the **Viewport Options** window and **Unlock** the camera in the same manner it was locked above. 
+1. Then return to the Viewport Options window and **Unlock** the camera in the same manner it was locked above. 
 
 	![](MHT3_UnlockCamera.png)
 
-1. Return to the **Matinee** window and move the **Time Bar** to **15.00** and add a keyframe on the **Movement** track of **Camera1**. 
+1. Return to the Matinee window and move the **Time Bar** to **15.00** and add a keyframe on the **Movement** track of **Camera1**. 
 	[REGION:fullwidth]
 	![](MHT3_LastCameraKey.png)
 	[/REGION]
@@ -180,11 +180,11 @@ In the following steps, we are going to create a short sequence that uses multip
 	[/PARAM]
 [/OBJECT]
 
-1. Open the **Matinee** window (if you have closed it you can re-open it by selecting the **Matinee**, then from the **Details** panel, selecting **Open Matinee**.). 
+1. Open the Matinee window (if you have closed it you can re-open it by selecting the **Matinee**, then from the **Details** panel, selecting **Open Matinee**.). 
 
 	![](MHT3_OpenMatinee.png)
 
-1. Under **Camera3**, **Right-Click** in an empty space in the **Tracks** window and from the context menu select **Add New Director Group**. 
+1. Under Camera3, **Right-Click** in an empty space in the **Tracks** window and from the context menu select **Add New Director Group**. 
 
 	![](MHT3_AddDirectorGroup.png)
 
@@ -203,7 +203,7 @@ In the following steps, we are going to create a short sequence that uses multip
 
 	_Click image for full view._
 
-1. Close the **Matinee** window, then select the **Matinee** in the **World Outliner** and in the **Details** panel, check the following options: **Play on Level Load**, **Looping**, **Disable Movement Input**, **Disable Look at Input**, and **Hide Player**. 
+1. **Close** the Matinee window, then select the **Matinee** in the **World Outliner** and in the Details panel, check the following options: **Play on Level Load**, **Looping**, **Disable Movement Input**, **Disable Look at Input**, and **Hide Player**. 
 
 	![](MHT3_MatineeActorSettings.png)
 
@@ -211,7 +211,7 @@ In the following steps, we are going to create a short sequence that uses multip
 
 	![](MHT3_BuildAndPlay.png)
 
-1. You should now have a **Matinee** sequence that starts when the level loads and cycles through three different cameras. 
+1. You should now have a Matinee sequence that starts when the level loads and cycles through three different cameras. 
 
 	![](MHT3_FinalResultGif.gif)(convert:false)
 
diff --git a/Engine/Documentation/Source/Engine/Matinee/HowTo/MatineeHowTo.INT.udn b/Engine/Documentation/Source/Engine/Matinee/HowTo/MatineeHowTo.INT.udn
--- a/Engine/Documentation/Source/Engine/Matinee/HowTo/MatineeHowTo.INT.udn
+++ b/Engine/Documentation/Source/Engine/Matinee/HowTo/MatineeHowTo.INT.udn
@@ -3,8 +3,8 @@ Title:Matinee - How To...
 Crumbs: %ROOT%, Engine, Engine/Matinee
 Description:A How To Guide for using different aspects of Matinee in Unreal Engine 4. 
 
-## Overview
-The **Matinee - How To...** pages will show you how you can set up simple scenarios using **Matinee** in a step-by-step process with images and instructions. 
+
+The **Matinee - How-To...** pages will show you how you can set up simple scenarios using Matinee in a step-by-step process with images and instructions. 
 
 Refer to the table below for more information and click each link to view its contents. 
 
