#UE4doc #docs: ##Overview removed as according to the instructional guide.

diff --git a/Engine/Documentation/Source/Engine/Rendering/ParticleSystems/Cascade/CascadeEditor.INT.udn b/Engine/Documentation/Source/Engine/Rendering/ParticleSystems/Cascade/CascadeEditor.INT.udn
--- a/Engine/Documentation/Source/Engine/Rendering/ParticleSystems/Cascade/CascadeEditor.INT.udn
+++ b/Engine/Documentation/Source/Engine/Rendering/ParticleSystems/Cascade/CascadeEditor.INT.udn
@@ -25,10 +25,6 @@ Description:An overview of the Cascade Editor's UI, breaking down each key area.
 
 [TOC]
 
-
-
-## Overview
-
 Cascade is a tool for creating modular particle-based effects using emitters in the Unreal Engine.
 
 
