#UE4doc #docs: ##Overview removed as according to the instructional guide.

diff --git a/Engine/Documentation/Source/Engine/Rendering/ParticleSystems/LODs/LODs.INT.udn b/Engine/Documentation/Source/Engine/Rendering/ParticleSystems/LODs/LODs.INT.udn
--- a/Engine/Documentation/Source/Engine/Rendering/ParticleSystems/LODs/LODs.INT.udn
+++ b/Engine/Documentation/Source/Engine/Rendering/ParticleSystems/LODs/LODs.INT.udn
@@ -5,8 +5,6 @@ Description:GPU-simulation of particles that allows for hundreds of thousands of
 
 [TOC(start:2 end:3)]
 
-## Overview
-
 Levels of detail (LODs) provide a way for effects artists to create particle effects that make the most efficient use of screen space, depending on player range. Consider a particle system of a camp fire: there may be several emitters involved, such as:
 
 * A core flame
