#UE4doc #docs: ##Overview removed as according to the instructional guide.

diff --git a/Engine/Documentation/Source/Engine/Rendering/ParticleSystems/Optimization/Optimization.INT.udn b/Engine/Documentation/Source/Engine/Rendering/ParticleSystems/Optimization/Optimization.INT.udn
--- a/Engine/Documentation/Source/Engine/Rendering/ParticleSystems/Optimization/Optimization.INT.udn
+++ b/Engine/Documentation/Source/Engine/Rendering/ParticleSystems/Optimization/Optimization.INT.udn
@@ -6,10 +6,6 @@ Description:Solving performance issues by optimizing Particle Systems.
 
 [TOC (start:2 end:3)]
 
-
-
-## Overview
-
 This document assumes prior knowledge of Cascade, the Material Editor, Blueprints, and Matinee along with common effects related performance issues.  There are multiple avenues to a performance issue solution. This document attempts to outline how Epic's effects team goes about solving these issues with the tools available.
 
 
