#UE4doc #docs: Removed extraneous bolding as according to the instructional guide. ##Overview removed as according to the instructional guide. Altered images to match the instructional guide: green -> yellow.

diff --git a/Engine/Documentation/Source/Engine/Rendering/ParticleSystems/Optimization/Results/Results.INT.udn b/Engine/Documentation/Source/Engine/Rendering/ParticleSystems/Optimization/Results/Results.INT.udn
--- a/Engine/Documentation/Source/Engine/Rendering/ParticleSystems/Optimization/Results/Results.INT.udn
+++ b/Engine/Documentation/Source/Engine/Rendering/ParticleSystems/Optimization/Results/Results.INT.udn
@@ -205,7 +205,7 @@ There are a few tricks to controlling emission:
 
 Cascade offers the ability to display large amounts of valuable information relative to particle system creation.  It is often helpful to turn these displays on to get a good idea of just how many particles your system is allocating, and how many particles your system is emitting.
 
-**Particle Allocation** determines how many particles it is possible for the emitter to place in the world at any given time, this plays a large part in determining how much memory your particle system will consume.  You can see a memory estimation by enabling the particle memory overlay in the View drop down menu in Cascade.
+Particle Allocation determines how many particles it is possible for the emitter to place in the world at any given time, this plays a large part in determining how much memory your particle system will consume.  You can see a memory estimation by enabling the particle memory overlay in the View drop down menu in Cascade.
 
 By adjusting settings, it is possible to see memory cost increase and decrease.
 
@@ -213,8 +213,8 @@ Several ways to reduce memory footprint include:
 
 
 * Reduce the number of modules in a system to the essential modules required for a given behavior.
-* For looping effects reduce **Loop Time/Duration** (allocates fewer particles per loop).
-* Reduce **Lifetime** (fewer particles overall).
+* For looping effects reduce Loop Time/Duration (allocates fewer particles per loop).
+* Reduce Lifetime (fewer particles overall).
 * Share as many modules as possible.
 
 
@@ -230,13 +230,13 @@ Tick Time represents the amount of time spent updating particle systems in the s
 
 There are many methods which can be used to reduce Tick costs.
 
-Tick Time is directly influenced by the number of active **EmitterActors** in your scene.  The more active emitters in the scene, the higher Tick Time will be.  Emitters should only be set to autoActivate if they are required to loop when the level starts.
+Tick Time is directly influenced by the number of active EmitterActors in your scene.  The more active emitters in the scene, the higher Tick Time will be.  Emitters should only be set to autoActivate if they are required to loop when the level starts.
 
 If large amounts of sprites/explosions/burst effects are visible on level load in the editor, these effects are automatically evaluated at run time and can cause visible hitching when a level is loaded.
 
 A Particle Parameter can be used in an effect to allow for position offsets once the effect is placed in the world.  By using a particle parameter, it is possible to reduce emitter counts, which reduces tick cost overall.  Particle Parameters can be setup by selecting particleParameter from the distribution list in the module for the individual setting.
 
-It is possible to stream out and disable **EmitterActors** when they are not in view, or if they are associated with an area of a level which is no longer in memory.  Streaming **EmitterActors** with level geometry is a good practice to reduce tick overhead.  In some cases where atmospheric effects are heavy, Blueprint actions can be used to toggle atmosphere off and on during combat sequences to lower tick time and overdraw.
+It is possible to stream out and disable EmitterActors when they are not in view, or if they are associated with an area of a level which is no longer in memory.  Streaming EmitterActors with level geometry is a good practice to reduce tick overhead.  In some cases where atmospheric effects are heavy, Blueprint actions can be used to toggle atmosphere off and on during combat sequences to lower tick time and overdraw.
 
 Particle Systems loaded into memory, but out of view (say on the floor above you) can also be toggled with Blueprint actions and streaming volumes to reduce evaluation costs.
 
