#UE4doc #docs: Removed extraneous bolding as according to the instructional guide.

diff --git a/Engine/Documentation/Source/Engine/Rendering/ParticleSystems/Optimization/SplitScreen/SplitScreen.INT.udn b/Engine/Documentation/Source/Engine/Rendering/ParticleSystems/Optimization/SplitScreen/SplitScreen.INT.udn
--- a/Engine/Documentation/Source/Engine/Rendering/ParticleSystems/Optimization/SplitScreen/SplitScreen.INT.udn
+++ b/Engine/Documentation/Source/Engine/Rendering/ParticleSystems/Optimization/SplitScreen/SplitScreen.INT.udn
@@ -11,13 +11,13 @@ Description:Using detail mode and profiling commands to optimize effects for spl
 
 Unreal Engine 4 currently supports switching between 3 detail modes: Low, Medium, and High.  In cases involving effects, detail modes are set for Medium, and High detail. This is used primarily for the case of split screen where draw calls can reach performance lowering levels. The primary objective being to reduce draw calls in Medium detail mode by culling or reducing emission rates.
 
-Changing the value in the **Medium Detail Spawn Rate Scale** field on an emitter in Cascade, will reduce its emission rate when the game is running in Medium Detail Mode.  These settings can be used to reduce overdraw and draw calls depending upon requirements.  In the case of split screen, draw call reduction is key.
+Changing the value in the Medium Detail Spawn Rate Scale field on an emitter in Cascade, will reduce its emission rate when the game is running in Medium Detail Mode.  These settings can be used to reduce overdraw and draw calls depending upon requirements.  In the case of split screen, draw call reduction is key.
 
 [How do I determine Draw Calls?](Engine/Rendering/ParticleSystems/Optimization/Results#RenderThread)
 
 The detail mode can be set in Cascade using the **View Menu** -> **Detail Mode** setting, or for the entire Unreal Editor by changing the **Quick Settings** -> **Engine Scalability Settings** -> **Effects** quality level.  Change to Medium to see Medium Detail Spawn Rate Scale and you should see red text displaying the detail mode currently active.  It is critical to work in Medium Detail Mode only when adjusting effects for this mode.  Changes will not be visible in game by default, unless your project defaults to Medium Detail Mode.
 
-**Medium Detail Spawn Rate Scale** is a multiplier, sticking to a 0-1 range will set the appropriate value.  Essentially setting to .5 will halve particle emission counts, a value of 1 will use the current emission rate.  The values are clamped and will not go above 1.
+Medium Detail Spawn Rate Scale is a multiplier, sticking to a 0-1 range will set the appropriate value.  Essentially setting to .5 will halve particle emission counts, a value of 1 will use the current emission rate.  The values are clamped and will not go above 1.
 
 It is good practice to consider split screen when in the process of placing ambient effects.  If the effect is non-essential to split screen and has a high number of draw calls, it might be a good place to start reducing emission counts especially if the effect consists only of mesh emitters.
 
@@ -29,7 +29,7 @@ It is important to remember draw calls for particles are determined by material,
 
 There are several commands which are useful in determining draw call costs.
 
-To switch to Split Screen Mode in **Play in Editor** or the game, use the console command `debugCreatePlayer 1`.
+To switch to Split Screen Mode in Play in Editor or the game, use the console command `debugCreatePlayer 1`.
 
 If you only have one controller, use the command `ssswapControllers` to jump between players.
 
