#UE4doc #docs: Removed extraneous bolding as according to the instructional guide.

diff --git a/Engine/Documentation/Source/Engine/Rendering/ParticleSystems/Overview/KeyConcepts.INT.udn b/Engine/Documentation/Source/Engine/Rendering/ParticleSystems/Overview/KeyConcepts.INT.udn
--- a/Engine/Documentation/Source/Engine/Rendering/ParticleSystems/Overview/KeyConcepts.INT.udn
+++ b/Engine/Documentation/Source/Engine/Rendering/ParticleSystems/Overview/KeyConcepts.INT.udn
@@ -5,15 +5,13 @@ Description:Critical concepts for working with particles in Unreal and the Casca
 
 [TOC(start:2 end:3)]
 
-## Overview
-
 **ParticleSystems** and **Cascade**, the _Particle System_ editor, are extremely flexible and powerful. However, before you jump in and get your hands dirty, there are some important concepts that you will need to fully grasp. The purpose of this document is to introduce you to these key concepts, giving you a more solid feel for how these aspects of the particle system work, while avoiding the nuances of the user interface or the specifics of particle effect creation.
 
 ## A Modular Approach to Particle Effects
 
 The primary and overreaching concept of Cascade is that of modularly designed particle systems. In some 3D effects packages, such as Maya, a particle system is created with most of the necessary properties for behavior already in place. The user then alters those properties in order to achieve the desired result.
 
-In Cascade, on the other hand, a particle system begins with only a few _bare bones_ properties and a few behavioral **modules**.
+In Cascade, on the other hand, a particle system begins with only a few _bare bones_ properties and a few behavioral modules.
 
 Each module represents a specific aspect of particle behavior, and contains only the properties that control that behavior, such as color, position at birth, movement, scaling, and many others. The user may then add or remove modules as they see fit, thereby further defining particle behavior. Since only the modules for necessary behavior are added, there is no extraneous calculation of unnecessary properties.
 
@@ -111,7 +109,7 @@ Descriptions of the currently available modules can be found at [ParticleSystemR
 
 ## Emitters, Particle Systems, and Emitter Actors
 
-As you work with Cascade to create your own particle effects, it is important to keep in mind how each one of the objects relates to one another. In this document, we have already discussed the concept of modules, but they only make up one component of a complete particle effect. All told, the components of a particle system are **modules**, **emitters**, **particle systems**, and **emitter Actors**. A solid way to remember how they relate is like so:
+As you work with Cascade to create your own particle effects, it is important to keep in mind how each one of the objects relates to one another. In this document, we have already discussed the concept of modules, but they only make up one component of a complete particle effect. All told, the components of a particle system are modules, emitters, particle systems, and emitter Actors. A solid way to remember how they relate is like so:
 
 * **Modules**, which define particle behavior and are placed within...
 * **Emitters**, which are used to emit a specific type of particle for an effect, and any number of which can be placed within a...
@@ -201,7 +199,7 @@ There are 5 primary distribution types:
 * Parameter - This is a type of distribution that allows a property to be parameterized so that it can be manipulated externally via Blueprints, Matinee, or other means.
 
 [REGION:note]
-	Each of these distribution types exist separately for float and vector data type values. For instance, there exist both **Float Uniform Curve** and **Vector Uniform Curve** distributions. However, whether you use a float or a vector version of each distribution is governed by the type of property being controlled; the user has no choice. As an example, you would control color (a vector value for RGB) via a Vector Uniform distribution, but Lifespan (a single float value) via a Float Uniform.
+	Each of these distribution types exist separately for float and vector data type values. For instance, there exist both Float Uniform Curve and Vector Uniform Curve distributions. However, whether you use a float or a vector version of each distribution is governed by the type of property being controlled; the user has no choice. As an example, you would control color (a vector value for RGB) via a Vector Uniform distribution, but Lifespan (a single float value) via a Float Uniform.
 [/REGION]
 
 
