#UE4doc #docs: ##Overview removed as according to the instructional guide.

diff --git a/Engine/Documentation/Source/Engine/Rendering/ParticleSystems/ParticleLights/ParticleLights.INT.udn b/Engine/Documentation/Source/Engine/Rendering/ParticleSystems/ParticleLights/ParticleLights.INT.udn
--- a/Engine/Documentation/Source/Engine/Rendering/ParticleSystems/ParticleLights/ParticleLights.INT.udn
+++ b/Engine/Documentation/Source/Engine/Rendering/ParticleSystems/ParticleLights/ParticleLights.INT.udn
@@ -44,8 +44,6 @@ Description:Particle Lights are another weapon in the VFX artist's arsenal.
 
 [TOC(start:2 end:2)]
 
-## Overview
-
 [EXCERPT:PL001]
 _CPU Particle Emitters_ can be made to render with a light attached to each particle by adding a **Particle Light Module** in **Cascade**.  This is useful for effects like sparks or fire, where it emits enough light to be noticeable on the surrounding environment.
 [/EXCERPT:PL001]
