Minor edits.

diff --git a/Engine/Documentation/Source/Engine/Rendering/Materials/HowTo/ShinyMaterials/ShinyMaterials.INT.udn b/Engine/Documentation/Source/Engine/Rendering/Materials/HowTo/ShinyMaterials/ShinyMaterials.INT.udn
--- a/Engine/Documentation/Source/Engine/Rendering/Materials/HowTo/ShinyMaterials/ShinyMaterials.INT.udn
+++ b/Engine/Documentation/Source/Engine/Rendering/Materials/HowTo/ShinyMaterials/ShinyMaterials.INT.udn
@@ -29,8 +29,8 @@ all of this can affect a surface's ability to reflect light.
 ## Roughness & Shininess
 
 Roughness plays a major role in helping to determine the Shininess of a surface. 
-The more rough a surface is, the **less** Shininess it will have. 
-While the less rough a surface is, the **more** Shininess it will have.
+The more rough a surface is, the less Shininess it will have. 
+While the less rough a surface is, the more Shininess it will have.
 Here is an example of how Roughness can affect Shininess. 
 
 ![](Engine\Rendering\Materials\PhysicallyBased\roughness_nonmetal.png)
diff --git a/Engine/Documentation/Source/Engine/Rendering/Materials/IntroductionToMaterials/IntroductionToMaterials.INT.udn b/Engine/Documentation/Source/Engine/Rendering/Materials/IntroductionToMaterials/IntroductionToMaterials.INT.udn
--- a/Engine/Documentation/Source/Engine/Rendering/Materials/IntroductionToMaterials/IntroductionToMaterials.INT.udn
+++ b/Engine/Documentation/Source/Engine/Rendering/Materials/IntroductionToMaterials/IntroductionToMaterials.INT.udn
@@ -98,10 +98,10 @@ As you may know, color - in terms of digital imagery - is broken in to 4 primary
 
 As with all digital images, for each pixel, the value of any of these channels can be represented by a number. Much of the work of a Material is simply manipulating these numbers based on a series of circumstances and mathematical expressions. 
 
-Materials use **floating point** values to store color data. This generally means that values for each channel will range from 0.0 to 1.0, instead of 0 to 255 as they do in some image editing applications. It is important to note that you can always overdrive values, which will in some cases result in special behaviors. For instance, if sending color into the Emissive input of a Material, which causes a glow, then values greater than 1.0 will increase the glow brightness.
+Materials use floating point values to store color data. This generally means that values for each channel will range from 0.0 to 1.0, instead of 0 to 255 as they do in some image editing applications. It is important to note that you can always overdrive values, which will in some cases result in special behaviors. For instance, if sending color into the Emissive input of a Material, which causes a glow, then values greater than 1.0 will increase the glow brightness.
 
 
-When thinking about Materials in UE4, keep in mind that many expressions **operate independently on each color channel**.  For example, the [Add](Engine/Rendering/Materials/ExpressionReference/Math#Add) node takes two inputs and adds them together per channel.  If you add together two RGB colors (3-channel vector values), then the output color will be: (Red1+Red2, Green1+Green2, Blue1+Blue2).  
+When thinking about Materials in UE4, keep in mind that many expressions operate independently on each color channel.  For example, the [Add](Engine/Rendering/Materials/ExpressionReference/Math#Add) node takes two inputs and adds them together per channel.  If you add together two RGB colors (3-channel vector values), then the output color will be: (Red1+Red2, Green1+Green2, Blue1+Blue2).  
 	
 [REGION:simpletable]
 ||
