#UE4doc #docs:
Unitalicized and bolded a term
Fixed some grammar / captialization
Changed green arrows/boxes in pictures to yellow

diff --git a/Engine/Documentation/Source/Engine/AI/EnvironmentQuerySystem/EQSPawn/EQSPawn.INT.udn b/Engine/Documentation/Source/Engine/AI/EnvironmentQuerySystem/EQSPawn/EQSPawn.INT.udn
--- a/Engine/Documentation/Source/Engine/AI/EnvironmentQuerySystem/EQSPawn/EQSPawn.INT.udn
+++ b/Engine/Documentation/Source/Engine/AI/EnvironmentQuerySystem/EQSPawn/EQSPawn.INT.udn
@@ -53,7 +53,7 @@ The **Environment Query System Testing Pawn** (EQSTestingPawn) is a specialized
 [/REGION]
 
 1. In the Content Browser, **Create** a new Blueprint
-1. Expand **All Classes** and search for `EQSTestingPawn`
+1. Expand **All Classes** and search for **EQSTestingPawn**
 1. Select **EQSTestingPawn**
 1. Name the new **EQSTestingPawn** Blueprint
 
@@ -73,7 +73,7 @@ The **Environment Query System Testing Pawn** (EQSTestingPawn) is a specialized
 
 ## Feedback Reference
 
-**Blue**: Test failed at that item, so the item is completely un-weighted and culled from the results. Usually Occurs when a value is filtered out. Such as using a Distance Test and filtering out any Item over (or under) a range. If you don't want a value to be culled completely, use the Clamp Features in the scoring section. 
+**Blue**: Test failed at that item, so the item is completely un-weighted and culled from the results. Usually occurs when a value is filtered out, such as using a Distance Test and filtering out any Item over (or under) a range. If you don't want a value to be culled completely, use the Clamp Features in the scoring section. 
 
 **Green-Red**: Green is more desirable than Red. The number listed on these is the weight of the Item. If previewing only one Debug Step, this value will be the final value from the selected Test.
 
