#UE4doc #docs: ##Overview removed as according to the instructional guide.

diff --git a/Engine/Documentation/Source/Engine/Rendering/ParticleSystems/Reference/BaseClasses/ParticleEmitterClass/ParticleEmitterClass.INT.udn b/Engine/Documentation/Source/Engine/Rendering/ParticleSystems/Reference/BaseClasses/ParticleEmitterClass/ParticleEmitterClass.INT.udn
--- a/Engine/Documentation/Source/Engine/Rendering/ParticleSystems/Reference/BaseClasses/ParticleEmitterClass/ParticleEmitterClass.INT.udn
+++ b/Engine/Documentation/Source/Engine/Rendering/ParticleSystems/Reference/BaseClasses/ParticleEmitterClass/ParticleEmitterClass.INT.udn
@@ -3,9 +3,6 @@ Title:Particle Emitter Class
 Crumbs: %ROOT%, Engine, Engine/Rendering/ParticleSystems, Engine/Rendering/ParticleSystems/Reference
 Description:A property reference for Particle Emitters, the components that make up Cascade's Emitter List.
 
-
-## The Particle Emitter Class
-
 A ParticleEmitter is a single particle effect that can be held in a ParticleSystem. The ParticleEmitter class contains the following public variables: 
 
 | Properties | Description |
