#UE4doc #docs: ##Overview removed as according to the instructional guide.

diff --git a/Engine/Documentation/Source/Engine/Rendering/ParticleSystems/Reference/BaseClasses/ParticleSystemClass/ParticleSystemClass.INT.udn b/Engine/Documentation/Source/Engine/Rendering/ParticleSystems/Reference/BaseClasses/ParticleSystemClass/ParticleSystemClass.INT.udn
--- a/Engine/Documentation/Source/Engine/Rendering/ParticleSystems/Reference/BaseClasses/ParticleSystemClass/ParticleSystemClass.INT.udn
+++ b/Engine/Documentation/Source/Engine/Rendering/ParticleSystems/Reference/BaseClasses/ParticleSystemClass/ParticleSystemClass.INT.udn
@@ -3,8 +3,6 @@ Title:Particle System Class
 Crumbs: %ROOT%, Engine, Engine/Rendering/ParticleSystems, Engine/Rendering/ParticleSystems/Reference
 Description:A technical reference for the base ParticleSystem class, showing its properties.
 
-## The Particle System Class
-
 A ParticleSystem is a complete particle effect that contains any number of ParticleEmitters. By allowing multiple emitters in a system, the designer can create elaborate particle effects that are held in a single system. Once created using Cascade, a ParticleSystem can then be inserted into a level or created in script.
 
 The ParticleSystem class contains the following public variable(s): 
