#UE4doc #docs: Removed extraneous bolding as according to the instructional guide. ##Overview removed as according to the instructional guide.

diff --git a/Engine/Documentation/Source/Engine/Rendering/ParticleSystems/Reference/Modules/Acceleration/AccelerationModules.INT.udn b/Engine/Documentation/Source/Engine/Rendering/ParticleSystems/Reference/Modules/Acceleration/AccelerationModules.INT.udn
--- a/Engine/Documentation/Source/Engine/Rendering/ParticleSystems/Reference/Modules/Acceleration/AccelerationModules.INT.udn
+++ b/Engine/Documentation/Source/Engine/Rendering/ParticleSystems/Reference/Modules/Acceleration/AccelerationModules.INT.udn
@@ -5,8 +5,6 @@ Description:These modules provide the ability to apply a change in velocity over
 
 [TOC(start:2 end:3)]
 
-## Acceleration Modules 
-
 These modules provide the ability to apply acceleration, or a change in velocity over time, to particles.
 
 
@@ -23,7 +21,7 @@ The Acceleration module sets the initial acceleration of the particle. This valu
 | **Apply Owner Scale** | If _true_, the acceleration is multiplied by the scale of the ParticleSystemComponent. |
 | **Always In World Space** | If _true_, the acceleration vector is assumed to be in world-space coordinates. Otherwise, it is assumed to be in local-space relative to the ParticleSystemComponent. |
 
-This module will add a vector parameter to the particle payload data, **UsedAcceleration**. This value is used to retain the acceleration over the life of each particle.
+This module will add a vector parameter to the particle payload data, UsedAcceleration. This value is used to retain the acceleration over the life of each particle.
 
 Each frame, the current and base velocity values of the particle is then updated using the formula (velocity += acceleration * DeltaTime) where DeltaTime is the amount of time passed since the last frame.
 
@@ -44,7 +42,7 @@ The Const Acceleration module is very similar to the Acceleration module, in tha
 |[REGION:tablesection]Const Acceleration[/REGION]||
 | **Acceleration** | A vector distribution indicating the acceleration to use. The value is grabbed based on the EmitterTime at the spawn of the particle.|
 
-This module will add a vector parameter to the particle payload data, **UsedAcceleration**. This value is used to retain the acceleration over the life of each particle.
+This module will add a vector parameter to the particle payload data, UsedAcceleration. This value is used to retain the acceleration over the life of each particle.
 
 Each frame, the current and base velocity values of the particle is then updated using the formula (velocity += acceleration * DeltaTime) where DeltaTime is the amount of time passed since the last frame.
 
@@ -61,7 +59,7 @@ Drag applies a constant resistance to motion to each particle. The coefficient v
 
 ### Drag Scale/Life
 
-Drag Scale/Life is a multiplier value that works **with an existing Drag module**, allowing, for example, for the drag to be scaled via a curve across the life of the particle.
+Drag Scale/Life is a multiplier value that works with an existing Drag module, allowing, for example, for the drag to be scaled via a curve across the life of the particle.
 
 [REGION:warning]
 	Drag Scale/Life only works on **GPU Sprite** particles.
