#UE4doc #docs: ##Overview removed as according to the instructional guide.

diff --git a/Engine/Documentation/Source/Engine/Rendering/ParticleSystems/Reference/Modules/Beam/BeamModules.INT.udn b/Engine/Documentation/Source/Engine/Rendering/ParticleSystems/Reference/Modules/Beam/BeamModules.INT.udn
--- a/Engine/Documentation/Source/Engine/Rendering/ParticleSystems/Reference/Modules/Beam/BeamModules.INT.udn
+++ b/Engine/Documentation/Source/Engine/Rendering/ParticleSystems/Reference/Modules/Beam/BeamModules.INT.udn
@@ -5,8 +5,6 @@ Description:These modules are used to configure the behavior of emitters using t
 
 [TOC(start:2 end:3)]
 
-## Beam Modules
-
 These modules are used to configure or modify the behavior of emitters using the Beam TypeData module.
 
 
