#UE4doc #docs: ##Overview removed as according to the instructional guide.

diff --git a/Engine/Documentation/Source/Engine/Rendering/ParticleSystems/Reference/Modules/Camera/CameraModules.INT.udn b/Engine/Documentation/Source/Engine/Rendering/ParticleSystems/Reference/Modules/Camera/CameraModules.INT.udn
--- a/Engine/Documentation/Source/Engine/Rendering/ParticleSystems/Reference/Modules/Camera/CameraModules.INT.udn
+++ b/Engine/Documentation/Source/Engine/Rendering/ParticleSystems/Reference/Modules/Camera/CameraModules.INT.udn
@@ -5,8 +5,6 @@ Description:These modules modify the emitter behavior with respect to the camera
 
 [TOC(start:2 end:3)]
 
-## Camera Modules
-
 These modules modify the emitter behavior with respect to the camera.
 
 
