#UE4doc #docs: Removed ôOverviewö header

diff --git a/Engine/Documentation/Source/Engine/Rendering/Materials/HowTo/Main_Material_Node/Main_Material_Node.INT.udn b/Engine/Documentation/Source/Engine/Rendering/Materials/HowTo/Main_Material_Node/Main_Material_Node.INT.udn
--- a/Engine/Documentation/Source/Engine/Rendering/Materials/HowTo/Main_Material_Node/Main_Material_Node.INT.udn
+++ b/Engine/Documentation/Source/Engine/Rendering/Materials/HowTo/Main_Material_Node/Main_Material_Node.INT.udn
@@ -100,27 +100,27 @@ Setting up a Material to use the Main Material Node is something can be accompli
 	 
 	 ![](Material_Open_Material_Editor.png)
 
-1. To use the Main Material Node, you first have to add Material Expression nodes to the Material Graph. Add a **Constant3Vector** Material Expression node the Material Graph by holding down the **Number 3 Key** on the keyboard and using the **Left Mouse Button** to click anywhere inside of the Material Graph to place the node.
+1. To use the Main Material Node, you first have to add Material Expression nodes to the Material Graph. Add a Constant3Vector Material Expression node the Material Graph by holding down the Number 3 Key on the keyboard and using the **Left Mouse Button** to click anywhere inside of the Material Graph to place the node.
 
 	![](Material_Adding_Const3Vector_Node.png)
 
 	[region:tip]
-	There are multiple ways to add nodes to the Material Graph. You can drag and drop form the **Material Pallet**, you can also **Right-Click** in the **Material Graph** and then search for nodes via the Context menu. Finally, you can use **Keyboard Shortcuts** to also place nodes in the Material Graph. You can read more about this in the [Material Editor UI]( Engine/Rendering/Materials/Editor/Interface) page.
+	There are multiple ways to add nodes to the Material Graph. You can drag and drop form the **Material Pallet**, you can also **Right-Click** in the **Material Graph** and then search for nodes via the Context menu. Finally, you can use Keyboard Shortcuts to also place nodes in the Material Graph. You can read more about this in the [Material Editor UI]( Engine/Rendering/Materials/Editor/Interface) page.
 	[/region]
 
-1. Next add a **Constant** Material Expression nodes to the Material Graph by first searching in the **Material Pallet** using the word **Constant** as a search term. When the **Constant** Material Expression is found select it by clicking on it with the **Left Mouse Button** and with the **Left Mouse Button** still held down, drag the Material Expression node into the Material Graph releasing the **Left Mouse Button** when your mouse is over the location in the Material Graph you want to place the Material Expression node. 
+1. Next add a Constant Material Expression nodes to the Material Graph by first searching in the **Material Pallet** using the word Constant as a search term. When the Constant Material Expression is found select it by clicking on it with the **Left Mouse Button** and with the **Left Mouse Button** still held down, drag the Material Expression node into the Material Graph releasing the **Left Mouse Button** when your mouse is over the location in the Material Graph you want to place the Material Expression node. 
 
 	![](MMN_Add_Const_Pram.gif)(convert:false)
 
-1. Now that we have a **Constant node**, duplicate it by first selecting it and then pressing **Ctrl+W** on the keyboard once. When completed, you should have something that looks like this. 
+1. Now that we have a Constant node, duplicate it by first selecting it and then pressing Ctrl+W on the keyboard once. When completed, you should have something that looks like this. 
 	
 	![](Materail_Duplication_Of_Nodes.png)
 
-1. With all of the required Material Expression nodes now in place, it is now time to setup and connect them to the **Main Material Node**. First **Double-Click** on the **Constant3Vector** node using the **Left Mouse Button** to open up the Color Picker.  With the **Color Picker** open, select the color red and then press the **Ok** button.
+1. With all of the required Material Expression nodes now in place, it is now time to setup and connect them to the Main Material Node. First **Double-Click** on the Constant3Vector node using the **Left Mouse Button** to open up the Color Picker.  With the **Color Picker** open, select the color red and then press the **Ok** button.
 	 
 	 ![](Materials_Color_Picker.gif)(convert:false)
 
-1. Now connect the **Constant3Vector** to the **Base Color** input on the Main Material node. To do this, use the **Left Mouse Button** to click on the circle output that is on the right of the **Constant3Vector**. Then with the **Left Mouse Button** still held down, drag to the right until you see a line coming from the circle output. Connect this line to the **Base Color** input on the **Main Material node**.
+1. Now connect the **Constant3Vector** to the **Base Color** input on the Main Material node. To do this, use the **Left Mouse Button** to click on the circle output that is on the right of theConstant3Vector. Then with the **Left Mouse Button** still held down, drag to the right until you see a line coming from the circle output. Connect this line to the Base Color input on the Main Material node.
 	
 	![](Materials_Connect_Nodes_To_Graph.gif)(convert:false)
 
@@ -128,11 +128,11 @@ Setting up a Material to use the Main Material Node is something can be accompli
 	As you make changes to the Main Material Node the viewport in the Material Editor will update to reflect your changes.  However this viewport can take a few seconds to update with changes depending on the complexity of your Material. For making slight tweaks to the Material a [Material Instance](Engine/Rendering/Materials/MaterialInstances) should be used as it shows changes instantly. 
 	[/region]
 
-1. With the Base Color now set, set up the remaining **Constant** Material Expression nodes by first selecting one of them with the **Left Mouse Button**. Then under the **Details** panel, input a value. For the first node, input a value of  **1.0** and for the second node input a value of **0.5**. 
+1. With the Base Color now set, set up the remaining **Constant** Material Expression nodes by first selecting one of them with the **Left Mouse Button**. Then under the **Details** panel, input a value. For the first node, input a value of  1.0 and for the second node input a value of 0.5. 
 	
 	![](Materail_Setup_Constants.png)
 
-1. Once the values of the **Constant**  Material Expression are set, connect the **Constant** Material Expression node with the value of 1.0 to the **Metallic** input on the Main Material Node. Then connect the **Constant** Material Expression node with the value of 0.5 to the **Roughness** input on the Main Material node.  
+1. Once the values of the **Constant**  Material Expression are set, connect the **Constant** Material Expression node with the value of 1.0 to the **Metallic** input on the Main Material Node. Then connect the **Constant** Material Expression node with the value of 0.5 to the Roughness input on the Main Material node.  
 	
 	![](Materials_Connect_&_Compile.gif)(convert:false)
 
diff --git a/Engine/Documentation/Source/Engine/Rendering/Materials/HowTo/Making_Functions/Making_Functions.INT.udn b/Engine/Documentation/Source/Engine/Rendering/Materials/HowTo/Making_Functions/Making_Functions.INT.udn
--- a/Engine/Documentation/Source/Engine/Rendering/Materials/HowTo/Making_Functions/Making_Functions.INT.udn
+++ b/Engine/Documentation/Source/Engine/Rendering/Materials/HowTo/Making_Functions/Making_Functions.INT.udn
@@ -37,7 +37,7 @@ Creating Material Functions can be accomplished in the few following steps.
 
 1. Once the Material Function has been created it needs to be named. 
 	Your name should be as descriptive as possible so that when someone goes to look at your Material Function they can easily understand what it does. 
-	For this example, we are going to call our node **Custom_Fresnel**. 
+	For this example, we are going to call our node Custom_Fresnel. 
 	You can re-name your Material Function by selecting it in the **Content Browser**, pressing **F2** on the keyboard, and then typing in the name that you want. 
 	
 	![](Material_Functions_Rename.gif)(convert:false)
@@ -51,20 +51,20 @@ Creating Material Functions can be accomplished in the few following steps.
 
 	[region:tip]
 	The Material Function editor works in a very similar manner to how the Material Editor works, 
-	the two biggest differences between being what the **Details** panel does and how the **Final Material** node works. 
-	You can read more about how the Material Function editor works but checking out the **[Material Functions Overview](Engine/Rendering/Materials/Functions/Overview)**
+	the two biggest differences between being what the **Details** panel does and how the Final Material node works. 
+	You can read more about how the Material Function editor works but checking out the [Material Functions Overview](Engine/Rendering/Materials/Functions/Overview)
 	[/region]
 
 1. Now with the Material Function Editor open, let us begin to create our Fresnel effect. 
 	Since the Fresnel Effect works based on glancing angles, we are going to need to compare the current surface angle with our camera angle and then use that to determine our **Fresnel Value**. 
-	To do this inside of our shader, we use the **[Dot Product]( Engine/Rendering/Materials/ExpressionReference/VectorOps#dotproduct)** node to compare our current camera angle with the angle of the surface currently being viewed. 
+	To do this inside of our shader, we use the [Dot Product]( Engine/Rendering/Materials/ExpressionReference/VectorOps#dotproduct) node to compare our current camera angle with the angle of the surface currently being viewed. 
 	To perform this type of calculation with Material nodes, we are going to need to make use of the following nodes. 
-	* **[Constant 3 Vector Node](Engine/Rendering/Materials/ExpressionReference/Constant#constant3vector)** x 1
-	* **[Transform Node](Engine/Rendering/Materials/ExpressionReference/VectorOps#transform)** x 1
-	* **[Dot Product Node](Engine/Rendering/Materials/ExpressionReference/VectorOps#dotproduct)** x 1
-	* **[Camera Vector WS Node](Engine/Rendering/Materials/ExpressionReference/Vector#cameravectorws)** x 1
+	* [Constant 3 Vector Node](Engine/Rendering/Materials/ExpressionReference/Constant#constant3vector) x 1
+	* [Transform Node](Engine/Rendering/Materials/ExpressionReference/VectorOps#transform) x 1
+	* [Dot Product Node](Engine/Rendering/Materials/ExpressionReference/VectorOps#dotproduct) x 1
+	* [Camera Vector WS Node](Engine/Rendering/Materials/ExpressionReference/Vector#cameravectorws) x 1
 
-	You can add the above nodes to the Material Function graph a number of ways, but the easiest is to search for each node in the **Palette** menu and then select the node with the **Left Mouse Button** and Drag it into the Material Graph. When completed, you should have something that looks like this.   
+	You can add the above nodes to the Material Function graph a number of ways, but the easiest is to search for each node in the **Palette** menu and then select the node with the **Left Mouse Button** and Drag it into the Material Graph. When completed, you should have something that looks like this:   
 	
 	![](Material_Function_Needed_Nodes.png)
 
@@ -72,19 +72,19 @@ Creating Material Functions can be accomplished in the few following steps.
 	The Fresnel Effect is the term used to describe how the light you see reflects at different intensities based off of the angle you are viewing it from. For example, if you are standing over a pool looking straight down at the pool, you will not see a lot of reflections in the water. As you start to move your head so that the water in the pool becomes more and more parallel to your eye level, you will begin to notice more and more reflections in the water and this is the type of behavior that we are trying to re-create in our shader, the ability to add or remove visual information based on viewing angle. 
 	[/region]
 	
-1. With the nodes now placed down, it is now time to set them up and wire them together. First, set the **Blue Channel** of the **Constant 3 Vector** to **1.0** as this will represent the Z value in our Fresnel term. Then take the output of the **Constant 3 Vector** node and wire that into the Input of the **Vector Transform** node. Also make sure that the **Transform Node** is setup so that we are transforming from **Tangent Space** into **World Space**. Then take the output of the **Vector Transform** node and wire that into the **A** input of the **Dot Product** node. Once that has been completed, take the output of the **Camera Vector** node and wire that into the **B** input of the **Dot Product** node. When completed, you should have something that looks like this.
+1. With the nodes now placed down, it is now time to set them up and wire them together. First, set the Blue Channel of the Constant 3 Vector to **1.0** as this will represent the Z value in our Fresnel term. Then take the output of the Constant 3 Vector node and wire that into the Input of the Vector Transform node. Also make sure that the Transform Node is setup so that we are transforming from Tangent Space into World Space. Then take the output of the Vector Transform node and wire that into the A input of the Dot Product node. Once that has been completed, take the output of the Camera Vector node and wire that into the B input of the Dot Product node. When completed, you should have something that looks like this.
 
 	![](Material_Function_Nodes_Setup.png)
 
-1. The next thing that we want to do is add a **[Clamp Node](Engine/Rendering/Materials/ExpressionReference/Math#clamp)** to clamp our results to the 0 to 1 range to make sure that we do not get any errors that might cause rendering artifacts. You can find the **Clamp** node by searching for it in the **Palette** menu and the Drag it into the Material Graph. When completed, you should have something that looks like this. 
+1. The next thing that we want to do is add a [Clamp Node](Engine/Rendering/Materials/ExpressionReference/Math#clamp) to clamp our results to the 0 to 1 range to make sure that we do not get any errors that might cause rendering artifacts. You can find the **Clamp** node by searching for it in the **Palette** menu and the Drag it into the Material Graph. When completed, you should have something that looks like this. 
 	
 	![](Material_Function_Adding_Clamp_Node.png)
 
-1. After the Clamp node, we need to add a **[One-Minus node](Engine/Rendering/Materials/ExpressionReference/Math#oneminus)** to invert our effect so that only the edges are effected. You can find the One-Minus node by searching for **OneMinus** in the **Palette**. When completed, you should have something that looks like the following. 
+1. After the Clamp node, we need to add a [One-Minus node](Engine/Rendering/Materials/ExpressionReference/Math#oneminus) to invert our effect so that only the edges are effected. You can find the One-Minus node by searching for OneMinus in the Palette. When completed, you should have something that looks like the following. 
 	
 	![](Material_Function_Adding_OneMinus.png)
 
-1. Next, we need to add two **[Function Input]( Engine/Rendering/Materials/Functions/Overview #functioninputnodes)** nodes. Function Input nodes are how we supply our Material Function with values or textures when placed inside of another Material.  To add a Function Input to our material search for **Function Input** in the **Palette** and then drag it into the Material Graph. Once you have one placed, you can duplicate it by first selecting the node and then pressing **Ctrl+W** on the keyboard. When completed, you should have something that looks like this. 
+1. Next, we need to add two [Function Input]( Engine/Rendering/Materials/Functions/Overview #functioninputnodes) nodes. Function Input nodes are how we supply our Material Function with values or textures when placed inside of another Material.  To add a Function Input to our material search for **Function Input** in the **Palette** and then drag it into the Material Graph. Once you have one placed, you can duplicate it by first selecting the node and then pressing **Ctrl+W** on the keyboard. When completed, you should have something that looks like this. 
 	
 	![](Material_Function_Adding_Input_Nodes.png)
 
@@ -102,26 +102,26 @@ Creating Material Functions can be accomplished in the few following steps.
 	![](Material_Function_Input_Node_Pow.png)
 	
 	[region:tip]
-	For more information about what each option of the Function node does, check out the **[Function Input node documentation](Engine/Rendering/Materials/Functions/Overview#functioninputnodes)**. 
+	For more information about what each option of the Function node does, check out the [Function Input node documentation](Engine/Rendering/Materials/Functions/Overview#functioninputnodes). 
 	[/region]
 
-1. With the Function Input nodes now setup, we now need to add a **Power Node** to the Material Graph. We are using a **Power Node** so that we can either increase or decrease the falloff of our Fresnel effect. This will help to push the Fresnel effect either closer or further from the edges of the mesh that it is on. To add a **Power Node** search for it in the **Palette** and then Drag and Drop it into the Material Graph. When completed, you should have something that looks like this. 
+1. With the Function Input nodes now setup, we now need to add a Power Node to the Material Graph. We are using a Power Node so that we can either increase or decrease the falloff of our Fresnel effect. This will help to push the Fresnel effect either closer or further from the edges of the mesh that it is on. To add a Power Node search for it in the **Palette** and then Drag and Drop it into the Material Graph. When completed, you should have something that looks like this. 
 	
 	![](Material_Function_Adding_Power.png)
 
-1. Next, we are going to add a **Multiply Node** so that we can control the intensity of the Fresnel effect.  To add a **Multiply Node** search for it in the **Palette** and then Drag and Drop it into the Material Graph. When completed, you should have something that looks like this. 
+1. Next, we are going to add a Multiply Node so that we can control the intensity of the Fresnel effect.  To add a Multiply Node search for it in the **Palette** and then Drag and Drop it into the Material Graph. When completed, you should have something that looks like this. 
 	
 	![](Material_Function_Adding_Multiply.png)
 
-1. Now with all the nodes we need added to the Material Graph, it is now time to wire everything together. First, we need to wire the output of the **OneMinus** node to the **Base** input of the **Power** node. Then we need to connect our first **Function Input** node, **Fresnel Exponent**, into the **Exp** input of the **Power** node. Once that has been completed, we need to then connect the output of the **Power** node to the A input of the **Multiply** node and then connect our second **Function Input**, **Fresnel Intensity**, into the B input of the **Multiply** node. Finally we need to take the output of the ** Multiply** node and connect that to the input of the **Output Results** node.  
+1. Now with all the nodes we need added to the Material Graph, it is now time to wire everything together. First, we need to wire the output of the OneMinus node to the Base input of the Power node. Then we need to connect our first Function Input node, Fresnel Exponent, into the Exp input of the Power node. Once that has been completed, we need to then connect the output of the Power node to the A input of the Multiply node and then connect our second Function Input, Fresnel Intensity, into the B input of the Multiply node. Finally we need to take the output of the Multiply node and connect that to the input of the Output Results node.  
 	
-	![](Material_Functions_Making_Connections.gif)(convert:false)
+	![](Material_Functions_Making_Connections.gif)(convert:false)                                                                                                                                                                                                      
 
 1. Once the Material nodes have all been connected, we need to now compile the Material Function by pressing the **Apply** button on the main tool bar. If the Material Function compiled correctly, the **Apply** button will be grayed out. When completed, you should have something that looks like the following image.
 
 	![](Material_Function_Completed_Material.png)
 
-
+	
 ## Using Material Functions 
 
 1. With our Material Function now finished, close the Material Function Editor and create a new Material in the **Content Browser**. To create a new Material, **Right-Click** in the **Content Browser** and then select **Material** from **Basic Assets List** and name your newly created Material, **Fresnel_Material**.
@@ -140,7 +140,7 @@ Creating Material Functions can be accomplished in the few following steps.
 
 	![](Material_Function_Adding_Nodes.png)
 
-1. Now that all the nodes have been added, it is now time to connect to the proper inputs and outputs. First, connect the **Constant 3 Vector** to the **Base Color** input on the Main Shader node. Then, connect one of the **Constant One** to the **Metallic** input on the Main Shader node. Then connect another of the **Constant One** to the **Roughness** input on the Main Shader node. Finally we are going to want to connect the final two **Constant One** nodes to the **Fresnel Intensity** and the **Fresnel Exponent** and then connect the output of the **Custom_Fresnel** node to the **Emissive Color** input on the Main Shader node. When completed, you should have something that looks like this.
+1. Now that all the nodes have been added, it is now time to connect to the proper inputs and outputs. First, connect the Constant 3 Vector to the Base Color input on the Main Shader node. Then, connect one of the Constant One to the Metallic input on the Main Shader node. Then connect another of the Constant One to the Roughness input on the Main Shader node. Finally we are going to want to connect the final two Constant One nodes to the Fresnel Intensity and the Fresnel Exponent and then connect the output of the Custom_Fresnel node to the Emissive Color input on the Main Shader node. When completed, you should have something that looks like this.
 
 	![](Material_Functions_Connecting_Material_Function.gif)(convert:false)
 
diff --git a/Engine/Documentation/Source/Engine/Rendering/Materials/HowTo/Placing_Expressions_Functions/Placing_Expressions_Functions.INT.udn b/Engine/Documentation/Source/Engine/Rendering/Materials/HowTo/Placing_Expressions_Functions/Placing_Expressions_Functions.INT.udn
--- a/Engine/Documentation/Source/Engine/Rendering/Materials/HowTo/Placing_Expressions_Functions/Placing_Expressions_Functions.INT.udn
+++ b/Engine/Documentation/Source/Engine/Rendering/Materials/HowTo/Placing_Expressions_Functions/Placing_Expressions_Functions.INT.udn
@@ -28,9 +28,9 @@ As you will soon see,  Material Expression nodes are key to building and interac
 While Material Expressions and Material Functions are the building blocks used to create Materials, there are some differences between the nodes and how they work.
 The main difference between Material Functions and Material Expressions is that Material Functions are made inside of the Material Function Editor and saved assets in the **Content Browser**. 
 
-* **[Material Expressions]( Engine/Rendering/Materials/ExpressionReference)**: Material Expressions are static pieces of code that can only do what they are programmed to do. For example, a Multiply node can only Multiply, it will never be able to add two numbers together. The only way you could change this was to make changes to the source code.
+* [Material Expressions]( Engine/Rendering/Materials/ExpressionReference): Material Expressions are static pieces of code that can only do what they are programmed to do. For example, a Multiply node can only Multiply, it will never be able to add two numbers together. The only way you could change this was to make changes to the source code.
 
-* **[Material Functions]( Engine/Rendering/Materials/Functions)**: Material Functions are more dynamic because you can create Material Functions to perform any type of calculations you want. If you wanted your Material Function to have the ability to not only Multiply but also Add, you could set it up to do so as those changes would be made right inside the Material Function Editor without having to make source code changes.
+* [Material Functions]( Engine/Rendering/Materials/Functions): Material Functions are more dynamic because you can create Material Functions to perform any type of calculations you want. If you wanted your Material Function to have the ability to not only Multiply but also Add, you could set it up to do so as those changes would be made right inside the Material Function Editor without having to make source code changes.
 
 
 ## Placing Material Expressions
@@ -51,7 +51,7 @@ You can also add Material Expressions by **Right-Clicking** anywhere in the Mate
 
 ### Keyboard Shortcuts
 
-There are a number of Keyboard shortcuts that you can use to help you quickly place down Material Expressions that are most commonly used. You can change what key spawns what Material Expression by going to **Edit -> Editor Preferences -> Keyboard Shortcuts -> Material Editor�Spawn Nodes**.
+There are a number of Keyboard shortcuts that you can use to help you quickly place down Material Expressions that are most commonly used. You can change what key spawns what Material Expression by going to **Edit -> Editor Preferences -> Keyboard Shortcuts -> Material Editor Spawn Nodes**.
 
 ![](Placing_Functions_Setting_Shortcuts.gif)(convert:false)
 
diff --git a/Engine/Documentation/Source/Engine/Rendering/Materials/HowTo/Subsurface_Scattering/Subsurface_Scattering.INT.udn b/Engine/Documentation/Source/Engine/Rendering/Materials/HowTo/Subsurface_Scattering/Subsurface_Scattering.INT.udn
--- a/Engine/Documentation/Source/Engine/Rendering/Materials/HowTo/Subsurface_Scattering/Subsurface_Scattering.INT.udn
+++ b/Engine/Documentation/Source/Engine/Rendering/Materials/HowTo/Subsurface_Scattering/Subsurface_Scattering.INT.udn
@@ -60,9 +60,9 @@ Continuing with our Material from above that has Subsurface Scattering enabled,
 	| **Subsurface _Amount**| 0.5|
 	| **Subsurface_Color**|r:1.0, b:0, g:0|
 1. Once the Material Nodes have been named and filled in with the correct data, it is then time to connect the Material nodes to the Main Material Node. 
-	To wire a node to the main Material node, use the **Left Mouse Button** to click on the **White Circle** that is on the **Right Side** of the node. 
+	To wire a node to the main Material node, use the **Left Mouse Button** to click on the White Circle that is on the Right Side of the node. 
 	With the **Left Mouse Button** still held down, drag to the **Right** and the connection wire will appear. 
-	Keep dragging to the **Right** with the  **Left Mouse Button** still held until you come to the input you want to use on the Main Material node. 
+	Keep dragging to the Right with the  **Left Mouse Button** still held until you come to the input you want to use on the Main Material node. 
 	Once your mouse is over the input you want to use, release the  **Left Mouse Button** to connect.
 
 	![](SSM_Connecting_Nodes.gif)(convert:false)
@@ -99,7 +99,7 @@ Continuing with our Material from above that has Subsurface Scattering enabled,
 
 1. With all of the lights added, we can now add a Static Mesh to test with. For this example, we will be using the following mesh. 
 
-	**StaticMesh'/Game/StarterContent/Props/SM_Chair.SM_Chair'**
+	StaticMesh'/Game/StarterContent/Props/SM_Chair.SM_Chair'
 
 	[region:tip]
 	If you cannot find the SM_Chair mesh, it is probably because you did not include the Starter Content with your project. 
