#UE4doc #docs:
Changed some terms to match the editor/pictures
Unbolded a few terms

diff --git a/Engine/Documentation/Source/Engine/AI/EnvironmentQuerySystem/NodeReference/NodeReference.INT.udn b/Engine/Documentation/Source/Engine/AI/EnvironmentQuerySystem/NodeReference/NodeReference.INT.udn
--- a/Engine/Documentation/Source/Engine/AI/EnvironmentQuerySystem/NodeReference/NodeReference.INT.udn
+++ b/Engine/Documentation/Source/Engine/AI/EnvironmentQuerySystem/NodeReference/NodeReference.INT.udn
@@ -54,7 +54,7 @@ This is a quick reference for the Nodes and Tests available in the Environment Q
 
 ### Actors of Class
 
-Finds all of the Actors of a class within the Search Radius around the Context. The Actors are returned as the Items to be used in the tests.
+Finds all of the Actors of a class within the **Search Radius** around the **Search Center**. The Actors are returned as the Items to be used in the tests.
 
 ![](gen_actor.png)(convert:false)
 
@@ -70,7 +70,7 @@ Finds all of the Actors of a class within the Search Radius around the Context.
 
 ### On circle
 
-Radiates traces out from the Context, returning the hits or the edge of the radius (of the generated Circle) as Items.
+Radiates traces out from the **Circle Center**, returning the hits or the edge of the radius (of the generated Circle) as Items.
 
 ![](gen_circle.png)(convert:false)
 
@@ -79,9 +79,9 @@ Radiates traces out from the Context, returning the hits or the edge of the radi
 [REGION:raw]
 | Property | %spacer% | Description |
 | --- | --- | :--- |
-| **Circle Radius** || The radius of the circle that will expand out from the Circle Center Context. |
+| **Circle Radius** || The radius of the circle that will expand out from the **Circle Center**. |
 | **Space Between** || The space (in cm) between the items generated on the outer edge of the circle. |
-| **Arc Direction** || This will enable the use of the Arc Direction feature which will enable you to focus the arc in a direction based either on a rotation of a context or the vector formed between to locations. |
+| **Arc Direction** || This will enable the use of the **Arc Direction** feature which will enable you to focus the arc in a direction based either on a rotation of a context or the vector formed between to locations. |
 | **Arc Angle** || The angle of the arc measured in degrees. |
 | **Circle Center** || The context used as the center for the circle. |
 | **Trace Data** || The type of trace to use. |
@@ -90,7 +90,7 @@ Radiates traces out from the Context, returning the hits or the edge of the radi
 
 ### Pathing Grid
 
-Like Grid, but this paths to different points on the NavMesh, only returning Items that are less than or equal to the Max Distance away.
+Like Grid, but this paths to different points on the NavMesh, only returning Items that are less than or equal to the **Max Distance** away.
 
 ![](gen_pathing.png)(convert:false)
 
@@ -101,7 +101,7 @@ Like Grid, but this paths to different points on the NavMesh, only returning Ite
 | --- | --- | :--- |
 | **Max Distance** || How far to path away from the Querier. |
 | **Distance Between** || The distance between the generated items. |
-| **Path to Item** || Whether to path **to** or **from** the item. |
+| **Path to Item** || Whether to path to or from the item. |
 | **Generate Around** || The context to generate the pathing grid around. |
 | **Navigation Filter** || The Navigation Filter to use.  | 
 [/REGION]
@@ -128,7 +128,7 @@ Generates a grid of Items around the Querier.
 
 ### Distance
 
-This will return the direct distance between the Item and the chosen Context. if the Context is more than one location, it averages the results of each distance check.
+This will return the direct distance between the Item and the chosen **Distance To**. If the **Distance To** is more than one location, it averages the results of each distance check.
 
 ![](test_dist.png)
 
