#UE4doc #docs: ##Overview removed as according to the instructional guide.

diff --git a/Engine/Documentation/Source/Engine/Rendering/ParticleSystems/Reference/Modules/Color/ColorModules.INT.udn b/Engine/Documentation/Source/Engine/Rendering/ParticleSystems/Reference/Modules/Color/ColorModules.INT.udn
--- a/Engine/Documentation/Source/Engine/Rendering/ParticleSystems/Reference/Modules/Color/ColorModules.INT.udn
+++ b/Engine/Documentation/Source/Engine/Rendering/ParticleSystems/Reference/Modules/Color/ColorModules.INT.udn
@@ -5,8 +5,6 @@ Description:Color modules impact the color of the emitted particles.
 
 [TOC(start:2 end:3)]
 
-## Color Modules
-
 Color modules impact the color of the emitted particles.
 
 [REGION:note]
