#UE4doc #docs:
Removed Overview header

diff --git a/Engine/Documentation/Source/Engine/AI/EnvironmentQuerySystem/QuickStart/QuickStart.INT.udn b/Engine/Documentation/Source/Engine/AI/EnvironmentQuerySystem/QuickStart/QuickStart.INT.udn
--- a/Engine/Documentation/Source/Engine/AI/EnvironmentQuerySystem/QuickStart/QuickStart.INT.udn
+++ b/Engine/Documentation/Source/Engine/AI/EnvironmentQuerySystem/QuickStart/QuickStart.INT.udn
@@ -42,7 +42,6 @@ Description:This Quick Start Guide covers creating a Behavior Tree that uses mul
 [/OBJECT]
 [/VAR]
 
-## Overview
 
 This Quick Start will teach you the basics of using the **Environment Query System**. In it, you will use Environment Queries to find Actors and Locations, then, using a series of tests, you'll be given the "best" option. Then, using a Behavior tree, you'll have the AI flee from the Player Pawn (or Pawns) to a hiding spot.
 
