#UE4doc #docs: ##Overview removed as according to the instructional guide.

diff --git a/Engine/Documentation/Source/Engine/Rendering/ParticleSystems/Reference/Modules/Event/EventModules.INT.udn b/Engine/Documentation/Source/Engine/Rendering/ParticleSystems/Reference/Modules/Event/EventModules.INT.udn
--- a/Engine/Documentation/Source/Engine/Rendering/ParticleSystems/Reference/Modules/Event/EventModules.INT.udn
+++ b/Engine/Documentation/Source/Engine/Rendering/ParticleSystems/Reference/Modules/Event/EventModules.INT.udn
@@ -5,8 +5,6 @@ Description:Event modules allow you to generate events based on particles themse
 
 [TOC(start:2 end:3)]
 
-## Event Modules 
-
 Event modules allow you to generate events based on particles interacting with themselves, each other, or the world, and then listen for those events and cause reactions to them on an inter-particle system level. A good example is generating a named collision event whenever a particle collides with the world, and then spawning particles wherever those collisions occur.
 
 
