#UE4doc #docs: ##Overview removed as according to the instructional guide.

diff --git a/Engine/Documentation/Source/Engine/Rendering/ParticleSystems/Reference/Modules/Kill/KillModules.INT.udn b/Engine/Documentation/Source/Engine/Rendering/ParticleSystems/Reference/Modules/Kill/KillModules.INT.udn
--- a/Engine/Documentation/Source/Engine/Rendering/ParticleSystems/Reference/Modules/Kill/KillModules.INT.udn
+++ b/Engine/Documentation/Source/Engine/Rendering/ParticleSystems/Reference/Modules/Kill/KillModules.INT.udn
@@ -5,8 +5,6 @@ Description:Kill modules will kill a given particle if it meets the rules define
 
 [TOC(start:2 end:3)]
 
-## Kill Modules 
-
 Kill modules will kill a given particle if it meets the rules defined by specific implementations. 
 
 
