#UE4doc #docs: ##Overview removed as according to the instructional guide.

diff --git a/Engine/Documentation/Source/Engine/Rendering/ParticleSystems/Reference/Modules/Lifetime/LifetimeModules.INT.udn b/Engine/Documentation/Source/Engine/Rendering/ParticleSystems/Reference/Modules/Lifetime/LifetimeModules.INT.udn
--- a/Engine/Documentation/Source/Engine/Rendering/ParticleSystems/Reference/Modules/Lifetime/LifetimeModules.INT.udn
+++ b/Engine/Documentation/Source/Engine/Rendering/ParticleSystems/Reference/Modules/Lifetime/LifetimeModules.INT.udn
@@ -5,10 +5,8 @@ Description:The Lifetime module is used to set the initial lifetime of a particl
 
 [TOC(start:2 end:3)]
 
-## Lifetime Modules 
-
 (#Lifetime)
-### Lifetime
+## Lifetime
 
 ![psys_lifetime.gif](psys_lifetime.gif)(convert:false)
 
@@ -23,7 +21,7 @@ The Lifetime module is used to set the initial lifetime of a particle at spawn t
 In Spawn, the module retrieves the appropriate values from the distribution using the current emitter time. This is then added to the Particle.OneOverMaxLifetime field, to allow for multiple lifetime modules to be applied.
 
 
-### Lifetime (Seeded)
+## Lifetime (Seeded)
 
 The Lifetime (Seeded) module is identical to the [Lifetime](#Lifetime) module in that it sets the lifetime of a particle at spawn time; however, this module allows you to specify seed information to be used when choosing distribution values in order to provide a more consistent effect from the module each time the emitter is used. It contains the following member(s):
 
diff --git a/Engine/Documentation/Source/Engine/Rendering/ParticleSystems/Reference/Modules/Location/LocationModules.INT.udn b/Engine/Documentation/Source/Engine/Rendering/ParticleSystems/Reference/Modules/Location/LocationModules.INT.udn
--- a/Engine/Documentation/Source/Engine/Rendering/ParticleSystems/Reference/Modules/Location/LocationModules.INT.udn
+++ b/Engine/Documentation/Source/Engine/Rendering/ParticleSystems/Reference/Modules/Location/LocationModules.INT.udn
@@ -5,11 +5,9 @@ Description:Location modules affect the position of particles at various times t
 
 [TOC(start:2 end:3)]
 
-## Location Modules 
-
 Location modules affect the position of particles at various times throughout their lives.
 
-### Initial Location
+## Initial Location
 
 ![psys_initial_location.gif](psys_initial_location.gif)(convert:false)
 
@@ -25,7 +23,7 @@ The Initial Location module is used to set the initial location of a particle at
 In Spawn, the module retrieves the appropriate values from the distribution using the current emitter time. If the particle emitter does not have the _Use Local Space_ flag, this value is transformed into world-space. The Particle.Location field then has this value added to it.
 
 
-### Initial Loc (Seed)
+## Initial Loc (Seed)
 
 The Initial Loc (Seeded) module is identical to the [Initial Location](#InitialLocation) module in that it sets the initial location of a particle at spawn time; however, this module allows you to specify seed information to be used when choosing distribution values in order to provide a more consistent effect from the module each time the emitter is used. It contains the following member(s):
 
@@ -44,7 +42,7 @@ The Initial Loc (Seeded) module is identical to the [Initial Location](#InitialL
 In Spawn, the module retrieves the appropriate values from the distribution using the current emitter time. If the particle emitter does not have the _Use Local Space_ flag, this value is transformed into world-space. The Particle.Location field then has this value added to it.
 
 
-### World Offset
+## World Offset
 
 The World Offset module is used to offset the initial location of a particle. The offset is in world space but respects the _Use Local Space_ flag through the life of the particle. This means the particle will always spawn offset in world space regardless of the orientation of the emitter, but will keep that offset relative to the emitter throughout its life. It contains the following member(s):
 
@@ -57,7 +55,7 @@ The World Offset module is used to offset the initial location of a particle. Th
 
 
 
-### Bone/Socket Location
+## Bone/Socket Location
 
 The Bone/Socket Location module allows particle to be spawned directly at the locations of bones or sockets of a SkeletalMesh. It provides the following properties:
 
@@ -114,7 +112,7 @@ The Bone/Socket Location module allows particle to be spawned directly at the lo
 
 -------------End Bone/Socket Location Includes--------------->
 
-### Direct Location
+## Direct Location
 
 The Direct Location module is used to directly set the location of a particle. It contains the following member(s):
 
@@ -128,7 +126,7 @@ The Direct Location module is used to directly set the location of a particle. I
 
 
 
-### Emitter Init Loc
+## Emitter Init Loc
 
 The Emitter InitLoc module is used to set the initial location of a particle to the position of a particle from another emitter (in the same particle system). It contains the following member(s):
 
@@ -159,7 +157,7 @@ The Emitter InitLoc module is used to set the initial location of a particle to
 -------------End Emitter Init Loc Includes--------------->
 
 
-### Emitter Direct Loc
+## Emitter Direct Loc
 
 The Emitter DirectLoc module is used to set the location of a particle to the position of a particle from another emitter (in the same particle system) for the particle's entire life. It contains the following member(s):
 
@@ -172,7 +170,7 @@ The Emitter DirectLoc module is used to set the location of a particle to the po
 The particle used will be the one at the same index as the particle having its location set.
 
 
-### Cylinder
+## Cylinder
 
 ![psys_cylinder.gif](psys_cylinder.gif)(convert:false)
 
@@ -208,7 +206,7 @@ The Cylinder module is used to set the initial location of particles within a cy
 
 -------------End Cylinder Includes--------------->
 
-### Cylinder (Seeded)
+## Cylinder (Seeded)
 
 The Cylinder (Seeded) module is identical to the [Cylinder](#Cylinder) module in that it sets the initial location of particles within a cylinder shape; however, this module allows you to specify seed information to be used when choosing distribution values in order to provide a more consistent effect from the module each time the emitter is used. It contains the following member(s):
 
@@ -231,7 +229,7 @@ The Cylinder (Seeded) module is identical to the [Cylinder](#Cylinder) module in
 
 
 
-### Sphere
+## Sphere
 
 ![psys_sphere.gif](psys_sphere.gif)(convert:false)
 
@@ -253,7 +251,7 @@ The Sphere module is used to set the initial location of particles within a sphe
 
 
 
-### Sphere (Seeded)
+## Sphere (Seeded)
 
 The Sphere (Seeded) module is identical to the [Sphere](#Sphere) module in that it sets the initial location of particles within a sphere shape; however, this module allows you to specify seed information to be used when choosing distribution values in order to provide a more consistent effect from the module each time the emitter is used. It contains the following member(s):
 
@@ -272,7 +270,7 @@ The Sphere (Seeded) module is identical to the [Sphere](#Sphere) module in that
 
 
 
-### Triangle
+## Triangle
 
 ![](Triangle.png)
 
@@ -287,7 +285,7 @@ The Triangle module allows the user to specify the shape and location of a trian
 |**Thickness**|A float distribution that provides the thickness of the triangle, turning it into a 3D wedge. Particles will be emitted from within the 3-dimensional area.|
 
 
-### Skeletal Mesh Vertex/Surface Location (Skel Vert/Surf Loc)
+## Skeletal Mesh Vertex/Surface Location (Skel Vert/Surf Loc)
 
 The Skel Vert/Surf Loc module allows for particles to be emitted from a skeletal mesh's vertices or polygonal surfaces. It provides the following properties:
 
@@ -334,7 +332,7 @@ The Skel Vert/Surf Loc module allows for particles to be emitted from a skeletal
 
 
 
-### Source Movement
+## Source Movement
 
 The Source Movement module is used to offset the location of particles based on the movement of the source (i.e. the emitter). It contains the following properties:
 
