#UE4doc #docs: Adjusted bold text in accordance with the Instructional Guide

diff --git a/Engine/Documentation/Source/Engine/Rendering/Materials/LightingModels/SubSurface/SSSLightingModel.INT.udn b/Engine/Documentation/Source/Engine/Rendering/Materials/LightingModels/SubSurface/SSSLightingModel.INT.udn
--- a/Engine/Documentation/Source/Engine/Rendering/Materials/LightingModels/SubSurface/SSSLightingModel.INT.udn
+++ b/Engine/Documentation/Source/Engine/Rendering/Materials/LightingModels/SubSurface/SSSLightingModel.INT.udn
@@ -15,9 +15,9 @@ The two components of subsurface lighting are a wrap-around `N dot H` term and a
 
 ## Material Input Channels
 
-The **SubsurfaceColor** input channel defines the color inside the material which affects SSS lighting.
+The SubsurfaceColor input channel defines the color inside the material which affects SSS lighting.
 
-The **Opacity** channel of the material takes on a slightly different meaning when the material is using subsurface scattering and the **_MLM_Subsurface_** shading model. Since these types of surfaces are completely opaque, the **Opacity**, in this case, controls how dense the material is when it scatters light as well as:
+The Opacity channel of the material takes on a slightly different meaning when the material is using subsurface scattering and the _MLM_Subsurface_ shading model. Since these types of surfaces are completely opaque, the Opacity, in this case, controls how dense the material is when it scatters light as well as:
 
 * How much the normal affects the subsurface lighting, a more opaque material gets more normal influence.
 * How far lighting makes it through the material due to self-shadowing, a smaller opacity causes light to travel further.
@@ -25,7 +25,7 @@ The **Opacity** channel of the material takes on a slightly different meaning wh
 
 
 [REGION:tip]
-Be sure to set the **Opacity** for any subsurface materials to a reasonable value, e.g. 0.1. The default opacity is 1, which does not produce a very convincing subsurface-type effect.
+Be sure to set the Opacity for any subsurface materials to a reasonable value, e.g. 0.1. The default opacity is 1, which does not produce a very convincing subsurface-type effect.
 [/REGION]
 
 ## Limitations
