#UE4doc #docs: Altered images to match the instructional guide: green -> yellow.

diff --git a/Engine/Documentation/Source/Engine/Rendering/ParticleSystems/Reference/Modules/Orbit/OrbitModules.INT.udn b/Engine/Documentation/Source/Engine/Rendering/ParticleSystems/Reference/Modules/Orbit/OrbitModules.INT.udn
--- a/Engine/Documentation/Source/Engine/Rendering/ParticleSystems/Reference/Modules/Orbit/OrbitModules.INT.udn
+++ b/Engine/Documentation/Source/Engine/Rendering/ParticleSystems/Reference/Modules/Orbit/OrbitModules.INT.udn
@@ -5,12 +5,10 @@ Description:Orbit modules allow for rendering particle sprites offset/rotated aw
 
 [TOC(start:2 end:3)]
 
-## Orbit Modules
-
 Orbit modules allow for rendering particle sprites offset/rotated away from the actual particle center.
 
 
-### Orbit
+## Orbit
 
 ![psys_orbit.gif](psys_orbit.gif)(convert:false)
 
