#UE4doc #docs: Adjusted bold text in accordance with the Instructional Guide

diff --git a/Engine/Documentation/Source/Engine/Rendering/Materials/LightingModels/SubSurfaceProfile/SSSProfileLightingModel.INT.udn b/Engine/Documentation/Source/Engine/Rendering/Materials/LightingModels/SubSurfaceProfile/SSSProfileLightingModel.INT.udn
--- a/Engine/Documentation/Source/Engine/Rendering/Materials/LightingModels/SubSurfaceProfile/SSSProfileLightingModel.INT.udn
+++ b/Engine/Documentation/Source/Engine/Rendering/Materials/LightingModels/SubSurfaceProfile/SSSProfileLightingModel.INT.udn
@@ -5,10 +5,10 @@ Description:Description and technical details of the subsurface profile shading
 
 [TOC(start:2)]
 
-## Overview
+
 
 The ability to rendering realistic looking human skin is a must for any modern Video Game engine. To fulfill this need, Unreal Engine 4 (UE4) now offers a shading method specifically for skin or wax like surfaces called **Subsurface Profile**. 
-While the Subsurface Profile shading model has similar properties to the **Subsurface** shading model, its key difference is in how it renders. 
+While the Subsurface Profile shading model has similar properties to the Subsurface shading model, its key difference is in how it renders. 
 SubsurfaceProfile uses a rendering approach that is based in Screen Space because this helps to better display the subtle subsurface effects seen in human skin where backscattering is a secondary effect only seen in few cases like ears for example.
 The following document will cover what a Subsurface Profile is and how you can use them in your work. 
 
@@ -16,7 +16,7 @@ The following document will cover what a Subsurface Profile is and how you can u
 [![](Results.png)(w:920 convert:false)](Results.png)
 [/REGION]
 
-The **Left Side** shows a material not using the SubsurfaceProfile shading model. The **Right Side** shows an asset using the SubsurfaceProfile shading model. 
+The Left Side shows a material not using the SubsurfaceProfile shading model. The Right Side shows an asset using the SubsurfaceProfile shading model. 
 
 [Region:note]
 Special thanks to [REDACTED_PERSON] and his company [Infinite Realities] (http://ir-ltd.net ) for providing the 3D scanned head model and assistance. 
@@ -100,7 +100,7 @@ For example, when inputting a value of 0 into the Opacity input, Subsurface scat
 Inputting a value of 1 will show the full extent of the Subsurface Scattering. 
 In order to better control where the Subsurface scattering is stronger or weaker, it is best to use a mask texture.  
 Areas in the mask texture that have values closer to 1, or white, will have the strongest Subsurface scattering effect while areas that are closer to 0, or black, will have less of an effect. 
-Adjusting the Subsurface Color will help to compensate if areas get too dark and remember using **brighter colors** results in more subsurface scattering. 
+Adjusting the Subsurface Color will help to compensate if areas get too dark and remember using brighter colors results in more subsurface scattering. 
 
 Here you can see how the mask can be used to render two surface types with one material. Note that the transition is soft and not limited to triangles borders.
 
