#UE4doc #docs: fixed multiple formatting errors that involve the changes to ##Overview. Errors were found in the modules section.

diff --git a/Engine/Documentation/Source/Engine/Rendering/ParticleSystems/Reference/Modules/Acceleration/AccelerationModules.INT.udn b/Engine/Documentation/Source/Engine/Rendering/ParticleSystems/Reference/Modules/Acceleration/AccelerationModules.INT.udn
--- a/Engine/Documentation/Source/Engine/Rendering/ParticleSystems/Reference/Modules/Acceleration/AccelerationModules.INT.udn
+++ b/Engine/Documentation/Source/Engine/Rendering/ParticleSystems/Reference/Modules/Acceleration/AccelerationModules.INT.udn
@@ -8,7 +8,7 @@ Description:These modules provide the ability to apply a change in velocity over
 These modules provide the ability to apply acceleration, or a change in velocity over time, to particles.
 
 
-### Acceleration
+## Acceleration
 
 ![psys_acceleration.gif](psys_acceleration.gif)(convert:false)
 
@@ -26,7 +26,7 @@ This module will add a vector parameter to the particle payload data, UsedAccele
 Each frame, the current and base velocity values of the particle is then updated using the formula (velocity += acceleration * DeltaTime) where DeltaTime is the amount of time passed since the last frame.
 
 
-### Const Acceleration
+## Const Acceleration
 
 
 ![psys_acceleration.gif](psys_acceleration.gif)(convert:false)
@@ -46,7 +46,7 @@ This module will add a vector parameter to the particle payload data, UsedAccele
 
 Each frame, the current and base velocity values of the particle is then updated using the formula (velocity += acceleration * DeltaTime) where DeltaTime is the amount of time passed since the last frame.
 
-### Drag 
+## Drag 
 
 ![phys_drag.gif](phys_drag.gif)(w:430 convert:false)
 
@@ -57,7 +57,7 @@ Drag applies a constant resistance to motion to each particle. The coefficient v
 |[REGION:tablesection]Drag[/REGION]||
 | **Coefficient** | A float distribution defining the drag coefficient, or the amount of resistance to motion, that will be applied to each particle.|
 
-### Drag Scale/Life
+## Drag Scale/Life
 
 Drag Scale/Life is a multiplier value that works with an existing Drag module, allowing, for example, for the drag to be scaled via a curve across the life of the particle.
 
@@ -70,7 +70,7 @@ Drag Scale/Life is a multiplier value that works with an existing Drag module, a
 |[REGION:tablesection]Drag Scale/Life[/REGION]||
 | **Drag Scale** | A float distribution defining the scaling value that will be applied to the existing particle drag coefficient.|
 
-### AccelerationOverLife
+## AccelerationOverLife
 
 Sets the acceleration of the particle over its lifetime. It contains the following member(s):
 
diff --git a/Engine/Documentation/Source/Engine/Rendering/ParticleSystems/Reference/Modules/Attractor/AttractorModules.INT.udn b/Engine/Documentation/Source/Engine/Rendering/ParticleSystems/Reference/Modules/Attractor/AttractorModules.INT.udn
--- a/Engine/Documentation/Source/Engine/Rendering/ParticleSystems/Reference/Modules/Attractor/AttractorModules.INT.udn
+++ b/Engine/Documentation/Source/Engine/Rendering/ParticleSystems/Reference/Modules/Attractor/AttractorModules.INT.udn
@@ -5,15 +5,13 @@ Description:Modules for causing particles to be attracted toward specific points
 
 [TOC(start:2 end:3)]
 
-## Attractor Modules
-
 These modules implement methods of causing particles to be attracted toward a specific position in space, which can be defined in the form of a point, line, or another particle's position. They can even be combined to create more complex effects.
 
 This is the result of using a point attractor in combination with a line attractor with its strength animated over the life of the particle to create the swirling effect.
 
 ![psys_attractor.gif](psys_attractor.gif)(convert:false)
 
-### Attractor Line
+## Attractor Line
 
 The line attractor allows for having particles be drawn to a line in 3D space.
 
@@ -27,7 +25,7 @@ The line attractor allows for having particles be drawn to a line in 3D space.
 
 
 
-### Particle Attractor
+## Particle Attractor
 
 ![psys_attractor_particle.gif](psys_attractor_particle.gif)(convert:false)
 
@@ -60,7 +58,7 @@ The particle attractor allows for having particles be drawn to the particles emi
 --------End Attractor Particle Includes---------->
 
 
-### Point Attractor
+## Point Attractor
 
 ![psys_attractor_point.gif](psys_attractor_point.gif)(convert:false)
 
diff --git a/Engine/Documentation/Source/Engine/Rendering/ParticleSystems/Reference/Modules/Beam/BeamModules.INT.udn b/Engine/Documentation/Source/Engine/Rendering/ParticleSystems/Reference/Modules/Beam/BeamModules.INT.udn
--- a/Engine/Documentation/Source/Engine/Rendering/ParticleSystems/Reference/Modules/Beam/BeamModules.INT.udn
+++ b/Engine/Documentation/Source/Engine/Rendering/ParticleSystems/Reference/Modules/Beam/BeamModules.INT.udn
@@ -8,7 +8,7 @@ Description:These modules are used to configure the behavior of emitters using t
 These modules are used to configure or modify the behavior of emitters using the Beam TypeData module.
 
 
-### Beam Modifier
+## Beam Modifier
 
 The Beam Modifier module allows either the Source or Target of the beam to be modified. It provides the following properties:
 
@@ -79,7 +79,7 @@ The Beam Modifier module allows either the Source or Target of the beam to be mo
 
 
 
-### Noise
+## Noise
 
 ![psys_beam_noise.gif](psys_beam_noise.gif)(convert:false)
 
@@ -120,7 +120,7 @@ Currently, low-frequency noise is the only noise supported.
 
 -----------------------End Beam Noise Includes------------->
 
-### Source
+## Source
 
 The Beam Source module implements a single source for a beam emitter. (If a source module is not present in a beam emitter, the emitter location itself is used as the source.) It provides the following properties:
 
@@ -170,7 +170,7 @@ The Beam Source module implements a single source for a beam emitter. (If a sour
 -----------------------End Beam Source Includes------------->
 
 
-### Target
+## Target
 
 The Beam Target module implements a single target for a beam emitter. (If a target module is not present in a beam emitter, the emitter assumes that the beam is to be used directionally.) It provides the following properties:
 
diff --git a/Engine/Documentation/Source/Engine/Rendering/ParticleSystems/Reference/Modules/Camera/CameraModules.INT.udn b/Engine/Documentation/Source/Engine/Rendering/ParticleSystems/Reference/Modules/Camera/CameraModules.INT.udn
--- a/Engine/Documentation/Source/Engine/Rendering/ParticleSystems/Reference/Modules/Camera/CameraModules.INT.udn
+++ b/Engine/Documentation/Source/Engine/Rendering/ParticleSystems/Reference/Modules/Camera/CameraModules.INT.udn
@@ -8,7 +8,7 @@ Description:These modules modify the emitter behavior with respect to the camera
 These modules modify the emitter behavior with respect to the camera.
 
 
-### Camera Offset
+## Camera Offset
 
 ![CameraOffset.png](CameraOffset.png)
 
diff --git a/Engine/Documentation/Source/Engine/Rendering/ParticleSystems/Reference/Modules/Color/ColorModules.INT.udn b/Engine/Documentation/Source/Engine/Rendering/ParticleSystems/Reference/Modules/Color/ColorModules.INT.udn
--- a/Engine/Documentation/Source/Engine/Rendering/ParticleSystems/Reference/Modules/Color/ColorModules.INT.udn
+++ b/Engine/Documentation/Source/Engine/Rendering/ParticleSystems/Reference/Modules/Color/ColorModules.INT.udn
@@ -15,7 +15,7 @@ Color modules impact the color of the emitted particles.
 	Color is calculated on a full floating point range. Values beyond 1 will often result in a glow effect.
 [/REGION]
 
-### Initial Color
+## Initial Color
 
 ![psys_initial_color.gif](psys_initial_color.gif)(convert:false)
 
@@ -32,7 +32,7 @@ The Initial Color module sets the color of a particle at spawn time. The particl
 In Spawn, the module retrieves the appropriate values from the distributions using the emitter time, setting the Particle.Color and Particle.BaseColor values to it directly.
 
 
-### Init Color (Seed)
+## Init Color (Seed)
 
 The Init Color (Seed) module is identical to the [Initial Color](#InitialColor) module in that it sets the initial color of a particle at spawn time; however, this module allows you to specify seed information to be used when choosing distribution values in order to provide a more consistent effect from the module each time the emitter is used. It contains the following member(s):
 
@@ -55,7 +55,7 @@ In Spawn, the module retrieves the appropriate values from the distributions usi
 
 
 
-### Color Over Life
+## Color Over Life
 
 ![psys_color_life.gif](psys_color_life.gif)(convert:false)
 
@@ -78,7 +78,7 @@ In Spawn, the module retrieves the appropriate values from the distributions usi
 In Update, the module retrieves the appropriate values from the distributions using the particle time, and sets the Particle.Color value to it.
 
 
-### Scale Color/Life
+## Scale Color/Life
 
 The Scale Color/Life module is used to scale the color of a particle over its lifetime. It contains the following member(s):
 
diff --git a/Engine/Documentation/Source/Engine/Rendering/ParticleSystems/Reference/Modules/Event/EventModules.INT.udn b/Engine/Documentation/Source/Engine/Rendering/ParticleSystems/Reference/Modules/Event/EventModules.INT.udn
--- a/Engine/Documentation/Source/Engine/Rendering/ParticleSystems/Reference/Modules/Event/EventModules.INT.udn
+++ b/Engine/Documentation/Source/Engine/Rendering/ParticleSystems/Reference/Modules/Event/EventModules.INT.udn
@@ -8,7 +8,7 @@ Description:Event modules allow you to generate events based on particles themse
 Event modules allow you to generate events based on particles interacting with themselves, each other, or the world, and then listen for those events and cause reactions to them on an inter-particle system level. A good example is generating a named collision event whenever a particle collides with the world, and then spawning particles wherever those collisions occur.
 
 
-### Event Generator
+## Event Generator
 
 The Event Generator module will generate an event (or multiple events) based on conditions you specify. This module has a single **Events** array that contains a list of all events you want the emitter to be able to generate. Each item in the _Events_ array has the following properties:
 
@@ -45,7 +45,7 @@ The Event Generator module will generate an event (or multiple events) based on
 -------------End Event Generator Includes--------------->
 
 
-### Event Receiver Kill All (EventRcvr KillAll)
+## Event Receiver Kill All (EventRcvr KillAll)
 
 Listens for a named event and then kills all the particles of the emitter.
 
@@ -78,7 +78,7 @@ Listens for a named event and then kills all the particles of the emitter.
 
 -------------End Event Generator Includes--------------->
 
-### Event Receiver Spawn (EventRcvr Spawn)
+## Event Receiver Spawn (EventRcvr Spawn)
 
 Listens for a named event and then spawns particles based on what events fire.
 
diff --git a/Engine/Documentation/Source/Engine/Rendering/ParticleSystems/Reference/Modules/Kill/KillModules.INT.udn b/Engine/Documentation/Source/Engine/Rendering/ParticleSystems/Reference/Modules/Kill/KillModules.INT.udn
--- a/Engine/Documentation/Source/Engine/Rendering/ParticleSystems/Reference/Modules/Kill/KillModules.INT.udn
+++ b/Engine/Documentation/Source/Engine/Rendering/ParticleSystems/Reference/Modules/Kill/KillModules.INT.udn
@@ -8,7 +8,7 @@ Description:Kill modules will kill a given particle if it meets the rules define
 Kill modules will kill a given particle if it meets the rules defined by specific implementations. 
 
 
-### Kill Box
+## Kill Box
 
 ![](KillBox.png)
 
@@ -27,7 +27,7 @@ The Kill box module is used to kill particles when they move outside of a define
 If the 3D preview mode is enabled, the wire box will be drawn in the Cascade preview window.
 
 
-### Kill Height
+## Kill Height
 
 ![](KillHeight.png)
 
diff --git a/Engine/Documentation/Source/Engine/Rendering/ParticleSystems/Reference/Modules/Orientation/OrientationModules.INT.udn b/Engine/Documentation/Source/Engine/Rendering/ParticleSystems/Reference/Modules/Orientation/OrientationModules.INT.udn
--- a/Engine/Documentation/Source/Engine/Rendering/ParticleSystems/Reference/Modules/Orientation/OrientationModules.INT.udn
+++ b/Engine/Documentation/Source/Engine/Rendering/ParticleSystems/Reference/Modules/Orientation/OrientationModules.INT.udn
@@ -5,11 +5,9 @@ Description:These modules control the orientation of particles, preventing rotat
 
 [TOC(start:2 end:3)]
 
-## Orientation Modules
-
 These modules control the orientation of particles, preventing rotation on certain axes.
 
-### Axis Lock
+## Axis Lock
 
 The Axis Lock module is used to lock the particle facing to a given axis. Currently, this is a sprite-based module only. It contains the following member(s): 
 
diff --git a/Engine/Documentation/Source/Engine/Rendering/ParticleSystems/Reference/Modules/Parameter/ParameterModules.INT.udn b/Engine/Documentation/Source/Engine/Rendering/ParticleSystems/Reference/Modules/Parameter/ParameterModules.INT.udn
--- a/Engine/Documentation/Source/Engine/Rendering/ParticleSystems/Reference/Modules/Parameter/ParameterModules.INT.udn
+++ b/Engine/Documentation/Source/Engine/Rendering/ParticleSystems/Reference/Modules/Parameter/ParameterModules.INT.udn
@@ -5,11 +5,9 @@ Description:Modules that pass data out of a particle system to other aspects, su
 
 [TOC(start:2 end:3)]
 
-## Parameter Modules
-
 Parameter modules exist to pass data out of a particle system so that it can control other aspects, such as materials.
 
-### Dynamic
+## Dynamic
 
 The Dynamic module is used to pass four scalar values to the material(s) used by the emitter allowing the emitter to control material effects.
 
@@ -53,7 +51,7 @@ The Dynamic module is used to pass four scalar values to the material(s) used by
 -------------Dynamic Includes--------------->
 
 
-### Dynamic (Seed)
+## Dynamic (Seed)
 
 The Dynamic (Seed) module is identical to the [Dynamic Parameter](#Dynamic) module in that it allows the emitter to pass values to the material(s); however, this module allows you to specify seed information to be used when choosing distribution values in order to provide a more consistent effect from the module each time the emitter is used. It contains the following member(s):
 
diff --git a/Engine/Documentation/Source/Engine/Rendering/ParticleSystems/Reference/Modules/Required/Required.INT.udn b/Engine/Documentation/Source/Engine/Rendering/ParticleSystems/Reference/Modules/Required/Required.INT.udn
--- a/Engine/Documentation/Source/Engine/Rendering/ParticleSystems/Reference/Modules/Required/Required.INT.udn
+++ b/Engine/Documentation/Source/Engine/Rendering/ParticleSystems/Reference/Modules/Required/Required.INT.udn
@@ -8,7 +8,7 @@ Description:Modules that are included along with an emitter and required for its
 The following classes and modules always exist when a new emitter is created. They will always be present and cannot be deleted.
 
 
-### Required Module
+## Required Module
 
 Every particle emitter contains this module which contains all the properties that are required by a particle emitter. This includes the following:
 
@@ -58,7 +58,7 @@ Every particle emitter contains this module which contains all the properties th
 
 
 
-### Spawn Module
+## Spawn Module
 
 Every particle emitter contains this module which contains all the properties that determine how particles are spawned. This includes the following:
 
diff --git a/Engine/Documentation/Source/Engine/Rendering/ParticleSystems/Reference/Modules/Rotation/RotationModules.INT.udn b/Engine/Documentation/Source/Engine/Rendering/ParticleSystems/Reference/Modules/Rotation/RotationModules.INT.udn
--- a/Engine/Documentation/Source/Engine/Rendering/ParticleSystems/Reference/Modules/Rotation/RotationModules.INT.udn
+++ b/Engine/Documentation/Source/Engine/Rendering/ParticleSystems/Reference/Modules/Rotation/RotationModules.INT.udn
@@ -5,11 +5,9 @@ Description:Rotation modules control how particles will be rotated at various po
 
 [TOC(start:2 end:3)]
 
-## Rotation Modules 
-
 Rotation modules control how particles will be rotated at various points in their lives.
 
-### Init Mesh Rotation
+## Init Mesh Rotation
 
 The Init Mesh Rotation module applies an initial rotation to mesh particles created with a Mesh Data emitter.
 
@@ -21,7 +19,7 @@ The Init Mesh Rotation module applies an initial rotation to mesh particles crea
 
 
 
-### Init Mesh Rot (Seed)
+## Init Mesh Rot (Seed)
 
 The Init Mesh Rot (Seed) module applies an initial rotation to mesh particles created with a Mesh Data emitter, with the added benefit of randomization.
 
@@ -34,7 +32,7 @@ The Init Mesh Rot (Seed) module applies an initial rotation to mesh particles cr
 |**Inherit Parent** | If this Boolean value is _true_, then the rotation of the parent will be applied to the mesh particles as well. |
 
 
-### Initial Rotation
+## Initial Rotation
 
 ![psys_initial_rotation.gif](psys_initial_rotation.gif)(convert:false)
 
@@ -47,7 +45,7 @@ The Initial Rotation module is used to set the initial rotation of a particle at
 
 
 
-### Init Rotation (Seeded)
+## Init Rotation (Seeded)
 
 The Init Rotation (Seeded) module is identical to the [Initial Rotation](#InitialRotation) module in that it sets the initial rotation or particles; however, this module allows you to specify seed information to be used when choosing distribution values in order to provide a more consistent effect from the module each time the emitter is used. It contains the following member(s):
 
@@ -60,7 +58,7 @@ The Init Rotation (Seeded) module is identical to the [Initial Rotation](#Initia
 
 
 
-### Rotation/Life
+## Rotation/Life
 
 ![psys_rotation_life.gif](psys_rotation_life.gif)(convert:false)
 
diff --git a/Engine/Documentation/Source/Engine/Rendering/ParticleSystems/Reference/Modules/RotationRate/RotationRateModules.INT.udn b/Engine/Documentation/Source/Engine/Rendering/ParticleSystems/Reference/Modules/RotationRate/RotationRateModules.INT.udn
--- a/Engine/Documentation/Source/Engine/Rendering/ParticleSystems/Reference/Modules/RotationRate/RotationRateModules.INT.udn
+++ b/Engine/Documentation/Source/Engine/Rendering/ParticleSystems/Reference/Modules/RotationRate/RotationRateModules.INT.udn
@@ -5,11 +5,9 @@ Description:These modules apply a constant rate of rotation - or "spin" - to the
 
 [TOC(start:2 end:3)]
 
-## Rotation Rate Modules
-
 These modules apply a constant rate of rotation - or "spin" - to the particles.
 
-### Init Mesh RotRate
+## Init Mesh RotRate
 
 
 The Initial RotRate module is used to set the rotation rate of a mesh particle when it is emitted. It contains the following member(s):
@@ -19,7 +17,7 @@ The Initial RotRate module is used to set the rotation rate of a mesh particle w
 |[REGION:tablesection]Rotation[/REGION]||
 |**Start Rotation Rate** | A float distribution indicating the rotation rate the particle should have, in rotations per second. Value is retrieved based on the EmitterTime at the spawn of the particle, and is added to the Particle RotationRate and Base RotationRate values.|
 
-### Init Mesh RotRate (Seed)
+## Init Mesh RotRate (Seed)
 
 The Initial RotRate (Seeded) module is identical to the [Init Mesh RotRate](#InitMeshRotRate) module in that it sets the rotation rate of a mesh particle when it is emitted; however, this module allows you to specify seed information to be used when choosing distribution values in order to provide a more consistent effect from the module each time the emitter is used. It contains the following member(s):
 
@@ -31,7 +29,7 @@ The Initial RotRate (Seeded) module is identical to the [Init Mesh RotRate](#Ini
 |**Start Rotation Rate** | A float distribution indicating the rotation rate the particle should have, in rotations per second. Value is retrieved based on the EmitterTime at the spawn of the particle, and is added to the Particle RotationRate and Base RotationRate values.|
 
 
-### Mesh RotRate*Life
+## Mesh RotRate*Life
 
 
 The Mesh RotRate*Life module is used to scale the rotation rate of a mesh particle over its lifetime. It contains the following member(s):
@@ -43,7 +41,7 @@ The Mesh RotRate*Life module is used to scale the rotation rate of a mesh partic
 
 
 
-### Mesh Rot Rate Over Life
+## Mesh Rot Rate Over Life
 
 The Mesh Rot Rate Over Life module is used to set the rotation rate of a mesh-based particle over its lifetime. It contains the following member(s):
 
@@ -55,7 +53,7 @@ The Mesh Rot Rate Over Life module is used to set the rotation rate of a mesh-ba
 
 
 
-### Initial Rot Rate
+## Initial Rot Rate
 
 ![psys_initial_rotrate.gif](psys_initial_rotrate.gif)(convert:false)
 
@@ -67,7 +65,7 @@ The Initial RotRate module is used to set the rotation rate of a particle when i
 |**Start Rotation Rate** | A float distribution indicating the rotation rate the particle should have, in rotations per second. Value is retrieved based on the EmitterTime at the spawn of the particle, and is added to the Particle RotationRate and Base RotationRate values.|
 
 
-### Init Rot Rate (Seeded)
+## Init Rot Rate (Seeded)
 
 The Initial RotRate (Seeded) module is identical to the [Initial Rot Rate](#InitialRotRate) module in that it sets the rotation rate of a particle when it is emitted; however, this module allows you to specify seed information to be used when choosing distribution values in order to provide a more consistent effect from the module each time the emitter is used. It contains the following member(s):
 
@@ -80,7 +78,7 @@ The Initial RotRate (Seeded) module is identical to the [Initial Rot Rate](#Init
 
 
 
-### RotRate/Life
+## RotRate/Life
 
 ![psys_rotrate_life.gif](psys_rotrate_life.gif)(convert:false)
 
diff --git a/Engine/Documentation/Source/Engine/Rendering/ParticleSystems/Reference/Modules/Size/SizeModules.INT.udn b/Engine/Documentation/Source/Engine/Rendering/ParticleSystems/Reference/Modules/Size/SizeModules.INT.udn
--- a/Engine/Documentation/Source/Engine/Rendering/ParticleSystems/Reference/Modules/Size/SizeModules.INT.udn
+++ b/Engine/Documentation/Source/Engine/Rendering/ParticleSystems/Reference/Modules/Size/SizeModules.INT.udn
@@ -5,11 +5,9 @@ Description:These modules affect the sizes of particles across various points of
 
 [TOC(start:2 end:3)]
 
-## Size Modules 
-
 These modules affect the sizes of particles across various points of their lifespans.
 
-### Initial Size
+## Initial Size
 
 ![psys_initial_size.gif](psys_initial_size.gif)(convert:false)
 
@@ -22,7 +20,7 @@ The Initial Size module is used to set the initial size of a particle at spawn t
 
 
 
-### Initial Size (Seeded)
+## Initial Size (Seeded)
 
 The Initial Size (Seeded) module is identical to the [Initial Size](#InitialSize) module in that it sets the size of a particle when it is emitted; however, this module allows you to specify seed information to be used when choosing distribution values in order to provide a more consistent effect from the module each time the emitter is used. It contains the following member(s):
 
@@ -35,7 +33,7 @@ The Initial Size (Seeded) module is identical to the [Initial Size](#InitialSize
 
 
 
-### Size By Life
+## Size By Life
 
 ![](SizeByLife.gif)(convert:false)
 
@@ -51,7 +49,7 @@ The Size By Life module scales the size of the particle by a given value over it
 During both Spawn and Update, the Particle.Size value is multiplied with the retrieved scale value.
 
 
-### Size Scale
+## Size Scale
 
 The Size Scale module sets the Size of the particle to the BaseSize times the given scale factor. NOTE that this module overrides any size adjustments made prior to this module in that frame. It contains the following member(s): 
 
@@ -62,7 +60,7 @@ The Size Scale module sets the Size of the particle to the BaseSize times the gi
 |**Enable X, Enable Y, Enable Z** | Ignored.|
 
 
-### Size By Speed
+## Size By Speed
 
 ![](SizeBySpeed.gif)(convert:false)
 
diff --git a/Engine/Documentation/Source/Engine/Rendering/ParticleSystems/Reference/Modules/Spawn/SpawnModules.INT.udn b/Engine/Documentation/Source/Engine/Rendering/ParticleSystems/Reference/Modules/Spawn/SpawnModules.INT.udn
--- a/Engine/Documentation/Source/Engine/Rendering/ParticleSystems/Reference/Modules/Spawn/SpawnModules.INT.udn
+++ b/Engine/Documentation/Source/Engine/Rendering/ParticleSystems/Reference/Modules/Spawn/SpawnModules.INT.udn
@@ -4,13 +4,10 @@ Crumbs: %ROOT%, Engine, Engine/Rendering/ParticleSystems, Engine/Rendering/Parti
 Description:Spawn modules will impact the number/rate of emitter particles.
 
 [TOC(start:2 end:3)]
-
-## Spawn Modules
-
 Spawn modules will impact the number/rate of emitter particles.
 
 
-### Spawn Per Unit
+## Spawn Per Unit
 
 Allows for spawning particles based on the distance the emitter has traveled. Useful for things like sprite based smoke trails where the trail needs to stay cohesive whether traveling quickly or slowly and spawning a relative number of particles to always fill in gaps. It contains the following member(s): 
 
diff --git a/Engine/Documentation/Source/Engine/Rendering/ParticleSystems/Reference/Modules/SubUV/SubUVModules.INT.udn b/Engine/Documentation/Source/Engine/Rendering/ParticleSystems/Reference/Modules/SubUV/SubUVModules.INT.udn
--- a/Engine/Documentation/Source/Engine/Rendering/ParticleSystems/Reference/Modules/SubUV/SubUVModules.INT.udn
+++ b/Engine/Documentation/Source/Engine/Rendering/ParticleSystems/Reference/Modules/SubUV/SubUVModules.INT.udn
@@ -5,15 +5,13 @@ Description:SubUV modules allow particles to utilize subUV multi-frame textures.
 
 [TOC(start:2 end:3)]
 
-## SubUV Modules 
-
 SubUV modules allow particles to utilize subUV textures, in which a grid of individual "frames" is laid out.
 
 [REGION:note]
 SubUV modules should only be applied to an emitter when the InterpolationMethod is set to something other than PSUVIM_None. 
 [/REGION]
 
-### SubImage Index
+## SubImage Index
 
 The SubImage Index module selects the sub-image to use based on a float distribution. Currently, sub-images are ordered from left-to-right, top-to-bottom. It contains the following member(s):
 
@@ -26,7 +24,7 @@ The SubImage Index module selects the sub-image to use based on a float distribu
 
 
 
-### SubUV Movie
+## SubUV Movie
 
 The SubUV Movie module cycles through the sub-images of the texture in sequential order at a specified frame rate, similar to how a flipbook texture behaves. It contains the following properties:
 
diff --git a/Engine/Documentation/Source/Engine/Rendering/ParticleSystems/Reference/Modules/VectorField/VectorField.INT.udn b/Engine/Documentation/Source/Engine/Rendering/ParticleSystems/Reference/Modules/VectorField/VectorField.INT.udn
--- a/Engine/Documentation/Source/Engine/Rendering/ParticleSystems/Reference/Modules/VectorField/VectorField.INT.udn
+++ b/Engine/Documentation/Source/Engine/Rendering/ParticleSystems/Reference/Modules/VectorField/VectorField.INT.udn
@@ -5,11 +5,9 @@ Description:Reference for Vector Field modules, which provide 3D motion grid dat
 
 [TOC(start:2 end:3)]
 
-## Vector Field Modules
-
 [INCLUDE:Engine/Rendering/ParticleSystems/VectorFields#VectorFields]
 
-### Global Vector Fields
+## Global Vector Fields
 
 ![GlobalVF.png](GlobalVF.png)(w:720)
 
@@ -25,7 +23,7 @@ In order to be fully utilized, a Vector Field Volume Actor must exist in the lev
 | **Global Vector Field Scale** | This float value scales the amount of effect the global vector field will have on particle motion. Values lower than 1 will dampen the effect, while values greater than one will overdrive it. | 
 
 
-### Local Vector Field
+## Local Vector Field
 
 ![LocalVF.png](LocalVF.png)
 
@@ -41,7 +39,7 @@ The Local Vector Field module allows for a vector field to be associated directl
 | **Intensity** | This affects the overall intensity of the local vector field. Think of it as a multiplier for the vector field's force. A value of 0 effectively deactivates the vector field. Values greater than 1 will overdrive the effect. |
 | **Tightness** | This value controls how the vector field will control particle movement. A value of 0 means that the vector field is working as a force, additively influencing the velocity of the particles. A value of 1 means that particle velocity is completely replaced by the velocity defined in the vector field grid. Values between 0 and 1 blend the overall result between these two extremes. |
 
-### VF Init Rotation
+## VF Init Rotation
 
 ![VFInitRot.png](VFInitRot.png)
 
@@ -58,7 +56,7 @@ The Vector Field Initial Rotation module will allow for an initial amount of rot
 | **Min Initial Rotation** | The minimum initial rotation on the X, Y, and Z axes. This value is calculated as a [0..1] value. 0.25 = 90 degrees, 0.5 = 180 degrees, and so on. |
 | **Max Initial Rotation** | The maximum initial rotation on the X, Y, and Z axes. This value is calculated as a [0..1] value. 0.25 = 90 degrees, 0.5 = 180 degrees, and so on. |
 
-### VF Rotation Rate
+## VF Rotation Rate
 
 ![VFInitRotRate.gif](VFInitRotRate.gif)(convert:false)
 
@@ -75,7 +73,7 @@ If you need a local vector field to have a constant rate of spin, then the VF Ro
 | **Rotation Rate** | This sets the rotation rate on the X, Y, and Z axes. The value is calculated in full rotations per second. So a value of 0.1 will require 10 seconds to make one full rotation. |
 
 
-### Vector Field Scale
+## Vector Field Scale
 
 The Vector Field Scale module allows for scaling of a vector field's influence. This is calculated only at the moment of particle birth.
 
@@ -91,7 +89,7 @@ The Vector Field Scale module allows for scaling of a vector field's influence.
 
 
 
-### VF Scale/Life
+## VF Scale/Life
 
 The Vector Field Scale Over Life module allows for the scaling of a local vector field to take place over the life of an emitter. This is calculated throughout a particle's lifespan.
 
diff --git a/Engine/Documentation/Source/Engine/Rendering/ParticleSystems/Reference/Modules/Velocity/VelocityModules.INT.udn b/Engine/Documentation/Source/Engine/Rendering/ParticleSystems/Reference/Modules/Velocity/VelocityModules.INT.udn
--- a/Engine/Documentation/Source/Engine/Rendering/ParticleSystems/Reference/Modules/Velocity/VelocityModules.INT.udn
+++ b/Engine/Documentation/Source/Engine/Rendering/ParticleSystems/Reference/Modules/Velocity/VelocityModules.INT.udn
@@ -5,11 +5,9 @@ Description:Velocity modules exist to make adjustments to particle velocity in a
 
 [TOC(start:2 end:3)]
 
-## Velocity Modules
-
 Velocity modules exist to make adjustments to particle velocity in a variety of ways.
 
-### Initial Velocity
+## Initial Velocity
 
 ![psys_initial_velocity.gif](psys_initial_velocity.gif)(convert:false)
 
@@ -27,7 +25,7 @@ The Initial Velocity module sets the initial velocity of a particle. It contains
 
 
 
-### Initial Vel (Seed)
+## Initial Vel (Seed)
 
 The Initial Vel (Seed) module is identical to the [Initial Velocity](#InitialVelocity) module in that it sets the velocity of a particle when it is emitted; however, this module allows you to specify seed information to be used when choosing distribution values in order to provide a more consistent effect from the module each time the emitter is used. It contains the following member(s):
 
@@ -42,7 +40,7 @@ The Initial Vel (Seed) module is identical to the [Initial Velocity](#InitialVel
 |**Apply Owner Scale** | If _true_, the velocity value will be scaled by the ParticleSystemComponent's scale.|
 
 
-### Velocity Cone
+## Velocity Cone
 
 The Velocity Cone module emits particles in the shape of a user-defined cone. 
 
@@ -55,7 +53,7 @@ The Velocity Cone module emits particles in the shape of a user-defined cone.
 |**In World Space** | If _true_, the velocity is assumed to be in world-space.|
 |**Apply Owner Scale** | If _true_, the velocity value will be scaled by the ParticleSystemComponent's scale.|
 
-### Inherit Parent Vel
+## Inherit Parent Vel
 
 Contributes the velocity of the parent (the particle emitter itself) to a particle when spawning. It contains the following member(s): 
 
@@ -69,7 +67,7 @@ Contributes the velocity of the parent (the particle emitter itself) to a partic
 
 
 
-### Velocity/Life
+## Velocity/Life
 
 ![psys_velocity_life.gif](psys_velocity_life.gif)(convert:false)
 
