#UE4doc #docs:  Reworked phrase that used the word: expressibility work with the sentence better.

diff --git a/Engine/Documentation/Source/Engine/Rendering/PostProcessEffects/ColorGrading/ColorGrading.INT.udn b/Engine/Documentation/Source/Engine/Rendering/PostProcessEffects/ColorGrading/ColorGrading.INT.udn
--- a/Engine/Documentation/Source/Engine/Rendering/PostProcessEffects/ColorGrading/ColorGrading.INT.udn
+++ b/Engine/Documentation/Source/Engine/Rendering/PostProcessEffects/ColorGrading/ColorGrading.INT.udn
@@ -3,7 +3,7 @@ Title:Color Grading
 Crumbs: %ROOT%, Engine, Engine/Rendering/PostProcessEffects
 Description:Tonemapping and color correction effects for adjusting scene colors.
 
-Within the Unreal engine, the term **Color Grading** covers the **Tone Mapping** function (HDR to LDR transformation) and the further **Color Correction** (LDR color to screen color transformation).
+Within the Unreal engine, the term **Color Grading** covers the Tone Mapping function (HDR to LDR transformation) and the further Color Correction (LDR color to screen color transformation).
 
 ## Tone Mapping
 
@@ -21,7 +21,7 @@ This is the used function:
 
 Note that bright colors become gradually brighter but much less than dark ones. Black remains black and there is a mostly linear part of the curve that is a bit steeper than the non tone-mapped curve would be. That results in some contrast enhancing. While it is normal to get a different look when using a tone mapper, in order to get good results the source image needs to have a higher dynamic range. This can pay off in a more realistic film like look. 
 
-The presented tone mapper formula already has two tweakable constants and while more math would allow further tweaking, finding a good function is not easy as quality, flexibility, and performance need to be considered. We decided to stick to a simple formula and modify the resulting LDR color with a simple color lookup. As we already mapped the HDR color into a limited range, we get a lot of expressibility for the dark colors and still can modify the bright colors. This method is simple to understand and offers a lot of flexibility with local control and consistent performance. 
+The presented tone mapper formula already has two editable constants and while more math would allow further tweaking, finding a good function is not easy as quality, flexibility, and performance need to be considered. We decided to stick to a simple formula and modify the resulting LDR color with a simple color lookup. As we already mapped the HDR color into a limited range, we get several expressions of dark colors and still can modify the bright colors. This method is simple to understand and offers a lot of flexibility with local control and consistent performance. 
 
 | Property | Description |
 | -------- | ----------- |
