#UE4doc #docs: Adjusted bold text in accordance with the Instructional Guide

diff --git a/Engine/Documentation/Source/Engine/Rendering/Materials/MaterialInputs/MaterialInputs.INT.udn b/Engine/Documentation/Source/Engine/Rendering/Materials/MaterialInputs/MaterialInputs.INT.udn
--- a/Engine/Documentation/Source/Engine/Rendering/Materials/MaterialInputs/MaterialInputs.INT.udn
+++ b/Engine/Documentation/Source/Engine/Rendering/Materials/MaterialInputs/MaterialInputs.INT.udn
@@ -10,7 +10,7 @@ Not all inputs are used with every combination of Blend Mode and Shading Model.
 
 ## Inputs and Material Settings
 
-Not all inputs will be useful for every type of Material you create. For instance, when developing a Light Function - a Material that is applied to a light - you can only use the **Emissive Color** input on the material and nothing else, since other inputs, such as Metallic or Roughness, would not be applicable. Because of this, it is important to know what type of Material you are creating before you start worrying too much about inputs. The three primary controlling properties in this are:
+Not all inputs will be useful for every type of Material you create. For instance, when developing a Light Function - a Material that is applied to a light - you can only use the Emissive Color input on the material and nothing else, since other inputs, such as Metallic or Roughness, would not be applicable. Because of this, it is important to know what type of Material you are creating before you start worrying too much about inputs. The three primary controlling properties in this are:
 
 * **Blend Mode** - This controls how your Material will blend in with the pixels behind it.
 * **Shading Model** - This defines how light is calculated for the surface of the Material.
@@ -143,11 +143,11 @@ When using World Position Offset to expand your object beyond its original bound
 ## World Displacement & Tessellation Multiplier
 
 [EXCERPT:WorldDisp]
-**World Displacement** works very much like World Position Offset, but it uses Tessellation vertices rather than the base vertices of the mesh. In order for this to be enabled, the **Tessellation** property on the Material must be set to something other than _None_.
+**World Displacement** works very much like World Position Offset, but it uses Tessellation vertices rather than the base vertices of the mesh. In order for this to be enabled, the Tessellation property on the Material must be set to something other than _None_.
 [/EXCERPT:WorldDisp]
 
 [EXCERPT:TessMult]
-**Tessellation Multiplier** controls the amount tessellation along the surface, allowing more detail to be added where needed. As with World Displacement, for this to be enabled, the **Tessellation** property must be set to something other than _None_.
+**Tessellation Multiplier** controls the amount tessellation along the surface, allowing more detail to be added where needed. As with World Displacement, for this to be enabled, the Tessellation property must be set to something other than _None_.
 [/EXCERPT:TessMult]
 
 [EXCERPT:WorldDispAndTess]
diff --git a/Engine/Documentation/Source/Engine/Rendering/Materials/MaterialInstances/MaterialInstances.INT.udn b/Engine/Documentation/Source/Engine/Rendering/Materials/MaterialInstances/MaterialInstances.INT.udn
--- a/Engine/Documentation/Source/Engine/Rendering/Materials/MaterialInstances/MaterialInstances.INT.udn
+++ b/Engine/Documentation/Source/Engine/Rendering/Materials/MaterialInstances/MaterialInstances.INT.udn
@@ -9,7 +9,7 @@ Description:Hierarchical material relationships that allow modification of prope
 [TOC(start:2)]
 
 
-## Overview
+
 
 In Unreal Engine 4, Material instancing is used to change the appearance of a Material without incurring an expensive recompilation of the Material. Whereas a typical Material cannot be edited or changed without recompiling (something that _must_ happen prior to gameplay), an instanced Material can be made to change without such recompilation. Certain types of instanced Materials can even change during gameplay in response to in-game events (such as a tree whose Material blackens and chars while it burns). This allows tremendous visual flexibility in your artistic elements.
 
@@ -64,7 +64,7 @@ MIDs are created within script, either from a parameterized Material or a Materi
 To allow applying un-instanced Materials to primitives, the abstract base class _MaterialInterface_ is used.  This class is an interface to both the expressions and parameter values of an applied Material.  The Material class is a subclass of MaterialInterface which defines the expressions and default parameter values.  These types both inherit their expressions and parameter values from their parent, optionally overriding or animating some of the parameter values.
 
 [REGION:tip]
-Incidentally, this means that if you want a Blueprint variable to be able to accept _either_ a **Material** or a **Material Instance Constant**, you will need to make sure it is a variable of type _MaterialInterface_!
+Incidentally, this means that if you want a Blueprint variable to be able to accept _either_ a Material or a Material Instance Constant, you will need to make sure it is a variable of type _MaterialInterface_!
 [/REGION]
 
 
@@ -148,7 +148,7 @@ See the [](Engine/Rendering/Materials/ExpressionReference) for a complete listin
 
 Static parameters are applied at compile-time, so they produce more optimal code as entire branches of the material that are masked out by a static parameter will be compiled out and not executed at runtime.  Because they are applied at compile-time, they can only be changed from within the [MaterialInstanceEditor](Engine/Rendering/Materials/MaterialInstances/Editor) and not from script.
 
-**Warning: A new material will be compiled out for every combination of static parameters in the base material that are used by instances!**
+*Warning: A new material will be compiled out for every combination of static parameters in the base material that are used by instances!
 
 This can lead to an excessive amount of shaders being compiled.  Try to minimize the number of static parameters in the material and the number of permutations of those static parameters that are actually used.  See [Static Switch Parameter](Engine/Rendering/Materials/ExpressionReference/Parameters#StaticSwitchParameter) and [Static Component Mask Parameter](Engine/Rendering/Materials/ExpressionReference/Parameters#StaticComponentMaskParameter) for information on the specific static parameter types.
 
