#UE4doc #docs: Adjusted bold text in accordance with the Instructional Guide

diff --git a/Engine/Documentation/Source/Engine/Rendering/Materials/MaterialProperties/MaterialProperties.INT.udn b/Engine/Documentation/Source/Engine/Rendering/Materials/MaterialProperties/MaterialProperties.INT.udn
--- a/Engine/Documentation/Source/Engine/Rendering/Materials/MaterialProperties/MaterialProperties.INT.udn
+++ b/Engine/Documentation/Source/Engine/Rendering/Materials/MaterialProperties/MaterialProperties.INT.udn
@@ -8,7 +8,7 @@ Description:An introductory document on UE4 Materials and how they work.
 
 
 This document breaks down the properties available on a Material. These properties are accessible by selecting the Main Material node in the Material Editor, 
-or by **Right-Clicking** the Material in the **Content Browser** then selecting **Asset Actions** then [**Property Matrix**](Engine/UI/PropertyMatrix).
+or by **Right-Clicking** the Material in the **Content Browser** then selecting **Asset Actions** then [Property Matrix](Engine/UI/PropertyMatrix).
 
 * Properties inside Material:
 
