#UE4doc #docs:
Didn't commit changes the first time.

diff --git a/Engine/Documentation/Source/Engine/AI/EnvironmentQuerySystem/QuickStart/5/5.INT.udn b/Engine/Documentation/Source/Engine/AI/EnvironmentQuerySystem/QuickStart/5/5.INT.udn
--- a/Engine/Documentation/Source/Engine/AI/EnvironmentQuerySystem/QuickStart/5/5.INT.udn
+++ b/Engine/Documentation/Source/Engine/AI/EnvironmentQuerySystem/QuickStart/5/5.INT.udn
@@ -33,7 +33,7 @@ This Character will become our new AI Character. The reason we use a Character,
 
 	![](Engine\AI\BehaviorTrees\QuickStart\4\meshComp.png)
 
-1. In the **Details** panel, find the **Mesh** category and assign the `HeroTPP` (or your own) **Skeletal Mesh**.
+1. In the **Details** panel, find the **Mesh** category and assign the **SK_Mannequin** (or your own) **Skeletal Mesh**.
 
 	![](Engine\AI\BehaviorTrees\QuickStart\4\assignSkelMesh.png)
 
@@ -41,7 +41,7 @@ This Character will become our new AI Character. The reason we use a Character,
 
 1. Rotate the **Mesh** Component so it faces the direction the blue **Arrow Component** is pointing.
 
-1. In the **Details** panel, find the **Animation** category and assign the `HeroTPP_AnimBlueprint` to the **Animation Blueprint Generated Class** property.
+1. In the **Details** panel, find the **Animation** category and assign the `TopDown_AnimBlueprint` to the **Animation Blueprint Generated Class** property.
 
 	![](Engine\AI\BehaviorTrees\QuickStart\4\assignAnimBP.png)
 
@@ -53,13 +53,13 @@ This Character will become our new AI Character. The reason we use a Character,
 
 	![](Engine\AI\BehaviorTrees\QuickStart\4\charmve.png)
 
-1. In the **Details** panel, in the Movement Component category, adjust the Agent Radius and Agent Height values.
+1. In the **Details** panel, in the **Movement Capabilities** category, adjust the Agent Radius and Agent Height values.
 
 	![](Engine\AI\BehaviorTrees\QuickStart\4\agentProps.png)
 
 	_Set the Agent Radius to at least the Capsule's radius, and the Agent Height to at least twice the Capsule's "half height."_
 
-1. With the Character Movement Component still selected, change the **Max Walk Speed** property, in the **Character Movement** category to `400`.
+1. With the Character Movement Component still selected, change the **Max Walk Speed** property, in the **Character Movement: Walking** category to `400`.
 	
 	![](Engine\AI\BehaviorTrees\QuickStart\4\maxWalk.png)
 
diff --git a/Engine/Documentation/Source/Engine/AI/EnvironmentQuerySystem/QuickStart/6/6.INT.udn b/Engine/Documentation/Source/Engine/AI/EnvironmentQuerySystem/QuickStart/6/6.INT.udn
--- a/Engine/Documentation/Source/Engine/AI/EnvironmentQuerySystem/QuickStart/6/6.INT.udn
+++ b/Engine/Documentation/Source/Engine/AI/EnvironmentQuerySystem/QuickStart/6/6.INT.udn
@@ -21,7 +21,7 @@ Description: Adding the needed Key Values to the Runner Blackboard.
 
 %Steps%
 
-The **Blackboard** is the "memory" of the **Behavior Tree**. Here we will store a reference to the Actor we should flee from and a location to flee to.
+The **Blackboard** is the memory of the **Behavior Tree**. Here we will store a reference to the Actor we should flee from and a location to flee to.
 
 ## Steps
 
diff --git a/Engine/Documentation/Source/Engine/AI/EnvironmentQuerySystem/QuickStart/7/7.INT.udn b/Engine/Documentation/Source/Engine/AI/EnvironmentQuerySystem/QuickStart/7/7.INT.udn
--- a/Engine/Documentation/Source/Engine/AI/EnvironmentQuerySystem/QuickStart/7/7.INT.udn
+++ b/Engine/Documentation/Source/Engine/AI/EnvironmentQuerySystem/QuickStart/7/7.INT.udn
@@ -25,7 +25,9 @@ The setup for this AI Controller is a single **Run Behavior Tree** node running
 
 ## Steps
 
-1. In the Event Graph, have **On Begin Play** execute a **Run Beahvior Tree** node using the `RunnerBT` as the **BTAsset**:
+1. Double-click the `Runner_AI_CON` Blueprint in the Content Browser to open it for editing.
+
+2. In the Event Graph, have **On Begin Play** execute a **Run Beahvior Tree** node using the `RunnerBT` as the **BTAsset**:
 
 	![](setup.png)
 
