#UE4doc #docs: Adjusted bold text in accordance with the Instructional Guide

diff --git a/Engine/Documentation/Source/Engine/Rendering/Materials/PhysicallyBased/PhysicallyBasedShaders.INT.udn b/Engine/Documentation/Source/Engine/Rendering/Materials/PhysicallyBased/PhysicallyBasedShaders.INT.udn
--- a/Engine/Documentation/Source/Engine/Rendering/Materials/PhysicallyBased/PhysicallyBasedShaders.INT.udn
+++ b/Engine/Documentation/Source/Engine/Rendering/Materials/PhysicallyBased/PhysicallyBasedShaders.INT.udn
@@ -40,7 +40,7 @@ In terms of the "physically based" aspect of our Materials system, there are rea
 
 
 [REGION:note]
-**All** of these inputs are designed to take in values between 0 and 1. In the case of _Base Color_, this means a color with RGB values that fall between 0 and 1.
+All of these inputs are designed to take in values between 0 and 1. In the case of _Base Color_, this means a color with RGB values that fall between 0 and 1.
 [/REGION]
 
 Physically based values can be measured from real world Materials. Some examples will be given below.
@@ -128,7 +128,7 @@ If you have been making Materials in previous iterations of the Unreal Engine an
 [EXCERPT:Metallic]
 The **Metallic** input literally controls how "metal-like" your surface will be. Nonmetals have Metallic values of 0,
 metals have Metallic values of 1. For pure surfaces, such as pure metal, pure stone,
-pure plastic, etc. this value will be 0 **or** 1, not anything in between. When creating hybrid surfaces like corroded, dusty, or rusty metals,
+pure plastic, etc. this value will be 0 or 1, not anything in between. When creating hybrid surfaces like corroded, dusty, or rusty metals,
 you may find that you need some value _between_ 0 and 1. 
 [/EXCERPT:Metallic]
 
