#UE4doc #docs: Removed extraneous bolding as according to the instructional guide.

diff --git a/Engine/Documentation/Source/Engine/UI/ClassViewer/ClassViewer.INT.udn b/Engine/Documentation/Source/Engine/UI/ClassViewer/ClassViewer.INT.udn
--- a/Engine/Documentation/Source/Engine/UI/ClassViewer/ClassViewer.INT.udn
+++ b/Engine/Documentation/Source/Engine/UI/ClassViewer/ClassViewer.INT.udn
@@ -46,7 +46,7 @@ Description:Tool for examining UE4 classes.
 
 
 [EXCERPT:Intro]
-The **Class Viewer** allows the user to examine a hierarchical list of classes used by the editor. Through the tool, **Blueprints** can be created and opened for modification. You can also open associated C++ Header files or create new C++ classes based on the selected class. 
+The **Class Viewer** allows the user to examine a hierarchical list of classes used by the editor. Through the tool, Blueprints can be created and opened for modification. You can also open associated C++ Header files or create new C++ classes based on the selected class. 
 [/EXCERPT:Intro]
 
 ## Opening the Class Viewer
