#UE4doc #docs: ##Overview removed as according to the instructional guide.

diff --git a/Engine/Documentation/Source/Engine/UI/CurveEditor/CurveEditor.INT.udn b/Engine/Documentation/Source/Engine/UI/CurveEditor/CurveEditor.INT.udn
--- a/Engine/Documentation/Source/Engine/UI/CurveEditor/CurveEditor.INT.udn
+++ b/Engine/Documentation/Source/Engine/UI/CurveEditor/CurveEditor.INT.udn
@@ -44,10 +44,6 @@ Description:Tool for interactively setting up and modifying curves to drive valu
 
 [TOC(start:2)]
 
-
-
-## Overview
-
 [EXCERPT:Intro]
 The **Curve Editor** allows fine control over properties that change over time, e.g. properties which make use of [Distributions](Engine/Basics/Distributions). It is currently used in both [Matinee](Engine/Matinee) and the [Cascade](Engine/Rendering/ParticleSystems/Cascade) particle system editor.
 [/EXCERPT:Intro]
