#UE4doc #docs: Removed extraneous bolding as according to the instructional guide. Altered images to match the instructional guide: green -> yellow.

diff --git a/Engine/Documentation/Source/Engine/UI/InterfaceOverview/InterfaceOverview.INT.udn b/Engine/Documentation/Source/Engine/UI/InterfaceOverview/InterfaceOverview.INT.udn
--- a/Engine/Documentation/Source/Engine/UI/InterfaceOverview/InterfaceOverview.INT.udn
+++ b/Engine/Documentation/Source/Engine/UI/InterfaceOverview/InterfaceOverview.INT.udn
@@ -79,7 +79,7 @@ clicking** and choose **Close**.
 
 
 
-Each editor consists of a **Menu Bar** and one or more **Tabs** which comprise the majority of the functionality of the editor.
+Each editor consists of a Menu Bar and one or more Tabs which comprise the majority of the functionality of the editor.
 
 
 ### Editor Tabs
@@ -146,7 +146,7 @@ closing the editor. Tabs also remember their last location and will reopen there
 [/REGION]
 
 [REGION:tip]
-To save the layout manually, you can use the **Save layout** command in the **Window** menu.
+To save the layout manually, you can use the Save layout command in the Window menu.
 [/REGION]
 
 
