#UE4doc #docs:
Didn't commit changes the first time

diff --git a/Engine/Documentation/Source/Engine/AI/EnvironmentQuerySystem/QuickStart/9/9.INT.udn b/Engine/Documentation/Source/Engine/AI/EnvironmentQuerySystem/QuickStart/9/9.INT.udn
--- a/Engine/Documentation/Source/Engine/AI/EnvironmentQuerySystem/QuickStart/9/9.INT.udn
+++ b/Engine/Documentation/Source/Engine/AI/EnvironmentQuerySystem/QuickStart/9/9.INT.udn
@@ -28,12 +28,14 @@ Since you already have all the assets necessary to fill out the tree, you're goi
 Before we get into the steps to create this [**Behavior Tree**](Engine/AI/BehaviorTrees/UserGuide), let's talk about what it does (or rather what it will do once we create the necessary EQS bits). 
 
 * We enter the graph at the top Sequence Node, and then start executing it's children from left to right, and will keep executing as long as its children are successful.
-* First is the EQS Query node that runs the _FindPlayer_ template. We only allow this node to execute once every 0.5 seconds, and if it fails, the tree is locked out of executing the other nodes by the Cooldown Node returning **failed**, thusly failing its parent sequence node.
+* First is the EQS Query node that runs the _FindPlayer_ template. We only allow this node to execute once every 0.5 seconds, and if it fails, the tree is locked out of executing the other nodes by the Cooldown Node returning failed, thusly failing its parent sequence node.
 * But, if the EQS Query is successful (it finds the player), then the second EQS Query is allowed to execute, finding the AI a place to hide.
 * Finally we move to the place that the second EQS Query found.
 
 ## Steps
 
+1. Double-click the `RunnerBT` Blueprint in the Content Browser to open it for editing.
+
 1. [REGION:raw]Add a **Sequence Node** (right-click in the Graph Area %arrow% **Composites** %arrow% **Sequence**)[/REGION]
 
 	![](addSeq.png)
