#UE4doc #docs: Unbolded bolded hyperlinks. Removed extraneous bolding as according to the instructional guide.

diff --git a/Engine/Documentation/Source/Engine/UI/LevelEditor/LevelEditor.INT.udn b/Engine/Documentation/Source/Engine/UI/LevelEditor/LevelEditor.INT.udn
--- a/Engine/Documentation/Source/Engine/UI/LevelEditor/LevelEditor.INT.udn
+++ b/Engine/Documentation/Source/Engine/UI/LevelEditor/LevelEditor.INT.udn
@@ -99,15 +99,12 @@ Description:A guide that overviews the design and construction of game levels an
 
 [TOC(start:2)]
 
-
-## Overview
-
 [EXCERPT:Intro]
-The **Level Editor** provides the core level creation functionality for Unreal Editor. This is where levels are created, viewed, and modified mainly by placing, transforming, and editing the properties of [**Actors**](Engine/Actors).
+The **Level Editor** provides the core level creation functionality for Unreal Editor. This is where levels are created, viewed, and modified mainly by placing, transforming, and editing the properties of [Actors](Engine/Actors).
 [/EXCERPT:Intro]
 
 [EXCERPT:Levels]
-In Unreal Editor, the scenes in which you create your game experience are generally referred to as [**Levels**](Engine/Levels). You can think of a level as a 3D environment into which you place a series of objects and geometry to define the world your players will experience. Any object that is placed in your world, be it a light, a mesh, or a character, is considered to be an Actor. Technically speaking, _Actor_ is a programming class used within the Unreal Engine to define an object that has 3D position, rotation, and scale data. For sake of ease, however, it will be easiest to think of an Actor as _any object that can be placed in your levels_.
+In Unreal Editor, the scenes in which you create your game experience are generally referred to as [Levels](Engine/Levels). You can think of a level as a 3D environment into which you place a series of objects and geometry to define the world your players will experience. Any object that is placed in your world, be it a light, a mesh, or a character, is considered to be an Actor. Technically speaking, _Actor_ is a programming class used within the Unreal Engine to define an object that has 3D position, rotation, and scale data. For sake of ease, however, it will be easiest to think of an Actor as _any object that can be placed in your levels_.
 [/EXCERPT:Levels]
 
 [EXCERPT:Actors]
@@ -173,7 +170,7 @@ The button on the far right of the menu bar is:
 
 The **Toolbar** panel, like in most applications, is a group of commands providing quick access to commonly used tools and operations.
 
-See the [Toolbar](Engine/UI/LevelEditor/Toolbar) page for descriptions of each of the items on the **Toolbar**.
+See the [Toolbar](Engine/UI/LevelEditor/Toolbar) page for descriptions of each of the items on the Toolbar.
 
 ### Viewport
 
@@ -204,7 +201,7 @@ See the [Details](Engine/UI/LevelEditor/Details) page for a more complete overvi
 [/REGION]
 
 
-The **World Outliner** panel displays all of the Actors within the scene in a hierarchical tree view. Actors can be selected and modified directly from the **World Outliner**. You can also use the **Info** drop down menu to turn on an extra column that shows Levels, Layers, or ID Names.
+The **World Outliner** panel displays all of the Actors within the scene in a hierarchical tree view. Actors can be selected and modified directly from the **World Outliner**. You can also use the Info drop down menu to turn on an extra column that shows Levels, Layers, or ID Names.
 
 See the [World Outliner](Engine/UI/LevelEditor/SceneOutliner) page for details on using the **World Outliner**.
 
