#UE4doc #docs: Removed extraneous bolding as according to the instructional guide. ##Overview removed as according to the instructional guide.

diff --git a/Engine/Documentation/Source/Engine/UI/LevelEditor/Details/LevelEditorDetails.INT.udn b/Engine/Documentation/Source/Engine/UI/LevelEditor/Details/LevelEditorDetails.INT.udn
--- a/Engine/Documentation/Source/Engine/UI/LevelEditor/Details/LevelEditorDetails.INT.udn
+++ b/Engine/Documentation/Source/Engine/UI/LevelEditor/Details/LevelEditorDetails.INT.udn
@@ -25,8 +25,6 @@ Description:Panel providing information and customized editing tools for selecte
 
 [TOC(start:2 end:3)]
 
-## Overview
-
 [EXCERPT:DetailsOverview]
 The **Details** panel contains information, utilities, and functions specific to the current selection in the viewport. 
 It contains transform edit boxes for moving, rotating, and scaling Actors, displays all of the editable properties for 
@@ -44,7 +42,7 @@ For more information about the **Details** Panel UI, see [](Engine/UI/LevelEdito
 Actors can have _friendly_ names set directly in the editor by the designer. These names can be used to access these Actors or 
 find them using the search functionality in the [World Outliner panel](Engine/UI/LevelEditor\SceneOutliner).
 
-To edit the **Actor Name**, simply type the name into the text box in at the top of the **Details** panel.
+To edit the Actor Name, simply type the name into the text box in at the top of the **Details** panel.
 
 ![Actor Name Field](LE_DetailsFriendlyName.png)
 
@@ -80,7 +78,7 @@ The value of the property is reset to the default value as shown in the menu and
 ## Edit Conditions
 
 Properties can be either be enabled or disabled. A property can only be edited when it is enabled. By default, 
-all properties are enabled unless they have an **edit condition**. Properties with an edit condition, rely on 
+all properties are enabled unless they have an edit condition. Properties with an edit condition, rely on 
 the value of another variable to determine whether they are enabled and can be edited or not. 
 
 In some cases, edit conditions are used to determine whether the property will override some other value or whether it has any 
@@ -173,7 +171,7 @@ properties or functionality in a conspicuous manner making them easy to find and
 
 ### Advanced Properties
 
-Some properties are declared as **advanced** and are hidden away by default in the **Details** panel. Categories that 
+Some properties are declared as advanced and are hidden away by default in the **Details** panel. Categories that 
 contain advanced properties have an expander button ![](button_advanced.png). Clicking the button causes the advanced 
 properties to be displayed.
 
