#UE4doc #docs: Removed extraneous bolding as according to the instructional guide.

diff --git a/Engine/Documentation/Source/Engine/UI/LevelEditor/Details/Customizations/StaticMeshes/StaticMeshesCategory.INT.udn b/Engine/Documentation/Source/Engine/UI/LevelEditor/Details/Customizations/StaticMeshes/StaticMeshesCategory.INT.udn
--- a/Engine/Documentation/Source/Engine/UI/LevelEditor/Details/Customizations/StaticMeshes/StaticMeshesCategory.INT.udn
+++ b/Engine/Documentation/Source/Engine/UI/LevelEditor/Details/Customizations/StaticMeshes/StaticMeshesCategory.INT.udn
@@ -4,7 +4,7 @@ Crumbs: %ROOT%, Engine, Engine/UI, Engine/UI/LevelEditor, Engine/UI/LevelEditor/
 Description:Customized controls for viewing and modifying the StaticMesh assets applied to all components of an Actor.
 
 The **Static Mesh** category in the **Details** panel is only displayed when the selected Actor(s) are StaticMeshActors 
-or InterpActors. It promotes the normally buried **StaticMesh** property of the Actor's StaticMeshComponent, making it a 
+or InterpActors. It promotes the normally buried StaticMesh property of the Actor's StaticMeshComponent, making it a 
 first-class citizen and displays it using a customized widget that includes a thumbnail along with the usual 
 [Object reference](Engine/UI/LevelEditor/Details/Properties/ObjectReference) property controls.
 
