#UE4doc #docs: Removed extraneous bolding as according to the instructional guide.

diff --git a/Engine/Documentation/Source/Engine/UI/LevelEditor/Details/Customizations/Transform/TransformCategory.INT.udn b/Engine/Documentation/Source/Engine/UI/LevelEditor/Details/Customizations/Transform/TransformCategory.INT.udn
--- a/Engine/Documentation/Source/Engine/UI/LevelEditor/Details/Customizations/Transform/TransformCategory.INT.udn
+++ b/Engine/Documentation/Source/Engine/UI/LevelEditor/Details/Customizations/Transform/TransformCategory.INT.udn
@@ -8,14 +8,14 @@ The Transform section of the **Details** panel allows you to view and edit the t
 
 ![Transform Properties](Transform_Properties.png)
 
-Each transform property has numeric entry fields for the **X**, **Y**, and **Z** axes. You may type precise values directly 
+Each transform property has numeric entry fields for the X, Y, and Z axes. You may type precise values directly 
 into these fields to adjust the selected Actors. When more than one Actor is selected and the values of the properties differ, 
 the fields will display _Multiple Values_. In such cases, entering a number will cause that value to be entered for all selected Actors.
 
-The **Rotation** fields have the unique ability to be used as sliders. Clicking and dragging on the field allows you to 
+The Rotation fields have the unique ability to be used as sliders. Clicking and dragging on the field allows you to 
 _roll_ the value, increasing or decreasing depending on the distance the mouse is moved.
 
-The **Scale** fields can also be locked by clicking the ![](button_lock_scale.png) button. When locked, the ratio of 
+The Scale fields can also be locked by clicking the ![](button_lock_scale.png) button. When locked, the ratio of 
 the scale values to each other are maintained when any of the individual values are modified, allowing for uniform scaling.
 
 The transform properties default to _relative_ transformation. This means the transform is relative to the component's parent. 
