#UE4doc #docs: Removed extraneous bolding as according to the instructional guide. ##Overview removed as according to the instructional guide.

diff --git a/Engine/Documentation/Source/Engine/UI/LevelEditor/Details/Interface/DetailsInterface.INT.udn b/Engine/Documentation/Source/Engine/UI/LevelEditor/Details/Interface/DetailsInterface.INT.udn
--- a/Engine/Documentation/Source/Engine/UI/LevelEditor/Details/Interface/DetailsInterface.INT.udn
+++ b/Engine/Documentation/Source/Engine/UI/LevelEditor/Details/Interface/DetailsInterface.INT.udn
@@ -5,8 +5,6 @@ Description:Interface for panel providing information and customized editing too
 
 [TOC(start:2)]
 
-## Overview
-
 [INCLUDE:Engine/UI/LevelEditor/Details/#DetailsOverview]
 
 ![Level Editor Details Panel](LE_DetailsUI.png)
@@ -16,7 +14,7 @@ Description:Interface for panel providing information and customized editing too
 | **1** | **Lock Icon** | Locks the **Details** panel so the Actor whose properties are being viewed in the panel does not change based on Actor selection. |
 | **2** | **Friendly Name Box** | Displays the selected Actor's friendly name, which you can modify (grayed out when multiple Actors are selected). |
 | **3** | **Search Filter Box** | Allows you to type text to filter which details are displayed. |
-| **4** | **Property Matrix Icon** | Opens the **Property Matrix**, which allows you to view a matrix of property values belonging to multiple Actors at the same time, instead of as a combined list in the **Details** panel. <br/> For more information about the **Property Matrix**, see [](Engine/UI/PropertyMatrix). |
+| **4** | **Property Matrix Icon** | Opens the Property Matrix, which allows you to view a matrix of property values belonging to multiple Actors at the same time, instead of as a combined list in the **Details** panel. <br/> For more information about the Property Matrix, see [](Engine/UI/PropertyMatrix). |
 | **5** | **Display Filter Icon** | Displays options to expand, collapse, or otherwise filter property categories. |
 | **6** | **Properties Area** | Displays the Actor's properties so they can be viewed and/or modified. |
 
