#UE4doc #docs: Removed extraneous bolding as according to the instructional guide.

diff --git a/Engine/Documentation/Source/Engine/UI/LevelEditor/Details/Properties/Basic/BasicProperty.INT.udn b/Engine/Documentation/Source/Engine/UI/LevelEditor/Details/Properties/Basic/BasicProperty.INT.udn
--- a/Engine/Documentation/Source/Engine/UI/LevelEditor/Details/Properties/Basic/BasicProperty.INT.udn
+++ b/Engine/Documentation/Source/Engine/UI/LevelEditor/Details/Properties/Basic/BasicProperty.INT.udn
@@ -3,7 +3,7 @@ Title: Basic Controls
 Crumbs: %ROOT%, Engine, Engine/UI, Engine/UI/LevelEditor, Engine/UI/LevelEditor/Details
 Description:Simple property types such as numeric boxes and spinners, text boxes, checkboxes, radio buttons, and comboboxes.
 
-**Unreal Editor** has many simple control types you are familiar with for editing numeric values, text, etc.
+Unreal Editor has many simple control types you are familiar with for editing numeric values, text, etc.
 
 | Control | Description |
 | ------- | ----------- |
