#UE4doc #docs: Removed extraneous bolding as according to the instructional guide.

diff --git a/Engine/Documentation/Source/Engine/UI/LevelEditor/Details/Properties/VectorAndRotator/VectorAndRotatorProperty.INT.udn b/Engine/Documentation/Source/Engine/UI/LevelEditor/Details/Properties/VectorAndRotator/VectorAndRotatorProperty.INT.udn
--- a/Engine/Documentation/Source/Engine/UI/LevelEditor/Details/Properties/VectorAndRotator/VectorAndRotatorProperty.INT.udn
+++ b/Engine/Documentation/Source/Engine/UI/LevelEditor/Details/Properties/VectorAndRotator/VectorAndRotatorProperty.INT.udn
@@ -11,7 +11,7 @@ laid out horizontally.
  `Vectors` are properties that are used to store compound data that is ideally represented by three individual 
  `Float` values, such as a location in space or scale.
 
-Each of the `Float` values (**X**, **Y**, and **Z**) in a `Vector` property is displayed as a numeric input box 
+Each of the `Float` values (X, Y, and Z) in a `Vector` property is displayed as a numeric input box 
 that can be modified using direct input from the keyboard or you can click the 
 **Left Mouse Button** and drag horizontally to quickly modify each value.
 
@@ -22,7 +22,7 @@ that can be modified using direct input from the keyboard or you can click the
  `Rotators` are properties that are specifically designed to store rotations using three `Float` values limited to 
 the range `[0.0, 360.0]`, which represents the degrees of the rotation around each axis.
 
-Each of the `Rotator` values (**X**, **Y**, and **Z** mapped to  **Roll**, **Pitch**, and **Yaw** respectively) is displayed
+Each of the `Rotator` values (X, Y, and Z mapped to  Roll, Pitch, and Yaw respectively) is displayed
 as a numerical spinner. Values can be manually set using direct input from the keyboard or you can click the 
 **Left Mouse Button** and drag horizontally to quickly modify each value.
 
