#UE4doc #docs: Removed extraneous bolding as according to the instructional guide. ##Overview removed as according to the instructional guide.

diff --git a/Engine/Documentation/Source/Engine/UI/LevelEditor/InEditorTesting/InEditorTesting.INT.udn b/Engine/Documentation/Source/Engine/UI/LevelEditor/InEditorTesting/InEditorTesting.INT.udn
--- a/Engine/Documentation/Source/Engine/UI/LevelEditor/InEditorTesting/InEditorTesting.INT.udn
+++ b/Engine/Documentation/Source/Engine/UI/LevelEditor/InEditorTesting/InEditorTesting.INT.udn
@@ -61,7 +61,6 @@ Description: Guide to testing levels and gameplay directly in the editor.
 [/OBJECT]
 [/VAR]
 
-## Overview
 [EXCERPT:overview]
 Unreal Editor allows you to spawn a player anywhere and instantly preview your game from within the editor without waiting for files to save. There are two preview types available: **Play In Editor** (PIE) which is directly accessed via the **Play** button located on the Main Toolbar
 and **Simulate In Editor** (SIE) which is accessed via the **Play** button's drop-down menu (or by pressing **Alt+S**). The in-editor preview system supports toggling between _Play In Editor_ and _Simulate In Editor_ sessions, so that you can quickly iterate on gameplay and asset tweaks
@@ -78,7 +77,7 @@ appearance or behavior, even if you begin gameplay in a distant section of the l
 
 ## Toolbar
 
-You can begin a _Simulate In Editor_ or _Play In Editor_ session from either the **Level Editor** or the **Blueprint Editor**, using the appropriate **Toolbar** buttons or drop-down menus. 
+You can begin a _Simulate In Editor_ or _Play In Editor_ session from either the **Level Editor** or the **Blueprint Editor**, using the appropriate Toolbar buttons or drop-down menus. 
 
 From the **Level Editor**, if you click the **Play** button...
 
@@ -94,7 +93,7 @@ The **Play Options** drop-down menu appears...
 
 ![](button_PlayMenu.png)
 
-Here you can change your [Play In Editor mode](#Modes), set basic [networked Play In Editor](Engine/UI/LevelEditor\InEditorTesting\NetworkedPlayInEditor) options, or open the [full Play In Editor settings window](Engine/UI/LevelEditor\InEditorTesting\Settings). Clicking on **Simulate** begins _Simulate In Editor_ mode, which changes the testing section of the **Toolbar**.  While simulating, you can pause or stop gameplay, as well as swap to a _Play In Editor_ session with the **Possess** button.
+Here you can change your [Play In Editor mode](#Modes), set basic [networked Play In Editor](Engine/UI/LevelEditor\InEditorTesting\NetworkedPlayInEditor) options, or open the [full Play In Editor settings window](Engine/UI/LevelEditor\InEditorTesting\Settings). Clicking on **Simulate** begins _Simulate In Editor_ mode, which changes the testing section of the Toolbar.  While simulating, you can pause or stop gameplay, as well as swap to a _Play In Editor_ session with the **Possess** button.
 
 ![](SIE_toolbar.png)
 
@@ -117,7 +116,7 @@ $ Viewport : Gameplay will be shown in the active Level Editor viewport.
 	[![](PIE_viewport.png)(w:300)](PIE_viewport.png)
 	[/REGION:lightbox]
 	
-	This is the only display type that allows toggling between _Play In Editor_ and _Simulate in Editor_ by using the **Possess** and **Eject** buttons. In the **Viewport** modes, you can also use the **Pause** and **Stop** buttons to interrupt gameplay.
+	This is the only display type that allows toggling between _Play In Editor_ and _Simulate in Editor_ by using the Possess and Eject buttons. In the **Viewport** modes, you can also use the Pause and Stop buttons to interrupt gameplay.
 
 	![](PIE_toolbar.png)
 
@@ -169,8 +168,8 @@ $ Camera Location : Gameplay begins at the current camera location.
 $ Default Player Start : Gameplay begins at the Player Start location.
 
 [REGION:note]
-If your game uses a NavMesh for either player-controlled or AI-controlled movement, using **Camera Location** for your _Play In Editor_ starting location could cause you to spawn
-in a location that breaks pathing.  In this case, starting a _Play In Editor_ session from the **Default Player Start** is recommended.
+If your game uses a NavMesh for either player-controlled or AI-controlled movement, using Camera Location for your _Play In Editor_ starting location could cause you to spawn
+in a location that breaks pathing.  In this case, starting a _Play In Editor_ session from the Default Player Start is recommended.
 [/REGION]
 
 ### Play From Here
@@ -208,7 +207,7 @@ You will attach to the player controller, and begin a _Play In Editor_ session i
 You can also press the **Possess** keyboard shortcut to switch to _Play In Editor_ from _Simulate In Editor_. (default: **F10**)
 [/REGION:tip]
 
-If you began in-editor testing with a _Simulate In Editor_ session, and then used **Possess** to attach to the player controller, you will begin 
+If you began in-editor testing with a _Simulate In Editor_ session, and then used Possess to attach to the player controller, you will begin 
 your _Play In Editor_ session at the default player start. This is equivalent to selecting **Play In > Viewport at > Default Player Start** from the **Play In** dropdown menu.
 However, if you had previously swapped from a _Play In Editor_ session to a _Simulate In Editor_ session, without leaving in-editor testing mode, when you click **Possess**,
 you will re-attach to the player controller and resume gameplay at the location in the level where you previously clicked on **Eject**.
