#UE4doc #docs: Altered images to match the instructional guide: green -> yellow.

diff --git a/Engine/Documentation/Source/Engine/UI/LevelEditor/InEditorTesting/KeepSimulationChanges/KeepSimulationChanges.INT.udn b/Engine/Documentation/Source/Engine/UI/LevelEditor/InEditorTesting/KeepSimulationChanges/KeepSimulationChanges.INT.udn
--- a/Engine/Documentation/Source/Engine/UI/LevelEditor/InEditorTesting/KeepSimulationChanges/KeepSimulationChanges.INT.udn
+++ b/Engine/Documentation/Source/Engine/UI/LevelEditor/InEditorTesting/KeepSimulationChanges/KeepSimulationChanges.INT.udn
@@ -5,7 +5,7 @@ Description: Updates Actor properties from the simulated or in-game state.
 
 While using _Simulate In Editor_ or _Play In Editor_, you can edit the properties of Actors and see the effect of those changes on your level and the gameplay in it.  Using
 **Keep Simulation Changes**, you can take those changes and make them permanent. Changes that you can keep include properties seen in the **Details** tab, such as the Actor's
-transform in the level. After using **Keep Simulation Changes**, the selected Actors in your level will be updated from the runtime copy, so your changes will be reflected in the editor's copy. 
+transform in the level. After using Keep Simulation Changes, the selected Actors in your level will be updated from the runtime copy, so your changes will be reflected in the editor's copy. 
 
 Whether you alter an Actor's properties with _Blueprints_ or C++ code as a result of gameplay, or you begin simulating physics on an Actor and like where it ends up in the map,
 **Keep Simulation Changes** will allow you to save those changes. However, it  only works for Actors which were already placed in your level before you began simulating or playing the
@@ -14,7 +14,7 @@ Whether you alter an Actor's properties with _Blueprints_ or C++ code as a resul
 
 **To keep changes from Simulate In Editor:**
 	
-1.  Select the Actors that you want to keep changes for in either the **Viewport** or the **World Outliner**. **Right-click** in either the **Viewport** or the **World Outliner**, then select **Keep Simulation Changes** in the **right-click** menu.
+1.  Select the Actors that you want to keep changes for in either the **Viewport** or the **World Outliner**. **Right-click** in either the **Viewport** or the **World Outliner**, then select Keep Simulation Changes in the **right-click** menu.
 
 	![](in_simulate_keepsimchanges.png)(w:500)
 
@@ -29,6 +29,6 @@ Whether you alter an Actor's properties with _Blueprints_ or C++ code as a resul
 
 
 [REGION:tip]
-You can also press the **Keep Simulation Changes** keyboard shortcut after selecting your Actors. (default: **K**)
+You can also press the Keep Simulation Changes keyboard shortcut after selecting your Actors. (default: **K**)
 [/REGION:tip]
 
