Coding standard update for APIDocTool changes.

diff --git a/Engine/Documentation/Source/Programming/Development/CodingStandard/CodingStandard.INT.udn b/Engine/Documentation/Source/Programming/Development/CodingStandard/CodingStandard.INT.udn
--- a/Engine/Documentation/Source/Programming/Development/CodingStandard/CodingStandard.INT.udn
+++ b/Engine/Documentation/Source/Programming/Development/CodingStandard/CodingStandard.INT.udn
@@ -65,28 +65,28 @@ Functions that return a value should describe the return value; the name should
 
 ### Examples
 
-		/** the tea weight in kilograms */
+		/* the tea weight in kilograms */
 		float TeaWeight;
 
-		/** the number of tea leaves */
+		/* the number of tea leaves */
 		int32 TeaCount;
 
-		/** true indicates tea is smelly */
+		/* true indicates tea is smelly */
 		bool bDoesTeaStink;
 
-		/** non-human-readable FName for tea */
+		/* non-human-readable FName for tea */
 		FName TeaName;
 
-		/** human-readable name of the tea */
+		/* human-readable name of the tea */
 		FString TeaFriendlyName;
 
-		/** Which class of tea to use */
+		/* Which class of tea to use */
 		UClass* TeaClass;
 
-		/** The sound of pouring tea */
+		/* The sound of pouring tea */
 		USoundCue* TeaSound;
 
-		/** a picture of tea */
+		/* a picture of tea */
 		UTexture* TeaTexture;
 
 
@@ -171,12 +171,12 @@ Note that two different parameter comment styles are supported, as embodied by t
 Unlike the UE3 coding standard, the method comments should only be included once, where the method is publically declared.  The method comments should only contain information relevant to callers of the method, including any information about overrides of the method that may be relevant to the caller.  Details about the implementation of the method and its overrides that are not relevant to callers should be commented within the method implementation.
 
 
-		/** The interface for drinkable objects. */
+		/* The interface for drinkable objects. */
 		class IDrinkable
 		{
 		public:
 
-			/**
+			/*
 			 * Called when a player drinks this object.
 			 * @param OutFocusMultiplier - Upon return, will contain a multiplier to apply to the drinker's focus.
 			 * @param OutThirstQuenchingFraction - Upon return, will contain the fraction of the drinker's thirst to quench (0-1).
@@ -184,12 +184,12 @@ Unlike the UE3 coding standard, the method comments should only be included once
 			virtual void Drink(float& OutFocusMultiplier,float& OutThirstQuenchingFraction) = 0;
 		};
 
-		/** A cuppa (tea) */
+		/* A cuppa (tea) */
 		class FTea : public IDrinkable
 		{
 		public:
 
-			/**
+			/*
 			 * Calculate a delta-taste value for the tea given the volume and temperature of water used to steep.
 			 * @param VolumeOfWater - Amount of water used to brew in mL
 			 * @param TemperatureOfWater - Water temperature in Kelvins
@@ -202,24 +202,24 @@ Unlike the UE3 coding standard, the method comments should only be included once
 				float& OutNewPotency
 				);
 
-			/** Adds a sweetener to the tea, quantified by the grams of sucrose that would produce the same sweetness. */
+			// Adds a sweetener to the tea, quantified by the grams of sucrose that would produce the same sweetness.
 			void Sweeten(float EquivalentGramsOfSucrose);
 
-			/** The value in yen of tea sold in Japan. */
+			/* The value in yen of tea sold in Japan. */
 			float GetPrice() const
 			{
 				return Price;
 			}
 
-			// IDrinkable interface
+			//- IDrinkable interface
 			virtual void Drink(float& OutFocusMultiplier,float& OutThirstQuenchingFraction);
 
 		private:
 
-			/** The value of the tea in Yen. */
+			/* The value of the tea in Yen. */
 			float Price;
 
-			/** The sweetness of the tea, quantified by the grams of sucrose that would produce the same sweetness. */
+			// The sweetness of the tea, quantified by the grams of sucrose that would produce the same sweetness.
 			float Sweetness;
 		};
 
@@ -334,7 +334,9 @@ Enum classes are supported by all compilers and encouraged as a replacement for
 		{
 			enum Type
 			{
+				/* Description of Thing1 */
 				Thing1,
+				// Description of Thing2
 				Thing2
 			};
 		}
@@ -343,7 +345,9 @@ Enum classes are supported by all compilers and encouraged as a replacement for
 		UENUM()
 		enum class EThing : uint8
 		{
+			/* Description of Thing1 */
 			Thing1,
+			// Description of Thing2
 			Thing2
 		};
 
@@ -496,10 +500,10 @@ You can use namespaces to organize your classes, functions and variables where a
 
 * Example of a namespace-scoped enum:
 
-		/** Defining a enumeration within a namespace to achieve C#-style enum scoping */
+		/* Defining a enumeration within a namespace to achieve C#-style enum scoping */
 		namespace EColorChannel
 		{
-			/** Declare EColorChannel::Type as the actual type for this enum */
+			/* Declare EColorChannel::Type as the actual type for this enum */
 			enum Type
 			{
 				Red,
@@ -508,7 +512,7 @@ You can use namespaces to organize your classes, functions and variables where a
 			};
 		}
 
-		/** Given a color channel, returns the name of that channel. */
+		/* Given a color channel, returns the name of that channel. */
 		FString GetNameForColorChannel(const EColorChannel::Type ColorChannel)
 		{
 			switch(ColorChannel)
@@ -523,12 +527,12 @@ You can use namespaces to organize your classes, functions and variables where a
 
 * Note that for locally-declared enums, you will not be able to use a namespace for scoping.  In these cases, we opt to declare a local struct with no member variables, only a local enum type and use that struct for scoping.
 
-		/** Defining a locally-scoped enumeration using structs*/
+		/* Defining a locally-scoped enumeration using structs*/
 		class FObjectMover
 		{
         public:
 
-            /** Direction to move */
+            /* Direction to move */
             struct EMoveDirection
             {
 			    enum Type
@@ -538,7 +542,7 @@ You can use namespaces to organize your classes, functions and variables where a
 			    };
             };
 
-            /** Construct an FObjectMover with the specified movement direction */
+            /* Construct an FObjectMover with the specified movement direction */
             FObjectMover( const EMoveDirection::Type Direction );
 		}
 
