#UE4doc #docs: Adjusted bold text in accordance with the Instructional Guide

diff --git a/Engine/Documentation/Source/Engine/UI/LevelEditor/Modes/MeshPaintMode/VertexColor/SharingColorData/SharingColorData.INT.udn b/Engine/Documentation/Source/Engine/UI/LevelEditor/Modes/MeshPaintMode/VertexColor/SharingColorData/SharingColorData.INT.udn
--- a/Engine/Documentation/Source/Engine/UI/LevelEditor/Modes/MeshPaintMode/VertexColor/SharingColorData/SharingColorData.INT.udn
+++ b/Engine/Documentation/Source/Engine/UI/LevelEditor/Modes/MeshPaintMode/VertexColor/SharingColorData/SharingColorData.INT.udn
@@ -5,7 +5,7 @@ Description:How to go about transferring vertex color data across meshes using M
 
 
 
-There may be cases where you want to paint a mesh and then share those vertex colors with other instances of the mesh in the same level.  For example, if you had a pillar mesh you were using in an outdoor map and you wanted to paint the base of the pillar green, and have all instances in that map **only** appear that way.
+There may be cases where you want to paint a mesh and then share those vertex colors with other instances of the mesh in the same level.  For example, if you had a pillar mesh you were using in an outdoor map and you wanted to paint the base of the pillar green, and have all instances in that map only appear that way.
 
 Instanced vertex color data can be copied and pasted between instances in the level very easily using the **Copy** and **Paste** buttons in the **Mesh Paint Dialog**.
 
