#UE4 Docs: added linux version of OS toggle

diff --git a/Engine/Documentation/Source/Globals/GlobalIncludes.INT.udn b/Engine/Documentation/Source/Globals/GlobalIncludes.INT.udn
--- a/Engine/Documentation/Source/Globals/GlobalIncludes.INT.udn
+++ b/Engine/Documentation/Source/Globals/GlobalIncludes.INT.udn
@@ -61,6 +61,54 @@ Description:
 [/REGION]
 [/VAR]
 
+[VAR:OSSelect_Linux]
+[REGION:section]
+	[REGION:titlesection]
+	Choose Your Operating System
+	[/REGION]
+	[OBJECT:ToggleButton]
+		[PARAMLITERAL:category]
+		OS
+		[/PARAMLITERAL]
+		[PARAMLITERAL:id]
+		windows
+		[/PARAMLITERAL]
+		[PARAMLITERAL:active]
+		active_button
+		[/PARAMLITERAL]
+		[PARAM:label]
+		Windows
+		[/PARAM]
+	[/OBJECT]
+	[OBJECT:ToggleButton]
+		[PARAMLITERAL:category]
+		OS
+		[/PARAMLITERAL]
+		[PARAMLITERAL:id]
+		mac
+		[/PARAMLITERAL]
+		[PARAMLITERAL:active]
+		[/PARAMLITERAL]
+		[PARAM:label]
+		Mac OS
+		[/PARAM]
+	[/OBJECT]
+	[OBJECT:ToggleButton]
+		[PARAMLITERAL:category]
+		OS
+		[/PARAMLITERAL]
+		[PARAMLITERAL:id]
+		linux
+		[/PARAMLITERAL]
+		[PARAMLITERAL:active]
+		[/PARAMLITERAL]
+		[PARAM:label]
+		Linux
+		[/PARAM]
+	[/OBJECT]
+[/REGION]
+[/VAR]
+
 
 [VAR:DCCSelect]
 [REGION:section]
