#UE4doc #docs: Removed extraneous bolding as according to the instructional guide. ##Overview removed as according to the instructional guide.

diff --git a/Engine/Documentation/Source/Gameplay/HowTo/Networking/ReplicateActor/ReplicateActorOverview.INT.udn b/Engine/Documentation/Source/Gameplay/HowTo/Networking/ReplicateActor/ReplicateActorOverview.INT.udn
--- a/Engine/Documentation/Source/Gameplay/HowTo/Networking/ReplicateActor/ReplicateActorOverview.INT.udn
+++ b/Engine/Documentation/Source/Gameplay/HowTo/Networking/ReplicateActor/ReplicateActorOverview.INT.udn
@@ -3,17 +3,14 @@ Title:How to Replicate Actors
 Crumbs: %ROOT%, Gameplay
 Description:Guide for Replicating Actors in Multiplayer Games. 
 
+The process of passing data or commands between Server and Clients is referred to as Replication, and Actor Replication is where the Server identifies and acknowledges if an Actor is set to use Replication or not. If the Actor is set to Replicate, the Server will tell the Client to create and maintain a Replicated copy of that Actor on the Client's machine. The Server always has the authoritative version of that Actor, while a Client (or multiple Clients) may have the approximate, replicated versions of that Actor. 
 
-##Overview
-
-The process of passing data or commands between Server and Clients is referred to as **Replication**, and **Actor Replication** is where the Server identifies and acknowledges if an Actor is set to use Replication or not. If the Actor is set to Replicate, the Server will tell the Client to create and maintain a **Replicated** copy of that Actor on the Client's machine. The Server always has the authoritative version of that Actor, while a Client (or multiple Clients) may have the approximate, replicated versions of that Actor. 
-
-Options can be set to define how Replicated Actors behave on a Client, from whether or not to **Replicate Movement** (the location/rotation of the Actor) to setting its **Network Relevancy** (or when to send Replicated Data to the Client, as the Actor set to Replicate may or may not be relevant depending upon gameplay situations). 
+Options can be set to define how Replicated Actors behave on a Client, from whether or not to Replicate Movement (the location/rotation of the Actor) to setting its Network Relevancy (or when to send Replicated Data to the Client, as the Actor set to Replicate may or may not be relevant depending upon gameplay situations). 
 
 
 ## Setup
 
-To learn more about **Replicating Actors** using Blueprints, follow the link below:
+To learn more about Replicating Actors using Blueprints, follow the link below:
 
 [REGION:buttonlist]
 * [](Gameplay\HowTo\Networking\ReplicateActor\Blueprints)
