No changes

diff --git a/Engine/Documentation/Source/Engine/UI/LevelEditor/Modes/MeshPaintMode/VertexColor/AssetVsInstance/AssetVsInstance.INT.udn b/Engine/Documentation/Source/Engine/UI/LevelEditor/Modes/MeshPaintMode/VertexColor/AssetVsInstance/AssetVsInstance.INT.udn
--- a/Engine/Documentation/Source/Engine/UI/LevelEditor/Modes/MeshPaintMode/VertexColor/AssetVsInstance/AssetVsInstance.INT.udn
+++ b/Engine/Documentation/Source/Engine/UI/LevelEditor/Modes/MeshPaintMode/VertexColor/AssetVsInstance/AssetVsInstance.INT.udn
@@ -3,8 +3,8 @@ Title:Asset vs. Instance
 Crumbs:%ROOT%, Engine, Engine/UI/LevelEditor/Modes/MeshPaintMode/VertexColor
 Description:A look at the difference between painting on assets and painting on an individual instance.
 
-When painting **Vertex Colors**, you are only painting on an individual **Instance** of the material and not the actual material asset itself. 
-If you want to apply the vertex color data to another instance of the mesh, you can **Copy** and **Paste** the data onto another mesh or multiple meshes. 
+When painting Vertex Colors, you are only painting on an individual **Instance** of the material and not the actual material asset itself. 
+If you want to apply the vertex color data to another instance of the mesh, you can copy and paste the data onto another mesh or multiple meshes. 
 
 When painting vertex colors, the **Instance Vertex Colors** field will display information about the selected mesh(es).
 
diff --git a/Engine/Documentation/Source/Engine/UI/LevelEditor/Modes/MeshPaintMode/VertexColor/QuickGuide/QuickGuide.INT.udn b/Engine/Documentation/Source/Engine/UI/LevelEditor/Modes/MeshPaintMode/VertexColor/QuickGuide/QuickGuide.INT.udn
--- a/Engine/Documentation/Source/Engine/UI/LevelEditor/Modes/MeshPaintMode/VertexColor/QuickGuide/QuickGuide.INT.udn
+++ b/Engine/Documentation/Source/Engine/UI/LevelEditor/Modes/MeshPaintMode/VertexColor/QuickGuide/QuickGuide.INT.udn
@@ -1,5 +1,5 @@
 Availability: Public
-Title:Overview
+Title:Mesh Paint Overview
 Crumbs:%ROOT%, Engine, Engine/UI/LevelEditor/Modes/MeshPaintMode/VertexColor
 Description:A down-and-dirty quick guide to getting your feet wet with Mesh Paint Mode.
 
diff --git a/Engine/Documentation/Source/Engine/UI/LevelEditor/Modes/MeshPaintMode/VertexColor/VertexColor.INT.udn b/Engine/Documentation/Source/Engine/UI/LevelEditor/Modes/MeshPaintMode/VertexColor/VertexColor.INT.udn
--- a/Engine/Documentation/Source/Engine/UI/LevelEditor/Modes/MeshPaintMode/VertexColor/VertexColor.INT.udn
+++ b/Engine/Documentation/Source/Engine/UI/LevelEditor/Modes/MeshPaintMode/VertexColor/VertexColor.INT.udn
@@ -128,7 +128,7 @@ Mesh Paint can handle painting on individual level-placed instances of meshes an
 
 ## Setting up Materials
 
-The mesh paint system requires the material of the mesh to contain a **Vertex Color** expression as this is how the vertex color data is passed from the mesh to the material. The following sections cover ways to set up a material so that it can be used with Mesh Paint mode.
+The mesh paint system requires the material of the mesh to contain a Vertex Color expression as this is how the vertex color data is passed from the mesh to the material. The following sections cover ways to set up a material so that it can be used with Mesh Paint mode.
 
 
 [REGION:topics fifth]
