#UE4doc #docs: Adjusted bold text in accordance with the Instructional Guide

diff --git a/Engine/Documentation/Source/Engine/UI/LevelEditor/Modes/PlaceMode/PlacementBrowser.INT.udn b/Engine/Documentation/Source/Engine/UI/LevelEditor/Modes/PlaceMode/PlacementBrowser.INT.udn
--- a/Engine/Documentation/Source/Engine/UI/LevelEditor/Modes/PlaceMode/PlacementBrowser.INT.udn
+++ b/Engine/Documentation/Source/Engine/UI/LevelEditor/Modes/PlaceMode/PlacementBrowser.INT.udn
@@ -53,7 +53,7 @@ Description:Guide to the Place Mode, which allows for rapid deployment of placea
 
 Place Mode also makes it easier to add special types of Actors, such as Lights, that would have otherwise required that you dig through the hierarchy of the **Class Viewer** window. Further, it simplifies the creation of special subtypes of Actors, such as creating Physics Actors from selected Static Meshes.
 
-It should be noted that Place Mode is an efficiency- and convenience-driven tool, and is **not** intended to _replace_ the **Content Browser**. Instead, it exists as a way to greatly speed up design times when adding various props and other elements into your level, as it is optimized for fast Actor placement while offering you the convenience of quick access to common environment assets. However, whenever you need access to only those assets within a specific project folder, or any type of asset that is not natively placeable (such as a Material), you will still need to use the **Content Browser**.
+It should be noted that Place Mode is an efficiency- and convenience-driven tool, and is not intended to _replace_ the **Content Browser**. Instead, it exists as a way to greatly speed up design times when adding various props and other elements into your level, as it is optimized for fast Actor placement while offering you the convenience of quick access to common environment assets. However, whenever you need access to only those assets within a specific project folder, or any type of asset that is not natively placeable (such as a Material), you will still need to use the **Content Browser**.
 
 Place Mode is the default state of the Level Editor and can be activated in the **Modes** panel, which by default rests in the upper left corner of the Editor's main interface. You may also switch to this mode at any time by pressing **Shift + 1**.
 
