#UE4doc #docs: Removed ôOverviewö header

diff --git a/Engine/Documentation/Source/Engine/UI/LevelEditor/SceneOutliner/Interface/SceneOutlinerInterface.INT.udn b/Engine/Documentation/Source/Engine/UI/LevelEditor/SceneOutliner/Interface/SceneOutlinerInterface.INT.udn
--- a/Engine/Documentation/Source/Engine/UI/LevelEditor/SceneOutliner/Interface/SceneOutlinerInterface.INT.udn
+++ b/Engine/Documentation/Source/Engine/UI/LevelEditor/SceneOutliner/Interface/SceneOutlinerInterface.INT.udn
@@ -5,7 +5,7 @@ Description:Interface for hierarchical tree-view of all Actors within the curren
 
 [TOC(start:2)]
 
-## Overview
+
 
 [INCLUDE:Engine/UI/LevelEditor/SceneOutliner/#SceneOutliner]
 
