#UEdoc #docs - fixed grammatical errors, bold lettering, and updated images

diff --git a/Engine/Documentation/Source/Engine/Matinee/HowTo/MHT_4/MHT_4.INT.udn b/Engine/Documentation/Source/Engine/Matinee/HowTo/MHT_4/MHT_4.INT.udn
--- a/Engine/Documentation/Source/Engine/Matinee/HowTo/MHT_4/MHT_4.INT.udn
+++ b/Engine/Documentation/Source/Engine/Matinee/HowTo/MHT_4/MHT_4.INT.udn
@@ -4,7 +4,7 @@ Crumbs: %ROOT%, Engine, Engine/Matinee, Engine/Matinee/HowTo
 Description:A simple scene that fades in-and-out, uses multiple FOV angles and Post Process Effects. 
 
 ---
-In this How To series, we are going to create a short sequence that uses multiple camera angles with different FOVs (Fields of View), Depth of Field, and other Post Processing Effects applied to them. 
+In this How-to series, we are going to create a short sequence that uses multiple camera angles with different FOVs (Fields of View), Depth of Field, and other Post Processing Effects applied to them. 
 
 [INCLUDE:Engine/Matinee/HowTo/MHT_1#projectstart]
 
@@ -21,7 +21,7 @@ In this How To series, we are going to create a short sequence that uses multipl
 
 	![](Engine\Matinee\HowTo\MHT_1\MHT1_AddMatinee.png)(w:720)
 
-1. If the **Matinee** undo warning message appears, click the **Continue** button. 
+1. If the Matinee undo warning message appears, click the **Continue** button. 
 
 	![](Engine\Matinee\HowTo\MHT_1\MHT1_MatineeWarning.png)
 	
@@ -33,17 +33,17 @@ In this How To series, we are going to create a short sequence that uses multipl
 
 	_Click for Full View._
 
-1. Inside the **Content Browser** under the **Game/Props** folder, drag into the viewport an **SM_Brush**, **SM_Couch**, and **SM_Rock**.
+1. Inside the **Content Browser** under the **Game/Props** folder, **drag** into the viewport an **SM_Brush**, **SM_Couch**, and **SM_Rock**.
 
 	![](MHT4_AllPropsAdded.png)(w:720)
 
 	_Arrange the props similar to above by using the **Translation** (**W** key) and **Rotation** (**E** key) widgets to move and rotate the props._
 
-1. Next, inside the **Game/Particles** folder, drag in the **P_Fire** particle effect and place it inside the **SM_Brush**. 
+1. Next, inside the **Game/Particles** folder, **drag** in the **P_Fire** particle effect and place it inside the SM_Brush. 
 
 	![](MHT4_FireParticleAdded.png)(w:720)
 
-1. From the **Modes** menu, under the **Basic** tab, drag in a **Camera** and place it in front of the **SM_Brush**. 
+1. From the **Modes** menu, under the **Basic** tab, **drag** in a **Camera** and place it in front of the SM_Brush. 
 
 	![](MHT4_Camera1Placed.png)(w:720)
 
@@ -51,11 +51,11 @@ In this How To series, we are going to create a short sequence that uses multipl
 
 	![](MHT4_C1_CameraSettings.png)
 
-1. Drag in another **Camera** and place it behind the **SM_Rock** (rotating it so the rock is in view). 
+1. **Drag** in another Camera and place it behind the SM_Rock (rotating it so the rock is in view). 
 
 	![](MHT4_Camera2Placed.png)(h:420)
 
-1. With the new camera selected, in the **Details** panel under **Depth of Field**, enable each of the options below and assign the values shown.
+1. With the new camera selected, in the Details panel under Depth of Field enable each of the options below and assign the values shown.
 
 	![](MHT4_C2_CameraSettings.png)
 
@@ -69,13 +69,13 @@ In this How To series, we are going to create a short sequence that uses multipl
 	[/PARAM]
 [/OBJECT]
 
-1. Select the **Matinee** in the **World Outliner**, then select the **Open Matinee** option from the **Details** panel. 
+1. **Select** the Matinee in the **World Outliner**, then select the **Open Matinee** option from the Details panel. 
 
 	![](MHT4_OpenMatinee.png)
 
-1. Minimize the **Matinee** window, then in the **World Outliner** or viewport, select the **CameraActor1**. 
+1. **Minimize** the Matinee window, then in the World Outliner or viewport, **select** the **CameraActor1**. 
 
-1. With the **CameraActor1** highlighted, re-open the **Matinee** window and in the **Tracks** window, **Right-Click** and select **Add New Camera Group**. 
+1. With the CameraActor1 highlighted, **re-open** the Matinee window and in the **Tracks** window, **Right-Click** and select **Add New Camera Group**. 
 
 	![](MHT4_AddNewCameraGroup.png)
 
@@ -85,31 +85,31 @@ In this How To series, we are going to create a short sequence that uses multipl
 
 	![](MHT4_CamerasAdded.png)
 
-1. Holding **Control**, **Click-and-Drag** the red ending marker (yellow box) and drag it to the **20.00** mark (yellow X). 
+1. Holding **Control**, **Click-and-Drag** the red ending marker (yellow box) and **drag** it to the **20.00** mark (yellow X). 
 
 	![](MHT4_IncreaseLength.png)
 
-1. **Right-Click** on the **Camera1** group and select the **Add New Float Property Track** from the context menu. 
+1. **Right-Click** on the Camera1 group and **select** the **Add New Float Property Track** from the context menu. 
 
 	![](MHT4_Camera1FloatProperty.png)
 
-1. From the **Choose Property** drop-down menu, select the **CameraComponent.PostProcessBlendWeight** option and then press the **Ok** button. 
+1. From the **Choose Property** drop-down menu, **select** the **CameraComponent.PostProcessBlendWeight** option and then **press** the **Ok** button. 
 
 	![](MHT4_Camera1BlendAdded.png)
 
-1. Repeat step 7 and 8, but add a **CameraComponent.PostProcessSettings.FilmContrast** option and press **Ok**. 
+1. Repeat step 7 and 8, but add a **CameraComponent.PostProcessSettings.FilmContrast** option and **press Ok**. 
 
-1. Repeat step 7 and 8 once more, but this time add a **CameraComponent.PostProcessSettings.FilmSaturation** option and press **Ok**. 
+1. Repeat step 7 and 8 once more, but this time add a **CameraComponent.PostProcessSettings.FilmSaturation** option and **press Ok**. 
 
-1. Your **Camera1** group should now have a **Movement** and **FOV** track as well as three **Float Property** tracks as seen below. 
+1. Your Camera1 group should now have a **Movement** and **FOV** track as well as three **Float Property** tracks as seen below. 
 
 	![](MHT4_Camera1PropertiesComplete.png)
 
-1. **Right-Click** on the **Camera2** group and select the **Add New Float Property Track** from the context menu. 
+1. **Right-Click** on the Camera2 group and **select** the Add New Float Property Track from the context menu. 
 
-1. In the **Choose Property** drop-down menu, select the **CameraComponent.PostProcessSettings.DepthOfFieldScale** option and press **Ok**. 
+1. In the Choose Property drop-down menu, select the **CameraComponent.PostProcessSettings.DepthOfFieldScale** option and **press Ok**. 
 
-1. Repeat step 12 and 13, but add a **CameraComponent.PostProcessSettings.GrainIntensity** option and press **Ok**. 
+1. Repeat step 12 and 13, but add a **CameraComponent.PostProcessSettings.GrainIntensity** option and **press** Ok**. 
 
 1. Both camera groups should now be set up and should look similar to below. 
 
@@ -121,55 +121,55 @@ In this How To series, we are going to create a short sequence that uses multipl
 		Camera 1, Shot 1 - Movement and Post Processing Effects
 	[/PARAM]
 	[PARAM:description]
-	In this step, we will add camera movement and Post Processing effects to camera 1 which will also be our first shot.
+	In this step we will add camera movement and Post Processing effects to Camera 1 which will also be our first shot.
 	[/PARAM]
 [/OBJECT]
 
-1. Inside the **Matinee Editor**, select the **FOVAngle** track of the **Camera1** group and press **Enter** to add a keyframe at **0.00**. 
+1. Inside the Matinee Editor, select the **FOVAngle** track of the Camera1 group and **press** Enter** to add a keyframe at **0.00**. 
 
 	![](MHT4_Camera1FOV1.png)
 
-1. Also for the **Camera1** group, select the **...ssBlendWeight** (Blend Weight) track and press **Enter** to add a keyframe at **0.00**. 
+1. Also for the Camera1 group **select** the **...ssBlendWeight** (Blend Weight) track and **press Enter** to add a keyframe at **0.00**. 
 
-1. With the new keyframe added to the **BlendWeight** track still selected, **Right-Click** on it and choose **Set Value** then set its value to **0**. 
+1. With the new keyframe added to the **BlendWeight** track still selected, **Right-Click** on it and **choose** Set Value then set its value to **0**. 
 
 	![](MHT4_BlendWeight1SetValue.png)
 
-1. **Right-Click** on the first keyframe of the **FOVAngle** track of **Camera1** and choose the **Set Value** option and set its value to **25**.
+1. **Right-Click** on the first keyframe of the FOVAngle track of Camera1 and **choose** the Set Value option and **set** its value to **25**.
 
 	![](MHT4_C1_FOV1.png)
 
-1. **Left-Click** on first keyframe of the **Movement** track for **Camera1** to select it.
+1. **Left-Click** on first keyframe of the Movement track for Camera1 to select it.
 
-1. With the keyframe selected, minimize the **Matinee** window and move the camera straight up so the brush and fire are out of view.
+1. With the keyframe selected, **minimize** the Matinee window and **move the camera** straight up so the brush and fire are out of view.
 
 	![](MHT4_C1_Movement1.png)(w:520)
 
-1. Return to the **Matinee Editor** and scrub the **Time Bar** to the **5.00** 
+1. Return to the Matinee Editor and scrub the **Time Bar** to the **5.00** 
 
-1. With the **Movement** track selected for **Camera1**, press **Enter** to add a keyframe at **5.00**.
+1. With the Movement track selected for Camera1, **press Enter** to add a keyframe at **5.00**.
 
 	![](MHT4_C1_Movement2.png)
 
-1. With the new keyframe still selected, minimize the **Matinee** window and move the camera back down to the ground. 
+1. With the new keyframe still selected, **minimize** the Matinee window and **move** the camera back down to the ground. 
 
 	![](MHT4_C1_Movement3.png)
 	
 	_You should see a dotted yellow line which will represent the camera movement._
 
-1. Return to the **Matinee Editor** and click on the **FOVAngle** track for **Camera1**. 
+1. Return to the Matinee Editor and **click** on the FOVAngle track for Camera1. 
 
-1. With the **Time Bar** still at **5.00**, press **Enter** to add a new keyframe. 
+1. With the Time Bar still at **5.00**, **press Enter** to add a new keyframe. 
 
-1. **Right-Click** on the new keyframe and select **Set Value** and set the value to **45**. 
+1. **Right-Click** on the new keyframe and **select** Set Value and **set** the value to **45**. 
 
 	![](MHT4_C1_FOV2.png)
 
-1. Scrub the **Time Bar** to the **6.00** mark then click on the **BlendWeight** track. 
+1. Scrub the Time Bar to the **6.00** mark then click on the **BlendWeight** track. 
 
-1. Press **Enter** to add a new keyframe at **6.00** on the **BlendWeight** track. 
+1. **Press Enter** to add a new keyframe at **6.00** on the **BlendWeight** track. 
 
-1. **Right-Click** on the new keyframe and select **Set Value** and set the value to **0.3**. 
+1. **Right-Click** on the new keyframe and **select** Set Value and **set** the value to **0.3**. 
 
 	![](MHT4_C1_Blendweight2.png)
 
@@ -187,49 +187,49 @@ In this How To series, we are going to create a short sequence that uses multipl
 
 
 
-1. On the **Movement** track for **Camera2**, move the **Time Bar** to **5.00** and press **Enter** to add a keyframe, then delete the keyframe at **0.00**.
+1. On the Movement track for Camera2 **move** the Time Bar to **5.00** and **press Enter** to add a keyframe, then delete the keyframe at **0.00**.
 
 	![](MHT4_C2_Movement1.png)
 
-1. With the keyframe at **5.00** for the **Movement** track of **Camera2** selected, minimize the **Matinee** window. 
+1. With the keyframe at **5.00** for the Movement track of Camera2 selected, **minimize** the Matinee window. 
 
-1. In the viewport, press **E** to enter **Rotation** mode and rotate the camera so that it angled upwards looking at the sky. 
+1. In the viewport, **press E** to enter **Rotation** mode and **rotate the camera** so that it angled upwards looking at the sky. 
 
 	![](MHT4_C2_RotateCamera.png)(h:420)
 
-1. Return to the **Matinee Editor** and move the **Time Bar** to the **6.50** mark. 
+1. Return to the Matinee Editor and **move** the Time Bar to the **6.50** mark. 
 
-1. Select the **...hOfFieldScale** (Depth of Field Scale) track of **Camera2** and press **Enter** to add a keyframe at **6.50** for this track. 
+1. **Select** the **...hOfFieldScale** (Depth of Field Scale) track of Camera2 and **press Enter** to add a keyframe at **6.50** for this track. 
 
-1. Select the **...rainIntensity** (Grain Intensity) track of **Camera2** and press **Enter** to add a keyframe at **6.50** for this track. 
+1. **Select** the **...rainIntensity** (Grain Intensity) track of Camera2 and **press Enter** to add a keyframe at **6.50** for this track. 
 
 	![](MHT4_C2_GrainIntensity1.png)
 
-1. **Right-Click** on the keyframe for ***FieldScale** and select **Set Value**, then set the value to **0.50**. 
+1. **Right-Click** on the keyframe for ***FieldScale** and **select** Set Value, then set the value to **0.50**. 
 
-1. **Right-Click** on the keyframe for **GrainIntensity** and select **Set Value**, then set the value to **0.00**. 
+1. **Right-Click** on the keyframe for **GrainIntensity** and **select** Set Value, then set the value to **0.00**. 
 
-1. Move the **Time Bar** to the **9.00** mark, select the **Movement** track of **Camera2** then press **Enter** to add a keyframe. 
+1. **Move** the Time Bar to the **9.00** mark, **select** the Movement track of Camera2 then **press Enter** to add a keyframe. 
 
 	![](MHT4_C2_Movement2.png)
 
-1. With the new keyframe selected, minimize the **Matinee** window and return to the viewport. 
+1. With the new keyframe selected, **minimize** the Matinee window and return to the viewport. 
 
-1. Move the camera up and rotate it down so that the rock is partially still in view. 
+1. **Move** the camera up and **rotate** it down so that the rock is partially still in view. 
 
 	![](MHT4_C2_Movement3.png)(h:420)
 
-1. Return to the **Matinee Editor** window.
+1. Return to the Matinee Editor window.
 
-1. While still on the **Movement** track of **Camera2**, scrub the **Time Bar** to **12.00** and press **Enter** to add another keyframe.
+1. While still on the Movement track of Camera2, scrub the Time Bar to **12.00** and **press Enter** to add another keyframe.
 
-1. With the new keyframe still selected, return to the viewport and move the camera so that it is looking up at the clouds. 
+1. With the new keyframe still selected, return to the viewport and **move** the camera so that it is looking up at the clouds. 
 
 	![](MHT4_C2_Movement4.png)
 
 	_This will move our camera up, then pan upwards into the clouds._
 
-1. Return to the **Matinee Editor** and on the **FieldScale** track of **Camera2** at **12.00**, add a keyframe and set its value to **0**. 
+1. Return to the Matinee Editor and on the FieldScale track of Camera2 at **12.00**, **add** a keyframe and set its value to **0**. 
 
 	![](MHT4_C2_FieldScale2.png)
 
@@ -245,57 +245,57 @@ In this How To series, we are going to create a short sequence that uses multipl
 	[/PARAM]
 [/OBJECT]
 
-1. Inside the **Matinee Editor**, select and copy the keyframe at **5.00** on the **Movement** track of **Camera1**. 
+1. Inside the Matinee Editor, **select and copy** the keyframe at **5.00** on the Movement track of Camera1. 
 
-1. Paste the copied keyframe on the **Movement** track of **Camera1** at the **11.50** mark. 
+1. **Paste** the copied keyframe on the Movement track of Camera1 at the **11.50** mark. 
 
 	![](MHT4_C1_MovementS3.png)
 
-1. Still on the **Movement** track, move the **Time Bar** to **12.00** and press **Enter** to add a new keyframe. 
+1. Still on the Movement track, **move** the Time Bar to **12.00** and **press Enter** to add a new keyframe. 
 
-1. With the new keyframe still selected, minimize **Matinee** and return to the level viewport. 
+1. With the new keyframe still selected, **minimize** Matinee and return to the level viewport. 
 
-1. Move **CameraActor1** so that the brush is on the right side of the frame (similar to below). 
+1. **Move** CameraActor1 so that the brush is on the right side of the frame (similar to below). 
 
 	![](MHT4_C1_S3_Key1.png)
 
-1. Return to the **Matinee Editor** and **Camera1**. 
+1. Return to the Matinee Editor and Camera1. 
 
-1. Still at the **12.00** mark, click on the **....FilmContrast** (Contrast) track and press **Enter** to add a keyframe. 
+1. Still at the **12.00** mark, **click** on the **....FilmContrast** (Contrast) track and **press Enter** to add a keyframe. 
 
-1. Next click on the **...ilmSaturation** (Saturation) track and press **Enter** to add a keyframe. 
+1. Next **click** on the **...FilmSaturation** (Saturation) track and **press Enter** to add a keyframe. 
 
 	![](MHT4_C1_ContrastSaturation.png)(w:720)
 
-1. **Right-Click** on the keyframe at **12.00** for **Contrast** and choose **Set Value** and set it to **50**. 
+1. **Right-Click** on the keyframe at **12.00** for **Contrast** and **choose** Set Value and **set** it to **50**. 
 
-1. **Right-Click** on the keyframe at **12.00** for **Saturation** and choose **Set Value** and set it to **25**. 
+1. **Right-Click** on the keyframe at **12.00** for **Saturation** and **choose** Set Value and **set** it to **25**. 
 
 	![](MHT4_C1_SaturationValueSet.png)(w:720)
 
-	_Values for **Contrast** and **Saturation** have been set, but we need to set defaults which we will do now._
+	_Values for Contrast and Saturation have been set, but we need to set defaults which we will do now._
 
-1. For both the **Contrast** and **Saturation** tracks of **Camera1**, move the **Time Bar** back to **0.00** and add a keyframe for both. 
+1. For both the Contrast and Saturation tracks of Camera1, **move** the Time Bar back to **0.00** and add a keyframe for both. 
 
-1. **Right-Click** on both keyframes and choose **Set Value** and set both to **1**. 
+1. **Right-Click** on both keyframes and **choose** Set Value and **set** both to **1**. 
 
 	![](MHT4_C1_SaturationDefault.png)(w:720)
 
-1. Move the **Time Bar** to **11.50** and a new keyframe for both **Contrast** and **Saturation**. 
+1. **Move** the Time Bar to **11.50** and **add** a new keyframe for both Contrast and Saturation. 
 
-1. **Right-Click** on both keyframes and choose **Set Value** and set both to **1**. 
+1. **Right-Click** on both keyframes and **choose** Set Value and **set** both to **1**. 
 
 	![](MHT4_C1_SaturationFinalKey.png)(w:720)
 
-	_This will ensure that our **Contrast** and **Saturation** values are **1** for the first shot but are updated for the third shot._
+	_This will ensure that our Contrast and Saturation values are 1 for the first shot but are updated for the third shot._
 
-1. Minimize the **Matinee** window and select **CameraActor1** in the **World Outliner**. 
+1. **Minimize** the Matinee window and **select** CameraActor1 in the World Outliner. 
 
-1. In the **Details** panel for **CameraActor1**, scroll down and find the **Film** section and enable both **Contrast** and **Saturation** options. 
+1. In the Details panel for CameraActor1 scroll down and find the **Film** section and enable both Contrast and Saturation options. 
 
 	![](MHT4_C1_EnableSaturation.png)(h:540)
 
-	_Leave their settings as default since they will be updated via the **Matinee**._
+	_Leave their settings as default since they will be updated via the Matinee._
 
 
 [OBJECT:TopicList]
@@ -307,29 +307,29 @@ In this How To series, we are going to create a short sequence that uses multipl
 	[/PARAM]
 [/OBJECT]
 
-1. Inside the **Matinee Editor**, move the **Time Bar** to **15.50**. 
+1. Inside the Matinee Editor, **move** the Time Bar to **15.50**. 
 
-1. On the **Movement** track for **Camera2**, add a new keyframe at **15.50**. 
+1. On the Movement track for Camera2, **add** a new keyframe at **15.50**. 
 
-1. Select the **Intensity** track and add a new keyframe at **15.50** as well. 
+1. **Select** the **Intensity** track and **add** a new keyframe at **15.50** as well. 
 
 	![](MHT4_C2_FinalKeysAdded.png)(w:720)
 
-1. Add a new keyframe at the **20.00** mark for **Intensity** and **Right-Click** on it and **Set Value** to **5**. 
+1. **Add** a new keyframe at the **20.00** mark for Intensity and **Right-Click** on it and Set Value to **5**. 
 
 	![](MHT4_C2_FinalGrainKey.png)(w:720)
 
-1. Select the keyframe at **15.50** on the **Movement** track of **Camera2**, then minimize the **Matinee** window. 
+1. **Select** the keyframe at **15.50** on the Movement track of Camera2, then minimize the Matinee window. 
 
-1. In the level viewport, move the **CameraActor2** somewhere near the couch with the brush in the background as seen below. 
+1. In the level viewport, **move** the CameraActor2 somewhere near the couch with the brush in the background as seen below. 
 
 	![](MHT4_C2_FinalPlacement.png)(w:640)
 
-1. With the **CameraActor2** still selected, in the **Details** panel find the **Scene Color** section and enable **Grain Intensity**. 
+1. With the CameraActor2 still selected, in the Details panel find the **Scene Color** section and **enable** **Grain Intensity**. 
 
 	![](MHT4_C2_EnableIntensity.png)
 
-	_Leave the setting as default since it will be updated via the **Matinee**._
+	_Leave the setting as default since it will be updated via the Matinee._
 
 
 [OBJECT:TopicList]
@@ -341,29 +341,29 @@ In this How To series, we are going to create a short sequence that uses multipl
 	[/PARAM]
 [/OBJECT]
 
-1. Inside the **Matinee Editor** in an empty space under the last track of **Camera2**, **Right-Click** and from the context menu choose **Add New Director Group**. 
+1. Inside the Matinee Editor in an empty space under the last track of Camera2, **Right-Click** and from the context menu **choose** Add New Director Group. 
 
 	![](MHT4_AddDirectorGroup.png)(h:400)
 
-1. In the new **Director** track, with the **Time Bar** at **0.00**, press **Enter** to add a new key and choose **Camera1** from the pop-up menu. 
+1. In the new **Director** track, with the Time Bar at **0.00**, **press Enter** to add a new key and **choose** Camera1 from the pop-up menu. 
 
 	![](MHT4_DirectorTrackkey1.png)(w:800)
 
-	_**Camera1** has been assigned for the shot. Note: the color of the bars may be different._
+	_Camera1 has been assigned for the shot. Note: the color of the bars may be different._
 
-1. Scrub the **Time Bar** to **6.50** and press **Enter** to add a keyframe, then select **Camera2** to add the new shot. 
+1. Scrub the Time Bar to **6.50** and **press Enter** to add a keyframe, then **select** Camera2 to add the new shot. 
 
 	![](MHT4_SecondCutAdded.png)(w:800)
 	
-1. Scrub the **Time Bar** to **12.00** and add a keyframe for **Camera1**, then to **15.50** to add a keyframe for **Camera2**. 
+1. Scrub the Time Bar to **12.00** and **add** a keyframe for Camera1, then to **15.50** to add a keyframe for Camera2. 
 
 	![](MHT4_AllCutsAdded.png)(w:800)
 	
-1. **Right-Click** on the **DirGroup** and in the context menu, select **Add New Fade Track**. 
+1. **Right-Click** on the **DirGroup** and in the context menu, **select** **Add New Fade Track**. 
 
 	![](MHT4_AddFadeTrack.png)(h:400)
 
-1. For the **Fade** track, add a keyframe at **0.00** (value set to **1**) and a keyframe at **1.00** (value set to **0**). 
+1. For the **Fade** track, **add** a keyframe at **0.00** (value set to **1**) and a keyframe at **1.00** (value set to **0**). 
 
 	![](MHT4_Fade1.png)
 
@@ -375,21 +375,21 @@ In this How To series, we are going to create a short sequence that uses multipl
 
 	_This will fade out shot 1 and fade in shot 2._
 
-1. Add a keyframe at **12.00** (value set to **1**) and keyframes at ***11.00** and **13.00** (both value set to **0**). 
+1. **Add** a keyframe at **12.00** (value set to **1**) and keyframes at ***11.00** and **13.00** (both value set to **0**). 
 
 	![](MHT4_Fade3.png)(w:720)
 
-1. Last, add a keyframe at **19.00** (value set to **0**) and a keyframe at the end of the sequence (value set to **1**). 
+1. Last, **add** a keyframe at **19.00** (value set to **0**) and a keyframe at the end of the sequence (value set to **1**). 
 
 	![](MHT4_Fade4.png)(w:720)
 
-1. Close the **Matinee** window and select the **Matinee Actor** in the **World Outliner**. 
+1. **Close** the Matinee window and **select** the Matinee Actor in the World Outliner. 
 
-1. In the **Details** panel, enable the options **Play on Level Load**, **Looping**, **Disable Movement Input**, **Disable Look at Input**, and **Hide Player**. 
+1. In the Details panel, **enable** the options **Play on Level Load**, **Looping**, **Disable Movement Input**, **Disable Look at Input**, and **Hide Player**. 
 
 	![](MHT4_MatineeActorOptions.png)(h:500)
 
-1. On the main toolbar, click the **Build** icon (option to build lighting if you wish) and then the **Play** icon to play in editor. 
+1. On the main toolbar, **click** the **Build** icon (option to build lighting if you wish) and then the **Play** icon to play in editor. 
 
 	[![](MovieThumbnail.png)](MHT4_FinalVideo.wmv)
 	
