#UE4 Docs: initial cxheckin of linux specs info

diff --git a/Engine/Documentation/Source/GettingStarted/RecommendedSpecifications/Specifications.INT.udn b/Engine/Documentation/Source/GettingStarted/RecommendedSpecifications/Specifications.INT.udn
--- a/Engine/Documentation/Source/GettingStarted/RecommendedSpecifications/Specifications.INT.udn
+++ b/Engine/Documentation/Source/GettingStarted/RecommendedSpecifications/Specifications.INT.udn
@@ -1,10 +1,11 @@
 Availability:Public
 Title:Hardware & Software Specifications
+SEO-Title:Hardware & Software Specifications
 Crumbs: %ROOT%, GettingStarted
 Description:Minimum and recommended hardware specifications and necessary software for developing with Unreal Engine.
 
 [PUBLISH:Rocket]
-%Globals:OSSelect%
+%Globals:OSSelect_Linux%
 [/PUBLISH:Rocket]
 
 [TOC(start:2)]
@@ -47,23 +48,47 @@ Unreal Engine has some specific hardware and software requirements for running t
 	[INCLUDE:#macosminhardware]
 	[/PARAM]
 [/OBJECT]
+[OBJECT:ToggleButtonContent]
+	[PARAMLITERAL:category]
+	OS
+	[/PARAMLITERAL]
+	[PARAMLITERAL:id]
+	linux
+	[/PARAMLITERAL]
+	[PARAMLITERAL:active]
+	[/PARAMLITERAL]
+	[PARAM:content]
+	[INCLUDE:#linuxminhardware]
+	[/PARAM]
+[/OBJECT]
 [/PUBLISH:Rocket]
 
 <!--
 [EXCERPT:WindowsMinHardware]
-| Running/Developing with the Engine ||
+|||
 | -------- | --- |
+| **Operating System** | Windows 7/8 64-bit |
 | **Processor** | Quad-core Intel or AMD, 2.5 GHz or faster |
 | **Memory** | 8 GB RAM |
-| **Video Card** | NVIDIA GeForce 470 GTX or AMD Radeon 6870 HD series or higher |
+| **Video Card** | DirectX 11 compatible graphics card |
+| **DirectX Version** | [DirectX End-User Runtimes (June 2010)](http://www.microsoft.com/en-us/download/details.aspx?id=8109) |
 [/EXCERPT:WindowsMinHardware]
 [EXCERPT:MacOSMinHardware]
-| Running/Developing with the Engine ||
+|||
 | -------- | --- |
-| **Processor** | Quad-core Intel or AMD processor, 2.5 GHz or faster |
+| **Operating System** | Mac OS X 10.9.2  |
+| **Processor** | Quad-core Intel, 2.5 GHz or faster |
 | **Memory** | 8 GB RAM |
-| **Video Card** | NVIDIA GeForce 470 GTX or AMD Radeon 6870 HD series or higher |
+| **Video Card** | DirectX 11 compatible graphics card |
 [/EXCERPT:MacOSMinHardware]
+[EXCERPT:LinuxMinHardware]
+|||
+| -------- | --- |
+| **Operating System** | Ubuntu 15.04  |
+| **Processor** | Quad-core Intel or AMD, 2.5 GHz or faster |
+| **Memory** | 16 GB RAM |
+| **Video Card** | 	NVIDIA GeForce 470 GTX or higher with latest NVidia binary drivers |
+[/EXCERPT:LinuxMinHardware]
 -->
 
 [PUBLISH:Licensee]
@@ -78,7 +103,7 @@ Unreal Engine has some specific hardware and software requirements for running t
 | -------- | --- |
 | **Processor** | 4 core > 2.5 GHz or better |
 | **Memory** | 8+ GB RAM |
-| **Video Card** | DirectX 11 with 2 GB GDDR5 memory |
+| **Video Card** | DirectX 11 compatible graphics card with 2 GB GDDR5 memory |
 | **Monitor** | 2 monitors @ 1680x1050 or better |
 [/EXCERPT:WindowsRecHardware]
 [EXCERPT:MacOSRecHardware]
@@ -86,9 +111,17 @@ Unreal Engine has some specific hardware and software requirements for running t
 | -------- | --- |
 | **Processor** | 4 core > 2.5 GHz or better |
 | **Memory** | 8+ GB RAM |
-| **Video Card** | 2 GB GDDR5 memory |
+| **Video Card** | DirectX 11 compatible graphics card with 2 GB GDDR5 memory |
 | **Monitor** | 2 monitors @ 1680x1050 or better |
 [/EXCERPT:MacOSRecHardware]
+[EXCERPT:LinuxRecHardware]
+| Developing with the Engine ||
+| -------- | --- |
+| **Processor** | 4 core > 2.5 GHz or better |
+| **Memory** | 8+ GB RAM |
+| **Video Card** | DirectX 11 compatible graphics card with 2 GB GDDR5 memory |
+| **Monitor** | 2 monitors @ 1680x1050 or better |
+[/EXCERPT:LinuxRecHardware]
 -->
 
 [PUBLISH:Licensee]
@@ -144,18 +177,27 @@ Below are the requirements for running the engine or editor.
 	[INCLUDE:#macosminsoftware]
 	[/PARAM]
 [/OBJECT]
+[OBJECT:ToggleButtonContent]
+	[PARAMLITERAL:category]
+	OS
+	[/PARAMLITERAL]
+	[PARAMLITERAL:id]
+	linux
+	[/PARAMLITERAL]
+	[PARAMLITERAL:active]
+	[/PARAMLITERAL]
+	[PARAM:content]
+	[INCLUDE:#linuxminsoftware]
+	[/PARAM]
+[/OBJECT]
 
 <!--
 [EXCERPT:WindowsMinSoftware]
 | Running the Engine ||
 | -------- | --- |
-| **Operating System** | Windows 7/8 64-bit [INCLUDE:#WER] |
+| **Operating System** | Windows 7/8 64-bit |
 | **DirectX Version** | [DirectX End-User Runtimes (June 2010)](http://www.microsoft.com/en-us/download/details.aspx?id=8109) |
-| **.NET Version** | .NET 4.0 (should already be installed via Windows Update) |
 
-[REGION:note]
-All of the required software to run the engine is installed automatically by the installer.
-[/REGION]
 [/EXCERPT:WindowsMinSoftware]
 
 [EXCERPT:MacOSMinSoftware]
@@ -164,6 +206,14 @@ All of the required software to run the engine is installed automatically by the
 | **Operating System** | Mac OS X 10.9.2  |
 [/EXCERPT:MacOSMinSoftware]
 
+[EXCERPT:LinuxMinSoftware]
+| Running the Engine ||
+| -------- | --- |
+| **Operating System** | Any reasonable new Linux distro (CentOS 6.x being the oldest tested)  |
+| **Linux Kernel Version** | kernel 2.6.32 or newer  |
+| **Additional Dependencies** | glibc 2.12.2 or newer  |
+[/EXCERPT:LinuxMinSoftware]
+
 [EXCERPT:WER]
 [REGION:note]
 UE4's crash reporting feature requires you to have **Windows Error Reporting** enabled, which is the default. If you have disabled this feature, please enable it.
@@ -230,6 +280,19 @@ The requirements for programmers developing with the engine are listed below.
 	[INCLUDE:#macosmindevelopment]
 	[/PARAM]
 [/OBJECT]
+[OBJECT:ToggleButtonContent]
+	[PARAMLITERAL:category]
+	OS
+	[/PARAMLITERAL]
+	[PARAMLITERAL:id]
+	linux
+	[/PARAMLITERAL]
+	[PARAMLITERAL:active]
+	[/PARAMLITERAL]
+	[PARAM:content]
+	[INCLUDE:#linuxmindevelopment]
+	[/PARAM]
+[/OBJECT]
 
 <!--
 [EXCERPT:WindowsMinDevelopment]
@@ -245,15 +308,21 @@ The requirements for programmers developing with the engine are listed below.
 | -------- | --- |
 | **Xcode Version** | 5.0 |
 [/EXCERPT:MacOSMinDevelopment]
+[EXCERPT:LinuxMinDevelopment]
+| Developing with the Engine ||
+| -------- | --- |
+| **Operating System** | Ubuntu 15.04  |
+| **Compiler** | clang 3.5.x or 3.6.x  |
+| [REGION:tablesection]Optional[/REGION] ||
+| **IDE** | QtCreator 3.1 or KDevelop 4.6|
+| **Build System** | CMake 2.8 |
+[/EXCERPT:LinuxMinDevelopment]
 [EXCERPT:VS2013]
-Visual Studio 2013 Professional or Visual Studio 2013 Express for Windows Desktop:
+Visual Studio 2013 Professional or Visual Studio 2013 Community:
 
 * [Pro](http://www.microsoft.com/en-us/download/details.aspx?id=43730)
-* [Express](http://www.microsoft.com/en-us/download/details.aspx?id=43733)
+* [Community](https://www.visualstudio.com/en-us/products/visual-studio-community-vs.aspx)
 
-[REGION:note]
-It is recommended, but not required, to have the latest Visual Studio Update installed as well, for either version.  Visual Studio will prompt you to update to the latest version after you have installed and run the program.
-[/REGION]
 [/EXCERPT:VS2013]
 -->
 [/PUBLISH:Rocket]
