#UEDoc #docs - updated images. Gif needs to be converted to video

diff --git a/Engine/Documentation/Source/Engine/Matinee/HowTo/MHT_5/MHT_5.INT.udn b/Engine/Documentation/Source/Engine/Matinee/HowTo/MHT_5/MHT_5.INT.udn
--- a/Engine/Documentation/Source/Engine/Matinee/HowTo/MHT_5/MHT_5.INT.udn
+++ b/Engine/Documentation/Source/Engine/Matinee/HowTo/MHT_5/MHT_5.INT.udn
@@ -18,11 +18,11 @@ In the following steps, we will create a light that moves, changes colors and br
 	[/PARAM]
 [/OBJECT]
 
-1. From the **Toolbar** menu, select **Matinee** and click the **Add Matinee** button. 
+1. From the **Toolbar** menu, **select** **Matinee** and **click** the **Add Matinee** button. 
 
 	![](Engine\Matinee\HowTo\MHT_1\MHT1_AddMatinee.png)(w:720)
 
-1. If the **Matinee** undo warning message appears, click the **Continue** button. 
+1. If the Matinee undo warning message appears, **click** the **Continue** button. 
 
 	![](Engine\Matinee\HowTo\MHT_1\MHT1_MatineeWarning.png)
 	
@@ -34,7 +34,7 @@ In the following steps, we will create a light that moves, changes colors and br
 
 	_Click for Full View._
 
-1.  Minimize the **Matinee** window, then from the **Modes** menu under the **Geometry** tab, **Drag-and-drop** a **Box** into the viewport.
+1.  Minimize the Matinee window, then from the **Modes** menu under the **BSP** tab, **Drag-and-drop** a **Box** into the viewport.
 
 	![](MHT5_DragInBox.png)(w:580)
 
@@ -44,13 +44,13 @@ In the following steps, we will create a light that moves, changes colors and br
 
 	![](MHT5_BoxSettings.png)(w:380)
 
-1.  Fly inside the box and grab the red arrow of the **Translation** widget and move it towards the text to remove the pieces of geometry that clip inside. 
+1.  Fly inside the box and **grab** the red arrow of the **Translation** widget and **move it towards** the text to remove the pieces of geometry that clip inside. 
 
 	![](MHT5_MoveBox.png)(w:380)
 
-	_You can also grab the blue arrow of the **Translation** widget and pull it up slightly to raise the roof of the box._
+	_You can also grab the blue arrow of the Translation widget and pull it up slightly to raise the roof of the box._
 
-1. From the **Geometry** tab, **Drag-and-drop** another **Box** into the viewport and set its settings to the values below. 
+1. From the BSP tab, **Drag-and-drop** another **Box** into the viewport and set its settings to the values below. 
 
 	![](MHT5_SubtractBrush.png)(w:380)
 
@@ -58,11 +58,11 @@ In the following steps, we will create a light that moves, changes colors and br
 
 	![](MHT5_Opening1.png)(w:380)
 
-1. With the **Subtraction** box still selected, press **Control+W** to duplicate it and move it to the other end of the box to create another opening. 
+1. With the Subtraction box still selected, **press** **Control+W** to duplicate it and move it to the other end of the box to create another opening. 
 
 	![](MHT5_Opening2.png)(w:380)
 
-1. From the **Lights** tab of the **Modes** menu, **Drag-and-drop** a **Point Light** into the level and center it near the top of one of the door openings. 
+1. From the **Lights** tab of the Modes menu, **Drag-and-drop** a **Point Light** into the level and center it near the top of one of the door openings. 
 
 	![](MHT5_AddPointLight.png)
 
@@ -75,31 +75,31 @@ In the following steps, we will create a light that moves, changes colors and br
 	[/PARAM]
 [/OBJECT]
 
-1. Open **Matinee** by clicking on the **Matinee Actor** in the **World Outliner** and choosing **Open Matinee** from the **Details** panel. 
+1. Open Matinee by **clicking** on the Matinee Actor in the **World Outliner** and choosing **Open Matinee** from the Details panel. 
 
 	![](MHT5_OpenMatinee.png)
 
-1. Minimize **Matinee**, click on the **Point Light**, then re-open **Matinee** and **Right-Click** in the **Tracks** window and select **Add New Lighting Group**. 
+1. **Minimize** Matinee, **click** on the Point Light then re-open Matinee and **Right-Click** in the **Tracks** window and select **Add New Lighting Group**. 
 
 	![](MHT5_AddLightingGroup.png)
 
 	_In the **Name Group** box that appears, give it a name such as **Light1**._
 
-1. Grab the ending marker at **5.00** and drag it over to **8.00** to increase the length of the **Matinee**. 
+1. **Grab** the ending marker at **5.00** and **drag** it over to **8.00** to increase the length of the Matinee. 
 
 	![](MHT5_IncreaseLength.png)(w:720)
 
-1. **Right-Click** on the **Radius** track and select **Delete Track**, then **Right-Click** on the **Light1** group and select **Add New Float Property Track**. 
+1. **Right-Click** on the **Radius** track and select **Delete Track**, then **Right-Click** on the Light1 group and **select** **Add New Float Property Track**. 
 
 	![](MHT5_AddPropertyTrack.png)
 
 	_We will be adjusting the **Attenuation Radius** which will affect the range at which the light is displayed._
 
-1. In the pop-up menu that appears, select **LightComponent0.AttenuationRadius** and press **Ok** to add the track. 
+1. In the pop-up menu that appears, select **LightComponent0.AttenuationRadius** and **press** **Ok** to add the track. 
 
 	![](MHT5_AttenuationRadius.png)
 
-1. Click the **Movement** track, and press **Enter** to add a keyframe then **Right-Click** on the keyframe and choose **Set Time** and set it to **2**. 
+1. Click the **Movement** track, and **press** **Enter** to add a keyframe then **Right-Click** on the keyframe and choose **Set Time** and set it to **2**. 
 
 1. Repeat the previous step and assign keys to **0**, **2**, **4**, **6**, and **8**. 
 
@@ -119,23 +119,23 @@ In the following steps, we will create a light that moves, changes colors and br
 	[/PARAM]
 [/OBJECT]
 
-1. In **Matinee** on the **Movement** track, click on the second keyframe (at **2.00**) then minimize **Matinee** and move the light to the center of the hallway. 
+1. In Matinee on the Movement track, **click** on the second keyframe (at **2.00**) then minimize Matinee and **move** the light to the center of the hallway. 
 
 	![](MHT5_MovementKey2.png)(w:720)
 
 	_You can zoom out then grab green arrow of the **Translation** widget and slide it to the right into the center._
 
-1. Return to **Matinee** and click on the third keyframe of the **Movement** track (at **4.00**), minimize **Matinee**, then move the light to the end of the hallway. 
+1. **Return** to Matinee and **click** on the third keyframe of the Movement track (at **4.00**), **minimize** Matinee, then **move** the light to the end of the hallway. 
 
 	![](MHT5_MovementKey3.png)(w:720)
 
-1. Return to **Matinee**, click on the fourth keyframe of the **Movement** track (at **6.00**), minimize **Matinee** and then move the light to the center of the hallway. 
+1. **Return** to Matinee, **click** on the fourth keyframe of the Movement track (at **6.00**), **minimize** Matinee and then **move** the light to the center of the hallway. 
 
-1. In the **Details** panel for the light, under **Transform**, find the **Mobility** section and click the third icon to enable the **Movable** setting for the light. 
+1. In the Details panel for the light, under **Transform**, find the **Mobility** section and **click** the third icon to enable the **Movable** setting for the light. 
 
 	![](MHT5_LightMobility.png)
 
-1. In **Matinee**, **Right-click** on the second keyframe of the **Intensity** track (at **2.00**) and **Set Value** to **20,000**; do the same for the fourth keyframe (at **6.00**). 
+1. In Matinee, **Right-click** on the second keyframe of the **Intensity** track (at **2.00**) and Set Value to **20,000**; do the same for the fourth keyframe (at **6.00**). 
 
 	![](MHT5_IntensityValues.png)(w:720)
 
@@ -151,15 +151,15 @@ In the following steps, we will create a light that moves, changes colors and br
 	[/PARAM]
 [/OBJECT]
 
-1. In **Matinee** on the **Light Color** track, click on the first keyframe (at **0.00**) and select **Set Color**; in the **Color Picker** window select any color. 
+1. In Matinee on the **Light Color** track, **click** on the first keyframe (at **0.00**) and **select** **Set Color**; in the **Color Picker** window select any color. 
 
 	![](MHT5_RedLight.png)(w:720)
 
-1. Repeat the previous step for the third keyframe (at **4.00**) and in the **Color Picker** window, select a different color. 
+1. **Repeat** the previous step for the third keyframe (at **4.00**) and in the Color Picker window, **select** a different color. 
 
-1. Repeat the previous step for the last keyframe (at **8.00**) and in the **Color Picker** window, select the color that was used in step 1. 
+1. **Repeat** the previous step for the last keyframe (at **8.00**) and in the Color Picker window, **select** the color that was used in step 1. 
 
-1. **Right-click** on the second keyframe of the **Attenuation Radius** track (at **2.00**) and **Set Value** to **250**, do the same for the fourth keyframe (at **6.00**). 
+1. **Right-click** on the second keyframe of the **Attenuation Radius** track (at **2.00**) and Set Value to **250**, do the same for the fourth keyframe (at **6.00**). 
 
 	![](MHT5_AttenuationKey4.png)(w:720)
 
@@ -175,11 +175,11 @@ In the following steps, we will create a light that moves, changes colors and br
 	[/PARAM]
 [/OBJECT]
 
-1. In the **World Outliner**, select the **Matinee Actor** and under the **Play** section, enable **Play on Level Load** and **Looping**. 
+1. In the World Outliner, **select** the Matinee Actor and under the **Play** section, **enable** Play on Level Load and Looping. 
 
 	![](MHT5_MatineeSettings.png)
 
-1. From the main toolbar, click the **Build** icon, then when building is complete, select the **Apply Now** button in the lower right portion of the screen. 
+1. From the main toolbar, **click** the **Build** icon, then when building is complete, **select** the **Apply Now** button in the lower right portion of the screen. 
 
 	![](MHT5_BuildIcon.png)
 
@@ -189,7 +189,7 @@ In the following steps, we will create a light that moves, changes colors and br
 
 	_will appear. Click the **Apply Now** button."_
 
-1. From the main toolbar, click the **Play** icon to play in the editor. 
+1. From the main toolbar, **click** the Play icon to play in the editor. 
 
 	![](MHT5_LightingGif.gif)(convert:false)
 
