#UE4doc #docs: Removed ôOverviewö header

diff --git a/Engine/Documentation/Source/Engine/UI/LevelEditor/Viewports/ActorLocking/ActorLocking.INT.udn b/Engine/Documentation/Source/Engine/UI/LevelEditor/Viewports/ActorLocking/ActorLocking.INT.udn
--- a/Engine/Documentation/Source/Engine/UI/LevelEditor/Viewports/ActorLocking/ActorLocking.INT.udn
+++ b/Engine/Documentation/Source/Engine/UI/LevelEditor/Viewports/ActorLocking/ActorLocking.INT.udn
@@ -45,7 +45,7 @@ Description:A look at how to lock Actors to the viewport camera for intuitive pl
 
 [TOC(start:2)]
 
-## Overview
+
 
 Actor Locking allows for the perspective Viewport camera to be used as a placement mechanism for Actors within your levels. Simply lock to a given Actor, position the view where you want that Actor to be placed, and then unlock. The position and orientation of the camera will be used to position the given Actor.
 
