#UEDoc #docs - updated doc and images. Need to convert gifs to video

diff --git a/Engine/Documentation/Source/Engine/Matinee/HowTo/MHT_6/MHT_6.INT.udn b/Engine/Documentation/Source/Engine/Matinee/HowTo/MHT_6/MHT_6.INT.udn
--- a/Engine/Documentation/Source/Engine/Matinee/HowTo/MHT_6/MHT_6.INT.udn
+++ b/Engine/Documentation/Source/Engine/Matinee/HowTo/MHT_6/MHT_6.INT.udn
@@ -18,11 +18,11 @@ In the following steps, we will create a blinking **Material** on an object and
 	[/PARAM]
 [/OBJECT]
 
-1. From the **Toolbar** menu, select **Matinee** and click the **Add Matinee** button. 
+1. From the **Toolbar** menu, **select** Matinee and **click** the **Add Matinee** button. 
 
 	![](Engine\Matinee\HowTo\MHT_1\MHT1_AddMatinee.png)(w:720)
 
-1. If the **Matinee** undo warning message appears, click the **Continue** button. 
+1. If the Matinee undo warning message appears, **click** the **Continue** button. 
 
 	![](Engine\Matinee\HowTo\MHT_1\MHT1_MatineeWarning.png)
 	
@@ -34,47 +34,47 @@ In the following steps, we will create a blinking **Material** on an object and
 
 	_Click for Full View._
 
-1.  Close the **Matinee** then repeat steps 1-3 to create a second **Matinee Actor** so that there are two in the level. 
+1.  **Close** the Matinee then **repeat** steps 1-3 to create a second Matinee Actor so that there are two in the level. 
 
-1.  Close the second **Matinee** window, then in the **Content Browser** under **Game/Props**, locate the **SM_Chair** and drag it into the level. 
+1.  **Close** the second Matinee window, then in the **Content Browser** under **Game/Props**, locate the **SM_Chair** and **drag** it into the level. 
 
 	![](MHT6_AddChair.png)(w:740)
 
-1.  Also in the **Game/Props** folder, drag into the level the **SM_Stairs**. 
+1.  Also in the Game/Props folder, **drag** into the level the **SM_Stairs**. 
 
 	![](MHT6_AddStairs.png)(w:740)
 
 	_You can position both the chair and stairs as you wish in the level._
 
-1. From the **Modes** menu on the **Basic** tab, click and drag in the **Box Trigger** into the level then press **R** and scale the box so that the chair is inside it. 
+1. From the **Modes** menu on the **Basic** tab, click and drag in the **Box Trigger** into the level then **press** **R** and scale the box so that the chair is inside it. 
 
 	![](MHT6_AddTrigger.png)(w:740)
 
 	_When the player enters this trigger box, the chair will change colors which we will set up later on this page._
 
-1. Click on the **MatineeActor** (the first **Matinee** that was added) in the **World Outliner**, then in the **Details** panel click **Open Matinee**. 
+1. **Click** on the Matinee Actor (the first Matinee that was added) in the **World Outliner**, then in the **Details** panel **click** **Open Matinee**. 
 
 	![](MHT6_OpenMatinee.png)
 
-1. Minimize the **Matinee** and click on the **SM_Chair** then re-open the **Matinee** and in the **Tracks Window**, **Right-click** and select **Add New Empty Group**. 
+1. **Minimize** the Matinee and **click** on the SM_Chair then **re-open** the Matinee and in the **Tracks Window**, **Right-click** and **select** **Add New Empty Group**. 
 
 	![](MHT6_AddNewEmptyGroup.png)
 
 	_In the **Name Group** window, name the group something descriptive such as **Chair**._
 
-1. **Right-click** on the new group in **Matinee** and choose **Add New Vector Material Parameter Track**. 
+1. **Right-click** on the new group in Matinee and **choose** **Add New Vector Material Parameter Track**. 
 
 	![](MHT6_AddVectorMaterialTrack.png)
 
-1. Close the **Matinee** and then open the other **Matinee Actor**. 
+1. **Close** the Matinee and then open the other Matinee Actor. 
 
-1. Minimize the **Matinee** and click on the **SM_Stairs** then re-open the **Matinee** and in the **Tracks Window**, **Right-click** and select **Add New Empty Group**. 
+1. **Minimize** the Matinee and **click** on the SM_Stairs then **re-open** the Matinee and in the Tracks Window, **Right-click** and **select** Add New Empty Group. 
 
 	![](MHT6_AddStairsToMatinee.png)
 
 	_In the **Name Group** window, name the group something descriptive such as **Stairs**._
 
-1. **Right-click** on the new group in **Matinee** and choose **Add New Float Material Parameter Track**. 
+1. **Right-click** on the new group in Matinee and **choose** **Add New Float Material Parameter Track**. 
 
 	![](MHT6_AddFloatMaterialTrack.png)
 
@@ -88,105 +88,105 @@ In the following steps, we will create a blinking **Material** on an object and
 	[/PARAM]
 [/OBJECT]
 
-1.  Select the chair in the level, then in the **Details** panel under **Materials**, click the ![](MHT6_BrowseToIcon.png) icon. 
+1.  **Select** the chair in the level, then in the Details panel under **Materials**, **click** the ![](MHT6_BrowseToIcon.png) icon. 
 
 	![](MHT6_MaterialsSearch.png)
 
-1.  The **Material** used for the chair will be highlighted in the **Content Browser**, **Double-click** on the **M_Chair Material** to open it in the **Material Editor**. 
+1.  The Material used for the chair will be highlighted in the **Content Browser**, **Double-click** on the **M_Chair Material** to open it in the **Material Editor**. 
 
-1.  **Mouse-Wheel Down** to zoom out and find the node with the orange square and click on it, then in the **Details** panel, click the **Default Value** bar. 
+1.  **Mouse-Wheel Down** to zoom out and find the node with the orange square and **click** on it, then in the Details panel, **click** the **Default Value** bar. 
 
 	![](MHT6_ChairDefaultColor.png)	
 
-1.  In the **Color Picker** window, click and drag the orange color (in the yellow box) to the **Save Color** bar to save the color. 
+1.  In the **Color Picker** window, **click and drag** the orange color (in the yellow box) to the **Save Color** bar to save the color. 
 
 	![](MHT6_ColorPicker.png)
 
 	_We will use this later as now we have a reference to the default color of the seats._
 
-1.  Close the **Color Picker** window (no need to save) and close the **Material** as well. 
+1.  **Close** the **Color Picker** window (no need to save) and close the Material as well. 
 
-1.	In the **Content Browser**, **Right-click** on the **M_Chair Material** and choose **Create Material Instance**. 
+1.	In the Content Browser, **Right-click** on the **M_Chair Material** and **choose** **Create Material Instance**. 
 
 	![](MHT6_MaterialInstance1.png)
 
-	_A **Material Instance** of the **Material** will be created in the **Content Browser** which we can edit. You can leave the default name._
+	_A **Material Instance** of the Material will be created in the Content Browser which we can edit. You can leave the default name._
 
-1.	Select the chair in the level, then select the **M_Chair_Inst** in the **Content Browser** and in the **Details** panel, click the ![](MHT6_ArrowIcon.png) icon. 
+1.	**Select** the chair in the level, then **select** the **M_Chair_Inst** in the Content Browser and in the Details panel, **click** the ![](MHT6_ArrowIcon.png) icon. 
 
 	![](MHT6_AssignNewMaterial.png)(w:780)
 	
 	_This will assign the **Material Instance** to the chair._
 
-1.	Open the **MatineeActor** (the **Matinee** assigned to the chair) and click on the **...Material Param** (**Material Parameter** track). 
+1.	**Open** the Matinee Actor (the Matinee assigned to the chair) and click on the **...Material Param** (**Material Parameter** track). 
 
-1.	In the **Details** panel on the right, click the ![](MHT6_AddIcon.png) icon.
+1.	In the Details panel on the right, **click** the ![](MHT6_AddIcon.png) icon.
 
-1.	Back in the **Content Browser** make sure the **M_Chair_Inst Material Instance** is highlighted, then in **Matinee**, click the ![](MHT6_ArrowIcon.png) icon. 
+1.	Back in the Content Browser make sure the **M_Chair_Inst Material Instance** is highlighted, then in Matinee, **click** the ![](MHT6_ArrowIcon.png) icon. 
 
 	![](MHT6_AddInstanceToMatinee.png)
 
-1.	In the **Param Name** field, enter **ColorSeats** (as this is the name of the node in the **Material** from step 3) that we will alter. 
+1.	In the **Param Name** field, enter **ColorSeats** (as this is the name of the node in the Material from step 3) that we will alter. 
 
-1.	Click on the **Material Parameter** track and press **Enter** to add a keyframe at **0.00**, then move the **Time Bar** to **2.00** and add another keyframe. 
+1.	**Click** on the Material Parameter track and **press** **Enter** to add a keyframe at **0.00**, then **move** the **Time Bar** to **2.00** and add another keyframe. 
 
 	![](MHT6_M1_KeysAdded.png)
 
-1.	**Right-click** on the ending marker and choose **Move To Longest Track Endpoint** to reduce the **Matinee** length to two seconds. 
+1.	**Right-click** on the ending marker and **choose** **Move To Longest Track Endpoint** to reduce the Matinee length to two seconds. 
 
 	![](MHT6_ReduceMatineeLength.png)(w:720)
 
-1.	**Right-click** on the first keyframe and choose **Set Color**, then from the **Color Picker** menu, select the bar at the top to assign the default chair color. 
+1.	**Right-click** on the first keyframe and **choose** Set Color, then from the Color Picker menu, select the bar at the top to assign the default chair color. 
 
 	![](MHT6_Key1Chair.png)
 
-1.	**Right-click** on the second keyframe and choose **Set Color**, then from the **Color Picker** menu, select any color other than the default color then close the **Matinee**. 
+1.	**Right-click** on the second keyframe and **choose** Set Color, then from the Color Picker menu, **select** any color other than the default color then close the Matinee. 
 
 	![](MHT6_Key2Chair.png)
 
 	_This is the color the chair will blend to when the player approaches._
 
-1.	In the **World Outliner**, select the **MatineeActor**, then from the toolbar select **Blueprints - Open Level Blueprint**. 
+1.	In the World Outliner, **select** the Matinee Actor, then from the toolbar select **Blueprints - Open Level Blueprint**. 
 
 	![](MHT6_AddMatineeToBlueprint.png)
 
-1.	Inside the **Level Blueprint Event Graph** window, **Right-Click** and from the context menu, select **Add Reference to MatineeActor**. 
+1.	Inside the **Level Blueprint Event Graph** window, **Right-Click** and from the context menu, **select** **Add Reference to MatineeActor**. 
 
 	![](MHT6_AddMatineeReference.png)
 
-1.	Minimize the **Level Blueprint**, then in the **World Outliner** locate and select the **TriggerBox1** so that it is highlighted. 
+1.	**Minimize** the Level Blueprint, then in the World Outliner locate and **select** the **TriggerBox1** so that it is highlighted. 
 
-1.	Re-open the **Level Blueprint** and in the **Event Graph**, **Right-click** and select **Add Event for TriggerBox1**, **Collision** then **Add OnActorBeginOverlap**. 
+1.	**Re-open** the Level Blueprint and in the **Event Graph**, **Right-click** and **select** **Add Event for TriggerBox1**, **Collision** then **Add OnActorBeginOverlap**. 
 
 	![](MHT6_AddOnOverlap.png)
 
 	_This event will execute when an **Actor** (our player in this case) overlaps the trigger box._
 
-1.	Repeat the previous step, however this time select the **Add OnActorEndOverlap**. 
+1.	**Repeat the previous step**, however this time **select** the Add OnActorEndOverlap. 
 
-1.	**Click-and-drag** off the out pin of the **MatineeActor** node and release, in the context menu search for **Play**, then select **Play** from the menu. 
+1.	**Click-and-drag** off the out pin of the MatineeActor node and release, in the context menu search for **Play**, then **select** Play from the menu. 
 
 	![](MHT6_PlayMatinee.png)
 
-1.	Repeat the previous step, however this time search for and select **Reverse**. 
+1.	**Repeat the previous step**, however this time search for and **select** **Reverse**. 
 
 	![](MHT6_PlayAndReverse.png)
 
 	_The **Matinee** should be hooked up to both a **Play** and **Reverse** node._
 
-1.	**Click-and-drag** off the **Execute** out pin of the **OnActorBeginOverlap** node and connect it to the in **Execute** pin of the **Play** node. 
+1.	**Click-and-drag** off the **Execute** out pin of the OnActorBeginOverlap node and **connect** it to the in **Execute** pin of the **Play** node. 
 
 	![](MHT6_HookUpPlay.png)
 
-1.	Repeat the previous step, however drag off the **OnActorEndOverlap** node and connect it to the **Reverse** node. 
+1.	**Repeat the previous step**, however drag off the OnActorEndOverlap node and **connect** it to the **Reverse** node. 
 
 	![](MHT6_BlueprintComplete.png)
 
-1.	Click the **Compile** button located in the upper left of the **Blueprint** window. 
+1.	**Click** the **Compile** button located in the upper left of the Blueprint window. 
 
 	![](MHT6_Compile.png)
 
-1.	Close the **Blueprint**, then click the **Play** button from the main toolbar to play in the editor. 
+1.	**Close** the Blueprint, then click the Play button from the main toolbar to play in the editor. 
 
 	![](MHT6_M1_GIF.gif)(w:480 convert:false)
 
@@ -204,7 +204,7 @@ In the following steps, we will create a blinking **Material** on an object and
 	[/PARAM]
 [/OBJECT]
 
-1. In the **Content Browser** under **Game/Materials**, **Right-click** on the **M_Basic_Wall Material** and choose **Create Copy** and give it a name.
+1. In the Content Browser under **Game/Materials**, **Right-click** on the **M_Basic_Wall Material** and choose **Create Copy** and give it a name.
 
 1. **Double-click** on the copy of the **Material** to open it in the **Material Editor**, then click on and delete the **Color** node. 
 
@@ -214,53 +214,53 @@ In the following steps, we will create a blinking **Material** on an object and
 
 	_Click image for full view._
 
-1. **Right-click** in the graph window and type in **scalar**, then select the **ScalarParameter** option from the window. 
+1. **Right-click** in the graph window and type in **scalar**, then **select** the **ScalarParameter** option from the window. 
 
 	![](MHT6_ScalarParameter.png)
 
-1. Enter the name **StepsColor** for the node, connect it to the **Base Color** pin of the **Material Attributes** node. 
+1. Enter the name **StepsColor** for the node, **connect** it to the **Base Color** pin of the **Material Attributes** node. 
 
 	![](MHT6_StepsColor.png)
 
-1. Press the **Apply** and **Save** buttons, then close the **Material Editor** window. 
+1. Press the **Apply** and **Save** buttons, then **close** the **Material Editor** window. 
 
 	![](MHT6_ApplyAndSave.png)
 
-1.	Select the stairs in the level viewport so they are highlighted. 
+1.	**Select** the stairs in the level viewport so they are highlighted. 
 
-1.	Select the **M_Basic_Wall Material** in the **Content Browser** so it is highlighted. 
+1.	**Select** the **M_Basic_Wall Material** in the Content Browser so it is highlighted. 
 
-1.	In the **Details** panel for the stairs under **Materials**, click the ![](MHT6_ArrowIcon.png) icon for **Element 1** to assign the **Material**. 
+1.	In the Details panel for the stairs under Materials, **click** the ![](MHT6_ArrowIcon.png) icon for **Element 1** to assign the Material. 
 
 	![](MHT6_AssignBaseStairsMat.png)(w:740)
 
-1.	In the **Content Browser**, find the **Material** created in step 1 and select it so it is highlighted. 
+1.	In the Content Browser, find the Material created in step 1 and select it so it is highlighted. 
 
-1.	In the **Details** panel for the stairs under **Materials**, click the ![](MHT6_ArrowIcon.png) icon for **Element 0** to assign the **Material**. 
+1.	In the Details panel for the stairs under Materials, **click** the ![](MHT6_ArrowIcon.png) icon for **Element 0** to assign the Material. 
 
 	![](MHT6_StairsMaterials.png)
 
-1.	In the **World Outliner**, select **MatineeActor1** (the **Matinee** associated with the stairs) and in the **Details** panel, click the **Open Matinee** button. 
+1.	In the World Outliner, select **MatineeActor1** (the Matinee associated with the stairs) and in the Details panel, **click** the Open Matinee button. 
 
-1.	Click on the **...Material Param** (**Material Parameter**) track and in the **Details** panel under **Target Materials**, click the ![](MHT6_AddIcon.png) icon.
+1.	**Click** on the **...Material Param** (**Material Parameter**) track and in the Details panel under **Target Materials**, **click** the ![](MHT6_AddIcon.png) icon.
 
-1.	Back in the **Content Browser** make sure the **Material** from step 1 is highlighted, then in **Matinee**, click the ![](MHT6_ArrowIcon.png) icon. 
+1.	Back in the Content Browser make sure the Material from step 1 is highlighted, then in Matinee **click** the ![](MHT6_ArrowIcon.png) icon. 
 
 	![](MHT6_ApplyFloatMaterial.png)
 
-1.	Under the **Param Name** field, enter **StepsColor**. 
+1.	Under the **Param Name** field, **enter** **StepsColor**. 
 
-1.	On the **Material Parameter** track, add keyframes at **0.00**, **0.50** and **1.00** and move the end marker to **1.00** as well. 
+1.	On the Material Parameter track, add keyframes at **0.00**, **0.50** and **1.00** and **move** the end marker to **1.00** as well. 
 
 	![](MHT6_M2_KeysAdded.png)
 
-1.	**Right-click** on the second keyframe and select **Set Value** then set the value to **1** and close the **Matinee**. 
+1.	**Right-click** on the second keyframe and **select** Set Value then **set** the value to **1** and **close** the Matinee. 
 
-1.	Select the **MatineeActor1** in the **World Outliner**, then in the **Details** panel under **Play** check **Play on Level Load** and **Looping**. 
+1.	Select the **MatineeActor1** in the World Outliner, then in the Details panel under Play check Play on Level Load and Looping. 
 
 	![](MHT6_FinishingMatinee2.png)
 
-1.	From the main toolbar click the **Play** button to play in the editor. 
+1.	From the main toolbar click the Play button to play in the editor. 
 
 	![](MHT6_Matinee2Gif.gif)(w:480 convert:false)
 
