#UEDoc #docs - updated images and doc. need to convert gif to video

diff --git a/Engine/Documentation/Source/Engine/Matinee/HowTo/MHT_7/MHT_7.INT.udn b/Engine/Documentation/Source/Engine/Matinee/HowTo/MHT_7/MHT_7.INT.udn
--- a/Engine/Documentation/Source/Engine/Matinee/HowTo/MHT_7/MHT_7.INT.udn
+++ b/Engine/Documentation/Source/Engine/Matinee/HowTo/MHT_7/MHT_7.INT.udn
@@ -18,15 +18,15 @@ In the following steps, we will have a character jump up onto a platform by appl
 	[/PARAM]
 [/OBJECT]
 
-1. From the **Toolbar** menu, select **Matinee** and click the **Add Matinee** button. 
+1. From the **Toolbar** menu, **select** Matinee and **click** the **Add Matinee** button. 
 
 	![](Engine\Matinee\HowTo\MHT_1\MHT1_AddMatinee.png)(w:720)
 
-1. If the **Matinee** undo warning message appears, click the **Continue** button. 
+1. If the Matinee undo warning message appears, **click** the **Continue** button. 
 
 	![](Engine\Matinee\HowTo\MHT_1\MHT1_MatineeWarning.png)
 	
-1. The **Matinee Editor** window will open. 
+1. The Matinee Editor window will open. 
 
 	[REGION:lightbox]
 	[![](Engine\Matinee\HowTo\MHT_1\MHT1_MatineeOpen.png)(w:640)](Engine\Matinee\HowTo\MHT_1\MHT1_MatineeOpen.png)
@@ -34,19 +34,19 @@ In the following steps, we will have a character jump up onto a platform by appl
 
 	_Click for Full View._
 
-1.  Minimize the **Matinee** then in the **Content Browser** under the **Game/Character** folder, drag-and-drop in the **HeroTPP Skeletal Mesh** asset into the level. 
+1.  **Minimize** the Matinee then in the Content Browser under the **Game/Character** folder, **drag-and-drop** in the **HeroTPP Skeletal Mesh** asset into the level. 
 
 	![](MHT7_AddSkeleton.png)
 
-1.  In the **Details** panel for the **HeroTPP**, under **Transform** set the **Location** to **-610.0**, **-690.0**, **125.0** for the **X, Y, Z** values. 
+1.  In the Details panel for the **HeroTPP**, under **Transform** set the **Location** to **-610.0**, **-690.0**, **125.0** for the **X, Y, Z** values. 
 
 	![](MHT7_HeroDetails.png)
 
-1.  In the level, grab the small box brush above the template text and move it in front of the **Skeletal Mesh**. 
+1.  In the level, **grab** the small box brush above the template text and **move** it in front of the **Skeletal Mesh**. 
 
 	![](MHT7_MoveBoxBrush.png)
 	
-1.  In the **Details** panel for the box brush under **Brush Settings**, set the **X** and **Y** values to **250.0** and the **Z** to **1400.0**. 
+1.  In the Details panel for the box brush under **Brush Settings**, **set** the X and Y values to **250.0** and the Z to **1400.0**. 
 
 	![](MHT7_BoxDetails.png)
 
@@ -64,23 +64,23 @@ In the following steps, we will have a character jump up onto a platform by appl
 	[/PARAM]
 [/OBJECT]
 
-1.  Select the **MatineActor** from the **World Outliner**, then click the **Open Matinee** button. 
+1.  **Select** the MatineActor from the World Outliner, then **click** the **Open Matinee** button. 
 
 	![](MHT7_OpenMatineeActor.png)
 
-1.  Minimize the **Matinee**, then select the **HeroTPP** and re-open **Matinee**, then **Right-click** in the tracks window and select **Add New Skeletal Group**. 
+1.  **Minimize** the Matinee, then **select** the **HeroTPP** and **re-open** Matinee, then **Right-click** in the tracks window and **select** **Add New Skeletal Group**. 
 
 	![](MHT7_AddSkeletalGroup.png)
 
 	_Give the new group a name, in this example we are calling it **HeroTPP**._
 
-1.  **Right-click** on the **Group Name** and select **Add New Sound Track** from the menu. 
+1.  **Right-click** on the **Group Name** and **select** **Add New Sound Track** from the menu. 
 
 	![](MHT7_AddNewSoundTrack.png)
 
-1.  Repeat the previous step to add another **Sound Track** so that there are two. 
+1.  **Repeat the previous step** to add another **Sound Track** so that there are two. 
 
-1.  Grab the ending marker at **5.00** and move it to the **12.00** mark to increase the length of the **Matinee**. 
+1.  **Grab** the ending marker at **5.00** and **move** it to the **12.00** mark to increase the length of the Matinee. 
 
 	![](MHT7_MoveEndMarker.png)
 
@@ -94,59 +94,59 @@ In the following steps, we will have a character jump up onto a platform by appl
 	[/PARAM]
 [/OBJECT]
 
-1.  With **Matinee** open, click on the **Movement** track then move the **Time Bar** to **3.00** and press **Enter** to add a keyframe. 
+1.  With Matinee open, **click** on the Movement track then **move** the **Time Bar** to **3.00** and **press Enter** to add a keyframe. 
 
 	![](MHT7_MovementKey2.png)
 
-1.  Repeat the step above to add keyframes at **4.50**, **6.00**, **7.00**, and **8.00**. 
+1.  **Repeat the step above** to add keyframes at **4.50**, **6.00**, **7.00**, and **8.00**. 
 
 	![](MHT7_MovementKeysAdded.png)
 
-1.  Click on the second keyframe of the **Movement** track, minimize **Matinee**, then in the **Details** panel for **HeroTPP** set the **Location** of **Z** to **135**. 
+1.  **Click** on the second keyframe of the Movement track, minimize Matinee, then in the Details panel for HeroTPP **set** the **Location** of **Z** to **135**. 
 
 	![](MHT7_MoveDetailsKey2.png)
 
-1.  Return to **Matinee**, click on the third keyframe for **Movement**, minimize **Matinee** and set the **Location** of **Y** to **-674** and **Z** to **385**. 
+1.  **Return** to Matinee, **click** on the third keyframe for Movement, minimize Matinee and **set** the Location of Y to **-674** and Z to **385**. 
 
 	![](MHT7_MoveDetailsKey3.png)
 
-1.  Return to **Matinee**, click on the fourth keyframe for **Movement**, minimize **Matinee** and set the **Location** of **Y** to **-622.53** and **Z** to **804.640**. 
+1.  **Return** to Matinee, **click** on the fourth keyframe for Movement, minimize Matinee and **set** the Location of Y to **-622.53** and Z to **804.640**. 
 
 	![](MHT7_MoveDetailsKey4.png)
 
-1.  Return to **Matinee**, click on the fifth keyframe for **Movement**, minimize **Matinee** and set the **Location** of **Y** to **-561.0** and **Z** to **872.0**. 
+1.  **Return** to Matinee, **click** on the fifth keyframe for Movement, minimize Matinee and **set** the Location of Y to **-561.0** and Z to **872.0**. 
 
 	![](MHT7_MoveDetailsKey5.png)
 
-1.  Return to **Matinee**, click on the sixth keyframe for **Movement**, minimize **Matinee** and set the **Location** of **Y** to **-502.0** and **Z** to **847.0**. 
+1.  **Return** to Matinee, **click** on the sixth keyframe for Movement, minimize Matinee and **set** the Location of Y to **-502.0** and Z to **847.0**. 
 
 	![](MHT7_MoveDetailsKey6.png)
 
-1.  With the final keyframe in the **Movement** track still selected, in the level viewport, the **HeroTPP** Actor should be standing on the top of the platform. 
+1.  With the final keyframe in the Movement track still selected in the level viewport, the HeroTPP Actor should be standing on the top of the platform. 
 
 	![](MHT7_FinalMovementKeyframe.png)
 
 	_If your Actor is not standing on top of the platform, you can move him (and adjust the snapping size in the box above) so that he is standing on it._
 
-1.  Move the **Time Bar** back to **0.00**, select the **Anim** track and press **Enter** to add a keyframe, then select **Idle** from the menu that appears. 
+1.  **Move** the **Time Bar** back to **0.00**, **select** the Anim track and **press Enter** to add a keyframe, then **select** **Idle** from the menu that appears. 
 
 	![](MHT7_AnimKey1.png)
 
 	_If you look in the level viewport, the **HeroTPP** Actor will now be in an **Idle** pose since we have specified the animation to use to start the **Matinee**._ 
 
-1.  Move the **Time Bar** to **2.50** and press **Enter**, then select **Jump_Start** from the menu. 
+1.  **Move** the **Time Bar** to **2.50** and **press Enter**, then **select** **Jump_Start** from the menu. 
 
-1. **Right-click** on the **Jump_Start** keyframe at **2.50** and select **Set Play Rate** and set the value to **0.23**. 
+1. **Right-click** on the **Jump_Start** keyframe at **2.50** and **select** **Set Play Rate** and set the value to **0.23**. 
 
 	![](MHT7_JumpAnimAdded.png)
 
-1. Move the **Time Bar** to **4.50** and press **Enter**, then select **Jump_Start** again from the menu. 
+1. **Move** the Time Bar to **4.50** and **press Enter**, then **select** **Jump_Start** again from the menu. 
 
 	![](MHT7_AnimKey3.png)
 
 	_We could have added the **Jump_Loop** animation, however our character is going to continue ascending and using frames from the **Jump_Start** will look better._
 
-1. **Right-click** on the **Jump_Start** keyframe at **4.50** and select **Set Start Offset** and set this to **5**. 
+1. **Right-click** on the **Jump_Start** keyframe at **4.50** and **select** **Set Start Offset** and set this to **5**. 
 
 	![](MHT7_AnimKey4.png)
 
@@ -154,15 +154,15 @@ In the following steps, we will have a character jump up onto a platform by appl
 
 	![](MHT7_AnimKey5.png)(w:720)
 
-	_You should also see that the pose on the character has changed and is not using the normal **Jump_Start** pose of the character crouching._
+	_You should also see that the pose on the character has changed and is not using the normal Jump_Start pose of the character crouching._
 
-1. Move the **Time Bar** to **8.00** and press **Enter**, then select **Jump_End** from the menu. 
+1. **Move** the Time Bar to **8.00** and **press Enter**, then **select** **Jump_End** from the menu. 
 
-1. Scrub the **Time Bar** to **8.30** and press **Enter**, then select **Idle** from the menu. 
+1. Scrub the Time Bar to **8.30** and **press Enter**, then **select** **Idle** from the menu. 
 
 	![](MHT7_AnimKey6.png)
 
-1. **Right-click** on the **Idle** keyframe at **8.30** and choose **Set Time** and set it to **8.20**. 
+1. **Right-click** on the **Idle** keyframe at **8.30** and **choose** Set Time and **set** it to **8.20**. 
 
 1. **Right-click** on the **Idle** keyframe that was moved to **8.20** and choose **Set Looping**. 
 
@@ -180,41 +180,41 @@ In the following steps, we will have a character jump up onto a platform by appl
 	[/PARAM]
 [/OBJECT]
 
-1. Minimize the **Matinee** window, then in the **Content Browser** inside the **Game/Audio** folder, click on the **Explosion01** asset. 
+1. **Minimize** the Matinee window, then in the Content Browser inside the **Game/Audio** folder, **click** on the **Explosion01** asset. 
 
-1. Return to the **Matinee** window and move the **Time Bar** to the **2.50** mark and press **Enter**. 
+1. **Return** to the Matinee window and **move** the Time Bar to the **2.50** mark and **press Enter**. 
 
 	![](MHT7_Sound1.png)(w:780)
 
 	_In order to add a sound asset to the **Sound Track**, a **Sound Wave** or **Sound Cue** asset must be selected in the **Content Browser**._
 
-1. **Right-click** on the newly added sound asset's keyframe and select **Set Sound Volume** and set it to **1.5**. 
+1. **Right-click** on the newly added sound asset's keyframe and **select** **Set Sound Volume** and set it to **1.5**. 
 
-1. Repeat the previous step, however **Right-click** and select **Set Sound Pitch** and set it to **0.5**. 
+1. **Repeat the previous step**, however **Right-click** and **select** **Set Sound Pitch** and set it to **0.5**. 
 
 	![](MHT7_Sound2.png)
 
 	_You will see that the volume for this sound has been increased to 1.5 while the pitch has been decreased to 0.5 on the bar in the track._
 
-1. Click on the other **Sound Track**, then in the **Content Browser** under **Game/Audio**, select the **Start_Music_Cue** asset. 
+1. **Click** on the other Sound Track, then in the Content Browser under **Game/Audio**, **select** the **Start_Music_Cue** asset. 
 
-1. Return to the **Matinee** window and move the **Time Bar** back to **0.00** and press **Enter** to add the music. 
+1. **Return** to the Matinee window and **move** the Time Bar back to **0.00** and **press Enter** to add the music. 
 
 	![](MHT7_Sound3.png)(w:780)
 
-1. In the **Details** panel for the **Sound Track** with the music, check the **Continue Sound on Matinee End** checkbox. 
+1. In the Details panel for the Sound Track with the music, **check** the **Continue Sound on Matinee End** checkbox. 
 
 	![](MHT7_Sound4.png)
 
 	_This will allow the music to continue playing once the 12 second **Matinee** sequence is over (otherwise the music would abruptly end)._ 
 
-1. Close the **Matinee**, then select the **MatineeActor** in the **World Outliner**. 
+1. Close the Matinee, then **select** the MatineeActor in the World Outliner. 
 
-1. In the **Details** panel under **Play**, check the **Play on Level Load** checkbox. 
+1. In the Details panel under Play, **check** the **Play on Level Load** checkbox. 
 
 	![](MHT7_MatineeFinished.png)
 
-1. Click the **Play** button from the main toolbar to play in the editor. 
+1. **Click** the Play button from the main toolbar to play in the editor. 
 
 	![](MHT7_Gif.gif)(convert:false)
 
