TypoCorrection

diff --git a/Engine/Documentation/Source/Shared/Types/APostProcessVolume/APostProcessVolume.INT.udn b/Engine/Documentation/Source/Shared/Types/APostProcessVolume/APostProcessVolume.INT.udn
--- a/Engine/Documentation/Source/Shared/Types/APostProcessVolume/APostProcessVolume.INT.udn
+++ b/Engine/Documentation/Source/Shared/Types/APostProcessVolume/APostProcessVolume.INT.udn
@@ -1,4 +1,4 @@
-Availability:NoPublish
+﻿Availability:NoPublish
 Title:APostProcessVolume
 Crumbs:
 Description:Overview of the PostProcessVolume class
@@ -7,7 +7,7 @@ Description:Overview of the PostProcessVolume class
 
 The PostProcessVolume is a special type of volume that can be added to a level. Unreal Engine 3 had the post processing chains, but in Unreal Engine 4 all this got replaced by
 using the PostProcessSettings which can be conveniently adjusted using one or multiple PostProcessVolume. By using multiple volumes (or other objects) with priority and weight
-teh settings can be softly blended depending on the camera position and manipulating weights.
+the settings can be softly blended depending on the camera position and manipulating weights.
 
 [VAR:ToolTipFullLink]
 Engine/Rendering/PostProcessEffects
