TypoCorrection

diff --git a/Engine/Documentation/Source/Programming/Assertions/Assertions.INT.udn b/Engine/Documentation/Source/Programming/Assertions/Assertions.INT.udn
--- a/Engine/Documentation/Source/Programming/Assertions/Assertions.INT.udn
+++ b/Engine/Documentation/Source/Programming/Assertions/Assertions.INT.udn
@@ -1,4 +1,4 @@
-Availability:Public
+﻿Availability:Public
 Title: Assertions
 Crumbs:
 Description:
@@ -29,7 +29,7 @@ Examples:
 
 	verify((Mesh = GetRenderMesh()) != nullptr);
 
-**checkf(expression, …);**
+**checkf(expression, ...);**
 
 The checkf() macro allows you to assert an expression is true and print extra information to aid during debugging. It behaves the same as check() in terms of compilation behavior.
 
@@ -38,7 +38,7 @@ Examples:
 	checkf(WasDestroyed, TEXT( "Failed to destroy Actor %s (%s)"), *Actor->GetClass()->GetName(), *Actor->GetActorLabel());
 	checkf( TCString<ANSICHAR>::Strlen( Key ) >= KEYLENGTH( AES_KEYBITS ), TEXT( "AES_KEY needs to be at least %d characters" ), KEYLENGTH( AES_KEYBITS ) );
 
-**verifyf(expression, …);**
+**verifyf(expression, ...);**
 
 Just like the verify() macro always executes the expression, so does verifyf(). Similarly, it halts execution with an additional debug message, just like checkf();
 
@@ -89,7 +89,7 @@ Example:
 	int32 Recurse(int32 A, int32 B)
 	{
 		checkNoRecursion();
-		return Recurse(A – 1, B – 1);
+		return Recurse(A - 1, B - 1);
 	}
 
 **unimplemented();**
@@ -136,7 +136,7 @@ Example:
 
 	ensureMsg(Node != nullptr, TEXT("Node is invalid"));
 
-**ensureMsgf(expression, message, …);**
+**ensureMsgf(expression, message, ...);**
 
 Verifies the expression and includes even more information coupled with the callstack for the generated report. Just like checkf() or verifyf(), you can include contextual information to aid in tracking down the issue.
 
