TypoCorrection

diff --git a/Engine/Documentation/Source/Platforms/Android/Reference/AndroidGettingStarted.INT.udn b/Engine/Documentation/Source/Platforms/Android/Reference/AndroidGettingStarted.INT.udn
--- a/Engine/Documentation/Source/Platforms/Android/Reference/AndroidGettingStarted.INT.udn
+++ b/Engine/Documentation/Source/Platforms/Android/Reference/AndroidGettingStarted.INT.udn
@@ -1,4 +1,4 @@
-Availability:Public
+﻿Availability:Public
 Title: Android Development Reference
 Crumbs:%ROOT%, Platforms, Platforms/Android
 Description:How to install different Android SDKs, set environment variables, and work with texture formats.
@@ -91,7 +91,7 @@ The important pieces are:
 * NDK
 * JDK
 * Ant 
-* Night 
+* Nsight 
 	
 If you use the default installation set, a pretty minimal set of components will be installed. If you customize the install, you can add more items (like profilers, etc). Note that during installation, 
 one or two of the sub-installers may appear to hang, but they will eventually finish. The install could take as long as 30 to 45 minutes. If you never plan to write code, you may want to 
@@ -119,7 +119,7 @@ If you use the default installation set, a pretty minimal set of components will
 If you customize the install, you can add more items (like profilers, etc). Note that during installation, one or two of the sub-installers may appear to hang, but they will eventually finish. 
 The install could take as long as 30 to 45 minutes on some computers. 
 If you never plan to write code, you may want to select **Customize**, and only install the SDK, Ant, and Java (skipping the NDK and Nsight debugger) to make for the quickest installation.
-The Nsight portion of TADP is used for debugging an Android game on a device so t is not required to compile C++ code.
+The Nsight portion of TADP is used for debugging an Android game on a device so it is not required to compile C++ code.
 
 ![](TADP_Mac_OS.png)
 
