#UE4doc #docs:
Changed color of multiple borders/arrows to yellow from green
Unbolded multiple terms in accordance to the instructional guide.

diff --git a/Engine/Documentation/Source/Engine/AI/EnvironmentQuerySystem/QuickStart/11/11.INT.udn b/Engine/Documentation/Source/Engine/AI/EnvironmentQuerySystem/QuickStart/11/11.INT.udn
--- a/Engine/Documentation/Source/Engine/AI/EnvironmentQuerySystem/QuickStart/11/11.INT.udn
+++ b/Engine/Documentation/Source/Engine/AI/EnvironmentQuerySystem/QuickStart/11/11.INT.udn
@@ -21,15 +21,15 @@ Description: Create an Environment Query that finds all instances of the player'
 
 %Steps%
 
-This is the simpler of the two **Environment Queries**, as all it will do is return a list of all the TopDownCharacters within 1500 units of the AI then test if the AI can see the TopDownCharacters. If it can, then we'll get our **ActorToFleeFrom** for our **Behavior Tree**.
+This is the simpler of the two **Environment Queries**, as all it will do is return a list of all the TopDownCharacters within 1500 units of the AI then test if the AI can see the TopDownCharacters. If it can, then we'll get our ActorToFleeFrom for our Behavior Tree.
 
 There are a number of ways this process can be done, each with their own advantages and disadvantages. Using EQS for this means that we can add additional tests at a later time to choose the _best_ enemy (TopDownCharacter) to run from.
 
 Some quick glossary terms you're going to run into:
 
-* A **Generator** creates **Items** in relation to a **Context**
-* **Items** are used in **Tests** often in relation to a **Context** (how far is Item X from Context Y)
-* **Items** are culled or scored based on the results of **Tests**
+* A Generator creates Items in relation to a Context
+* Items are used in Tests often in relation to a Context (how far is Item X from Context Y)
+* Items are culled or scored based on the results of Tests
 
 ## Steps
 
@@ -51,11 +51,11 @@ Some quick glossary terms you're going to run into:
 
 	![](cga_made.png)
 
-1. To connect nodes in the Environment Query editor, you just need to drag from the bottom of the Root node to the Top of the **ActorsOfClass** node, exactly how it is done in the **Behavior Tree** editor.
+1. To connect nodes in the Environment Query editor, you just need to drag from the bottom of the Root node to the Top of the **ActorsOfClass** node, exactly how it is done in the Behavior Tree editor.
 
 	![](connected.png)
 
-	There can only be one Generator per **Environment Query**
+	There can only be one Generator per Environment Query
 
 1. Next, setup the **ActorsOfClass** node by selecting it, and changing its **Search Radius** and **Searched Actor Class** properties:
 
@@ -64,7 +64,7 @@ Some quick glossary terms you're going to run into:
 	* **Search Radius** = 1500
 	* **Searched Actor Class** = TopDownCharacter.
 
-	This **Generator** finds all the Actors of the given Class, within the radius, and creates **Items** from them.
+	This Generator finds all the Actors of the given Class, within the radius, and creates Items from them.
 
 1. Add a **Trace** test to the **ActorsOfClass** node
 
