#UE4doc #docs: Altered images to match the instructional guide: green -> yellow. Updated title to match How To - Instructional Guide.

diff --git a/Engine/Documentation/Source/Gameplay/HowTo/Networking/ReplicateActor/Blueprints/ReplicateActor_BP.INT.udn b/Engine/Documentation/Source/Gameplay/HowTo/Networking/ReplicateActor/Blueprints/ReplicateActor_BP.INT.udn
--- a/Engine/Documentation/Source/Gameplay/HowTo/Networking/ReplicateActor/Blueprints/ReplicateActor_BP.INT.udn
+++ b/Engine/Documentation/Source/Gameplay/HowTo/Networking/ReplicateActor/Blueprints/ReplicateActor_BP.INT.udn
@@ -1,5 +1,5 @@
 Availability: Public
-Title:How to Replicate Actors in Blueprints
+Title:Replicating Actors in Blueprints
 Crumbs: %ROOT%, Gameplay, Gameplay\HowTo\Networking/ReplicateActor
 Description:Guide to using the Replicates checkbox for Actors in Blueprints. 
 Related: Engine/Blueprints
@@ -7,8 +7,6 @@ Related: Resources/ContentExamples/Networking
 
 [TOC (start:2 end:2)]
 
-## Overview
-
 This page shows how to use the **Replicates** option on Actors in order to have them appear (or not appear) on the server versus a client machine. In the two examples cited, the **Replicating Actors** example shows how to set an Actor to replicate and the difference between a replicated and non-replicated Actor. The second example illustrates how you can dynamically spawn and control where an Actor appears by using the **Switch Has Authority** node (which checks to see if the server or client is running the script). 
 
 ## Replicating Actors
@@ -31,7 +29,7 @@ This page shows how to use the **Replicates** option on Actors in order to have
 
 	![](HowTo3.png)
 
-	With this option un-checked, the Actor will not be automatically loaded on the client when the map is loaded. Since **Replicates** is also not checked, the Actor will not appear on the client's machine at all even if it were spawned in. 
+	With this option un-checked, the Actor will not be automatically loaded on the client when the map is loaded. Since Replicates is also not checked, the Actor will not appear on the client's machine at all even if it were spawned in. 
 
 1.	**Compile** and **Save**, then close the Blueprint and open the **Fire_Replicate** Blueprint. 
 
@@ -39,7 +37,7 @@ This page shows how to use the **Replicates** option on Actors in order to have
 
 	![](HowTo4.png)	
 
-	By checking **Replicates**, we are saying that this Actor can be replicated on client machines. With **Net Load on Client** also checked, the Actor is automatically loaded when the map loads and is displayed on client's machines. Where **Replicates** comes into play is, say you wanted to spawn in this Actor during gameplay at some point instead of on map load, you would need **Replicates** enabled to say that when the Actor is spawned, go ahead and spawn a replicated copy on client machines. 
+	By checking **Replicates**, we are saying that this Actor can be replicated on client machines. With **Net Load on Client** also checked, the Actor is automatically loaded when the map loads and is displayed on client's machines. Where Replicates comes into play is, say you wanted to spawn in this Actor during gameplay at some point instead of on map load, you would need Replicates enabled to say that when the Actor is spawned, go ahead and spawn a replicated copy on client machines. 
 
 1.	**Compile** and **Save** then close the Blueprint. 
 
@@ -61,7 +59,7 @@ This page shows how to use the **Replicates** option on Actors in order to have
 
 	![](HowTo9.png)	
 
-	When you play in the editor and a new client window is opened, you should notice that the **Fire_NoReplicate** Blueprint is displayed on the server and not shown on the client. The **Fire_Replicate** Blueprint however is shown on both the client and the server as depicted above. 
+	When you play in the editor and a new client window is opened, you should notice that the Fire_NoReplicate Blueprint is displayed on the server and not shown on the client. The Fire_Replicate Blueprint however is shown on both the client and the server as depicted above. 
 
 
 ## Dynamically Spawning Replicated Actors
@@ -124,7 +122,7 @@ The steps below showcase how to dynamically spawn Actors during gameplay using *
 
 	![](Switch10.png)
 
-	You should see that since the Server executed the script, both fire effects are shown on the Server while only the one that is set to replicate is shown on the Client. This is useful when you want to spawn something on the Server or Client and not on both, or more importantly, execute any script following the **Switch Has Authority** node on either the Server or Client (as previously mentioned). 
+	You should see that since the Server executed the script, both fire effects are shown on the Server while only the one that is set to replicate is shown on the Client. This is useful when you want to spawn something on the Server or Client and not on both, or more importantly, execute any script following the Switch Has Authority node on either the Server or Client (as previously mentioned). 
 
 1.	Return to the **Level Blueprint** and drag off the **Remote** pin and connect it as shown below. 
 
@@ -134,7 +132,7 @@ The steps below showcase how to dynamically spawn Actors during gameplay using *
 
 	![](Switch12.png)
 
-	You should now see that since the **Switch Has Authority** has been set to **Remote**, the script following it is only executed on the client machines and not the server. Say for example you had some sort of screen effect that is applied when a player takes damage, using **Remote** as the Authority Guard, you can spawn that effect only on the remote client that is executing the script. 
+	You should now see that since the Switch Has Authority has been set to Remote, the script following it is only executed on the client machines and not the server. Say for example you had some sort of screen effect that is applied when a player takes damage, using Remote as the Authority Guard, you can spawn that effect only on the remote client that is executing the script. 
 
 	
 
