#UE4doc #docs: Removed extraneous bolding as according to the instructional guide. ##Overview removed as according to the instructional guide. Changed title of How To to match instructional guide.

diff --git a/Engine/Documentation/Source/Gameplay/HowTo/Networking/ReplicateFunction/ReplicateFunctionOverview.INT.udn b/Engine/Documentation/Source/Gameplay/HowTo/Networking/ReplicateFunction/ReplicateFunctionOverview.INT.udn
--- a/Engine/Documentation/Source/Gameplay/HowTo/Networking/ReplicateFunction/ReplicateFunctionOverview.INT.udn
+++ b/Engine/Documentation/Source/Gameplay/HowTo/Networking/ReplicateFunction/ReplicateFunctionOverview.INT.udn
@@ -1,11 +1,8 @@
 Availability: Public
-Title:How to Replicate Functions
+Title:Replicating Functions
 Crumbs: %ROOT%, Gameplay
 Description:Guide for Replicating Functions in Multiplayer Games. 
 
-
-##Overview
-
 **Function Replication** (Remote Procedure Calls or RPCs for short) are functions that are called locally, but executed remotely on another machine (separate from the machine performing the call). RPC functions can be very useful and allow either the Client or the Server to send messages to each other over a network connection. Replicated function calls can be set to either **Reliable** or **Unreliable** where Reliable calls are guaranteed to occur, while Unreliable calls may be dropped during heavy network traffic. Most Replicated Functions that handle cosmetic visuals are unreliable to avoid saturating the network. 
 
 There are 3 primary types of Replicated Functions: **Multicast**, **Run on Server**, and **Run on owning Client**. Multicast functions should be called on the Server, where they are executed, and then forwarded automatically to Clients. Server functions are called by a Client and then only executed on the Server. Client functions are called by the Server and then only execute on the owning Client.
@@ -14,7 +11,7 @@ Server and Client Replicated Functions have some restrictions, for one they can
 
 ## Setup
 
-To learn how to set up **Replicated Functions** using Blueprints, follow the link below:
+To learn how to set up Replicated Functions using Blueprints, follow the link below:
 
 [REGION:buttonlist]
 * [](Gameplay\HowTo\Networking\ReplicateFunction\Blueprints)
