#UE4doc #docs:
Changed multiple green borders to yellow.
Unbolded and bolded many different terms in accordance to our instructional guide.
Fixed some grammar.

diff --git a/Engine/Documentation/Source/Engine/AI/EnvironmentQuerySystem/QuickStart/12/12.INT.udn b/Engine/Documentation/Source/Engine/AI/EnvironmentQuerySystem/QuickStart/12/12.INT.udn
--- a/Engine/Documentation/Source/Engine/AI/EnvironmentQuerySystem/QuickStart/12/12.INT.udn
+++ b/Engine/Documentation/Source/Engine/AI/EnvironmentQuerySystem/QuickStart/12/12.INT.udn
@@ -25,13 +25,13 @@ For this step you'll be creating a much more involved Environment Query. It's go
 
 Just a quick reminder:
 
-* A **Generator** creates **Items** in relation to a **Context**
-* **Items** are used in **Tests** often in relation to a **Context** (how far is Item X from Context Y)
-* **Items** are culled or scored based on the results of **Tests**
+* A Generator creates Items in relation to a Context
+* Items are used in Tests often in relation to a Context (how far is Item X from Context Y)
+* Items are culled or scored based on the results of Tests
 
 ## Steps
 
-1. To start, double click the `FindHidingSpot` **Environment Query** asset in the **Content Browser**.
+1. To start, **Double-Click** the `FindHidingSpot` Environment Query asset in the **Content Browser**.
 
 	![](findHide.png)
 
@@ -39,7 +39,7 @@ Just a quick reminder:
 
 	![](simpleGrid.png)
 
-1. Select the Simple Grid Generator and in the **Details** Panel set the following Properties:
+1. Select the **Simple Grid** Generator and in the **Details** Panel set the following Properties:
 	
 	 ![](gridProp.png)
 
@@ -47,7 +47,7 @@ Just a quick reminder:
 
 	These settings for the Grid Generator will create a sizable, but manageable, number of Items for the Query to test against. The **Project Down**, **Project Up,** and **Post Projection Vertical Offset** properties are particularly useful to change if you have an AI that needs to find something above or below them.
 
-1. [REGION:raw]Add 2 **Distance** Tests by Right-Clicking the **Simple Grid** node and selecting **Tests** %arrow% **Distance**[/REGION]
+1. [REGION:raw]Add 2 **Distance** Tests by **Right-Clicking** the **Simple Grid** node and selecting **Tests** %arrow% **Distance**[/REGION]
 
 	![](addDistance.png)
 
@@ -57,7 +57,7 @@ Just a quick reminder:
 
 	[INCLUDE:#eqsqs122]
 
-	This Distance test is meant to give **Items** _further_ from the Player a higher Score.
+	This Distance test is meant to give Items _further_ from the Player a higher Score.
 
 1. Select the second one and in the **Details** Panel adjust the following Properties:
 
@@ -65,9 +65,9 @@ Just a quick reminder:
 
 	[INCLUDE:#eqsqs123]
 
-	This Distance test is meant to give **Items** _closer_ to the AI a higher Score, and this test carries more weight than the previous **Distance Test** meaning it will favor closer to the AI but still select away from the Player.
+	This Distance test is meant to give Items _closer_ to the AI a higher Score, and this test carries more weight than the previous **Distance Test** meaning it will favor closer to the AI but still select away from the Player.
 
-1. [REGION:raw]Add 2 **Trace** Tests by Right-Clicking the **Simple Grid** node and selecting **Tests** %arrow% **Trace**[/REGION]
+1. [REGION:raw]Add 2 **Trace** Tests by **Right-Clicking** the **Simple Grid** node and selecting **Tests** %arrow% **Trace**[/REGION]
 
 	![](addTrace.png)
 
@@ -77,11 +77,11 @@ Just a quick reminder:
 
 	Here we remove any Item that the Player Pawn or AI pawn has line of sight to. Also, you may be concerned that this test comes after a bunch of **Distance** tests, but it's ok, as the Environment Query runs **Trace** tests first regardless of order here. Also, the other Trace requires no additional setup.
 
-1. [REGION:raw]Add a **Dot** test by Right-Clicking the **Simple Grid** node and selecting **Tests** %arrow% **Dot**[/REGION]
+1. [REGION:raw]Add a **Dot** test by **Right-Clicking** the **Simple Grid** node and selecting **Tests** %arrow% **Dot**[/REGION]
 
 	![](addDot.png)
 
-1. Select the Dot test and in the **Details** panel adjust the following properties: 
+1. Select the **Dot** test and in the **Details** panel adjust the following properties: 
 	
 	![](dotProps.png)
 
@@ -89,15 +89,15 @@ Just a quick reminder:
 
 	Here things get tricky, at the high level, this is saying "Favor Items that are to the sides of the Player Pawn," but how it works is a bit odd. A Dot Product function will return how "alike" two vectors are on a scale of 1.0 (perfectly the same) to -1.0 (completely opposite). So if we get the **Absolute Value** of the Dot Product between the rotation vector of the Player and the vector from the Player to the Item we'll favor points in front of and behind the Player. But if we set the **Scoring Factor** to -1.0, then we favor to the sides of the Player.
 
-1. [REGION:raw]Add a **PathFinding** test by Right-Clicking the **Simple Grid** node and selecting **Tests** %arrow% **PathFinding**[/REGION]
+1. [REGION:raw]Add a **PathFinding** test by **Right-Clicking** the **Simple Grid** node and selecting **Tests** %arrow% **PathFinding**[/REGION]
 
 	![](addPath.png)
 
-1. Select the PathFinding test and in the **Details** panel adjust **Test Purpose** to `Score Only`:
+1. Select the **PathFinding** test and in the **Details** panel adjust **Test Purpose** to `Score Only`:
 
 	![](pathProps.png)
 
-	And that's it for the editing! The final Environment Query will look similar to this:
+	That's it for the editing! The final Environment Query will look similar to this:
 
 	![](final.png)
 
@@ -115,11 +115,11 @@ Just a quick reminder:
 [/EXCERPT]
 
 [EXCERPT:eqsqs121]
-* Grid Size = `2000.0`
-* Space Between = `150.0`
-* Navigation Filer = `RecastFilter_UseDefaultArea`
-* Project Down = `2048.0`
-* Post Projection Vertical Offset = `120.0`
+* **Grid Size** = `2000.0`
+* **Space Between** = `150.0`
+* **Navigation Filer** = `RecastFilter_UseDefaultArea`
+* **Project Down** = `2048.0`
+* **Post Projection Vertical Offset** = `120.0`
 [/EXCERPT]
 
 [EXCERPT:eqsqs122]
