#UE4doc #docs: Removed extraneous bolding as according to the instructional guide. Changed title of How To to match instructional guide.

diff --git a/Engine/Documentation/Source/Gameplay/HowTo/Networking/ReplicateFunction/Blueprints/ReplicateFunction_BP.INT.udn b/Engine/Documentation/Source/Gameplay/HowTo/Networking/ReplicateFunction/Blueprints/ReplicateFunction_BP.INT.udn
--- a/Engine/Documentation/Source/Gameplay/HowTo/Networking/ReplicateFunction/Blueprints/ReplicateFunction_BP.INT.udn
+++ b/Engine/Documentation/Source/Gameplay/HowTo/Networking/ReplicateFunction/Blueprints/ReplicateFunction_BP.INT.udn
@@ -1,5 +1,5 @@
 Availability: Public
-Title:How to Replicate Functions in Blueprints
+Title:Replicating Functions in Blueprints
 Crumbs: %ROOT%, Gameplay, Gameplay\HowTo\Networking\ReplicateFunction
 Description:Guide to Replicating Functions in multiplayer games with Blueprints. 
 Related: Engine/Blueprints
@@ -8,8 +8,6 @@ Related: Resources/ContentExamples/Networking
 
 [TOC (start:2 end:2)]
 
-## Overview
-
 This page will show you how to set up and use different forms of **Replicated Functions**. 
 
 There are 3 primary types of replicated functions: **Multicast**, **Run on Server**, and **Run on owning Client**. Multicast functions should be called on the server, where they are executed, and then forwarded automatically to clients. Server functions are called by a client and then only executed on the server. Client functions are called by the server and then only execute on the owning client.
@@ -78,7 +76,7 @@ There are 3 primary types of replicated functions: **Multicast**, **Run on Serve
 
 	![](HowTo7.png)
 
-	You should now see that when you press jump with the player that is the server, that the fire effect is spawned not only on the server but on all clients as well. Attempting to jump with other players still only spawns the effect locally, as we have not told the server that the client has spawned the effect. The next section, **Run on Server** will show you how to set that up. 
+	You should now see that when you press jump with the player that is the server, that the fire effect is spawned not only on the server but on all clients as well. Attempting to jump with other players still only spawns the effect locally, as we have not told the server that the client has spawned the effect. The next section, Run on Server will show you how to set that up. 
 
 
 ## Run On Server
@@ -170,9 +168,9 @@ What we will do in this example is create a variable that only gets updated on a
 
 	![](HowTo17.png)
 
-	The **Switch Has Authority** node checks where the script that is currently running is executing from, then can branch into two different directions based on whether the script is running on the network **Authority** (usually the Server) or a **Remote** machine (a client). 
+	The Switch Has Authority node checks where the script that is currently running is executing from, then can branch into two different directions based on whether the script is running on the network Authority (usually the Server) or a Remote machine (a client). 
 	
-	More often than not, you will want to use **Authority** for things you only want to occur on the server (these are usually gameplay critical, such as adjusting health values for players or giving out awards or loot items as you do not want to allow clients to determine when those types of changes are to occur to prevent cheating). 
+	More often than not, you will want to use Authority for things you only want to occur on the server (these are usually gameplay critical, such as adjusting health values for players or giving out awards or loot items as you do not want to allow clients to determine when those types of changes are to occur to prevent cheating). 
 
 	In this example, we are going to update a text variable, that could just as easily be a variable that contains the player's health or a variable for an item collected. 
 
@@ -188,7 +186,7 @@ What we will do in this example is create a variable that only gets updated on a
 
 1.	Create another **Custom Event** called **Remove Item** with the same settings as the **Add Item** event. 
 
-1.	Connect the **Add Item** and **Remove Item** nodes as shown below off the two overlap events. 
+1.	Connect the **Add Item** and Remove Item nodes as shown below off the two overlap events. 
 
 	![](HowTo19.png)
 
@@ -218,7 +216,7 @@ What we will do in this example is create a variable that only gets updated on a
 
 	_Click to play video example._
 
-	In the video above when the game is launched, pressing **P** to print the text variable to the screen displays it as "empty" for each of the characters. However when one of the characters enters the trigger box, the text "item added" is displayed. When that character presses **P**, the text changes to say "has the item" while each of the other characters still displays "empty". When the character leaves the box, the text "item removed" is displayed and when they press **P** again, the text now changes back to "empty". 
+	In the video above when the game is launched, pressing P to print the text variable to the screen displays it as "empty" for each of the characters. However when one of the characters enters the trigger box, the text "item added" is displayed. When that character presses P, the text changes to say "has the item" while each of the other characters still displays "empty". When the character leaves the box, the text "item removed" is displayed and when they press P again, the text now changes back to "empty". 
 
 	This example illustrates how you can run events on the server and replicate those events to individual clients. 
 
