#UE4doc #docs: Removed extraneous bolding as according to the instructional guide. ##Overview removed as according to the instructional guide. Changed title of How To to match instructional guide.

diff --git a/Engine/Documentation/Source/Gameplay/HowTo/Networking/ReplicateVariable/ReplicateVariableOverview.INT.udn b/Engine/Documentation/Source/Gameplay/HowTo/Networking/ReplicateVariable/ReplicateVariableOverview.INT.udn
--- a/Engine/Documentation/Source/Gameplay/HowTo/Networking/ReplicateVariable/ReplicateVariableOverview.INT.udn
+++ b/Engine/Documentation/Source/Gameplay/HowTo/Networking/ReplicateVariable/ReplicateVariableOverview.INT.udn
@@ -1,19 +1,17 @@
 Availability: Public
-Title:How to Replicate Variables
+Title:Replicating Variables
 Crumbs: %ROOT%, Gameplay
 Description:Guide for Replicating Variables in Multiplayer Games. 
 
-
-##Overview
 [EXCERPT:Overview]
-Just as [Actors](Gameplay\HowTo\Networking\ReplicateActor) and [Functions](Gameplay\HowTo\Networking\ReplicateFunction) can be replicated over the network, so can **Variables**. Variables that are important to gameplay should only be modified on the Network Authority (Server) and then replicated to Remote Machines (Clients) on a need-to-know basis. You do not want to allow Clients to define when variables can be changed, as this could lead to cheating (setting Health, Ammo, or Experience values for example). 
+Just as [Actors](Gameplay\HowTo\Networking\ReplicateActor) and [Functions](Gameplay\HowTo\Networking\ReplicateFunction) can be replicated over the network, so can Variables. Variables that are important to gameplay should only be modified on the Network Authority (Server) and then replicated to Remote Machines (Clients) on a need-to-know basis. You do not want to allow Clients to define when variables can be changed, as this could lead to cheating (setting Health, Ammo, or Experience values for example). 
 
-You can set a variable to **Replicated** which means the variable produces a one-to-one copy of the variable that is replicated from the Server to Clients. Or you can use the **RepNotify** which does everything that **Replicated** does, but also provides a Function that is called that executes on both the Server and Client machines whenever the variable it is attached to updates. 
+You can set a variable to Replicated which means the variable produces a one-to-one copy of the variable that is replicated from the Server to Clients. Or you can use the RepNotify which does everything that Replicated does, but also provides a Function that is called that executes on both the Server and Client machines whenever the variable it is attached to updates. 
 [/EXCERPT:Overview]
 
 ## Setup
 
-To learn how to set up **Replicated Variables** using Blueprints, follow the link below:
+To learn how to set up Replicated Variables using Blueprints, follow the link below:
 
 [REGION:buttonlist]
 * [](Gameplay\HowTo\Networking\ReplicateVariable\Blueprints)
