#UE4doc #docs: Adjusted bold text in accordance with the Instructional Guide

diff --git a/Engine/Documentation/Source/Engine/UI/LevelEditor/Viewports/ShowFlags/ShowFlags.INT.udn b/Engine/Documentation/Source/Engine/UI/LevelEditor/Viewports/ShowFlags/ShowFlags.INT.udn
--- a/Engine/Documentation/Source/Engine/UI/LevelEditor/Viewports/ShowFlags/ShowFlags.INT.udn
+++ b/Engine/Documentation/Source/Engine/UI/LevelEditor/Viewports/ShowFlags/ShowFlags.INT.udn
@@ -135,7 +135,7 @@ The Show Flags are found under the Viewport **Show** menu. They are a means by w
 |---|---|
 |[REGION:tablesection]Name[/REGION]|[REGION:tablesection]Description[/REGION]|
 |**AI Debug**|This shows the paths an AI-driven pawn is using to navigate along a Nav Mesh.|
-|**Blur GBuffer**|Shows/hides the results of **Blue GBuffer Intensity**, a property found on Post Process Volumes. This feature blurs the GBuffer normals and adjusts the GBuffer roughness to reduce specular aliasing. It is, however, a high performance cost for the result and you should generally consider just using Temporal AA.  |
+|**Blur GBuffer**|Shows/hides the results of Blue GBuffer Intensity, a property found on Post Process Volumes. This feature blurs the GBuffer normals and adjusts the GBuffer roughness to reduce specular aliasing. It is, however, a high performance cost for the result and you should generally consider just using Temporal AA.  |
 |**GBuffer Hints (material attributes)**|Colorizes objects that have improper Materials. Yellow for objects that are impossibly non-reflective, red for objects that emit more light than they receive.|
 |**Hit Proxies**|Shows/hides the proxy shapes used by hit detection.|
 |**Log Visualizer**|Draws any data to the screen that is required by the Log Visualizer system.|
