#UE4doc #docs:
Removed bolding from a term

diff --git a/Engine/Documentation/Source/Engine/Paper2D/Flipbooks/Components/FlipbookComponents.INT.udn b/Engine/Documentation/Source/Engine/Paper2D/Flipbooks/Components/FlipbookComponents.INT.udn
--- a/Engine/Documentation/Source/Engine/Paper2D/Flipbooks/Components/FlipbookComponents.INT.udn
+++ b/Engine/Documentation/Source/Engine/Paper2D/Flipbooks/Components/FlipbookComponents.INT.udn
@@ -15,7 +15,7 @@ While C++ documentation for working with Flipbook Components is still in develop
 
 ## Setup
 
-To learn more about **Flipbook Components** using Blueprints, follow the link below:
+To learn more about Flipbook Components using Blueprints, follow the link below:
 
 [REGION:buttonlist]
 * [](Engine\Paper2D\Flipbooks\Components\Blueprints)
