#UE4doc #docs: Removed extraneous bolding as according to the instructional guide. ##Overview removed as according to the instructional guide. Changed title of How To to match instructional guide. Altered images to match the instructional guide: green -> yellow.

diff --git a/Engine/Documentation/Source/Gameplay/HowTo/Networking/ReplicateVariable/Blueprints/ReplicateVariable_BP.INT.udn b/Engine/Documentation/Source/Gameplay/HowTo/Networking/ReplicateVariable/Blueprints/ReplicateVariable_BP.INT.udn
--- a/Engine/Documentation/Source/Gameplay/HowTo/Networking/ReplicateVariable/Blueprints/ReplicateVariable_BP.INT.udn
+++ b/Engine/Documentation/Source/Gameplay/HowTo/Networking/ReplicateVariable/Blueprints/ReplicateVariable_BP.INT.udn
@@ -1,5 +1,5 @@
 Availability: Public
-Title:How to Replicate Variables in Blueprints
+Title:Replicating Variables in Blueprints
 Crumbs: %ROOT%, Gameplay, Gameplay\HowTo\Networking\ReplicateVariable
 Description:Guide to using Replicated and RepNotify Variables in Blueprints.
 Related: Engine/Blueprints
@@ -8,8 +8,6 @@ Related: Resources/ContentExamples/Networking
 
 [TOC (start:2 end:2)]
 
-## Overview
-
 [INCLUDE:Gameplay/HowTo/Networking/ReplicateVariable#overview]
 
 This page covers some sample use cases for both **Replicated** variables and a **RepNotify** variable. 
@@ -18,7 +16,7 @@ This page covers some sample use cases for both **Replicated** variables and a *
 
  _For this example, we are using the **Blueprint Third Person Template with Starter Content**._
 
-In this example, we will create a **Replicated** variable that is used for a countdown timer that both the Server and Clients can see. 
+In this example, we will create a Replicated variable that is used for a countdown timer that both the Server and Clients can see. 
 
 1.	Inside the **Content/StarterContent/Blueprints** folder, open the **Blueprint_Effect_Fire** Blueprint. 
 
@@ -44,7 +42,7 @@ In this example, we will create a **Replicated** variable that is used for a cou
 
 	![](HT_3.png)
 
-	Here we used the **Switch Has Authority** node to call the **Set Timer** function only on the Authority (Server) when this Blueprint is spawned. We are using this node to ensure that the Timer function is executed by the server only and not by a client. If this script were to be handling game critical data, we would only want the server to execute it then replicate any data to client machines rather than the client run it as to prevent any cheating. 
+	Here we used the Switch Has Authority node to call the Set Timer function only on the Authority (Server) when this Blueprint is spawned. We are using this node to ensure that the Timer function is executed by the server only and not by a client. If this script were to be handling game critical data, we would only want the server to execute it then replicate any data to client machines rather than the client run it as to prevent any cheating. 
 
 1.	In the **MyBlueprint** window, create a new variable called **CountDownTimer**. 
 
@@ -120,9 +118,9 @@ In this example, we will create a **Replicated** variable that is used for a cou
 
 ## RepNotify
 
- _For this example, we are continuing from the **Replicated** example provided above._
+ _For this example, we are continuing from the Replicated example provided above._
 
-In the **Replicated** example provided above, we created a countdown timer that is displayed on both the Server and Client. In this example, we will use a **RepNotify** variable to signal when the countdown timer reaches 0 at which point we will deactivate the fire effect on both the server and the client machines. 
+In the Replicated example provided above, we created a countdown timer that is displayed on both the Server and Client. In this example, we will use a RepNotify variable to signal when the countdown timer reaches 0 at which point we will deactivate the fire effect on both the server and the client machines. 
 
 1.	Open the **Blueprint_Effect_Fire** Blueprint. 
 
@@ -150,7 +148,7 @@ In the **Replicated** example provided above, we created a countdown timer that
 
 	![](HT_18.png)
 
-	This is called whenever the variable we assigned as **RepNotify** changes and will be executed on both the server and client machines. 
+	This is called whenever the variable we assigned as RepNotify changes and will be executed on both the server and client machines. 
 
 1.	**Double-click** on the **On Rep Deactivate** Function to open it up. 
 
