#UE4doc #docs: Removed ôOverviewö header

diff --git a/Engine/Documentation/Source/Engine/UI/LevelEditor/Viewports/ShowFlags/ShowFlags.INT.udn b/Engine/Documentation/Source/Engine/UI/LevelEditor/Viewports/ShowFlags/ShowFlags.INT.udn
--- a/Engine/Documentation/Source/Engine/UI/LevelEditor/Viewports/ShowFlags/ShowFlags.INT.udn
+++ b/Engine/Documentation/Source/Engine/UI/LevelEditor/Viewports/ShowFlags/ShowFlags.INT.udn
@@ -51,7 +51,7 @@ Description:Descriptions of all of the available Show Flags for the viewports.
 [TOC(start:2 end:2)]
 
 
-## Overview
+
 
 The Show Flags are found under the Viewport **Show** menu. They are a means by which you can show and hide many of the items that are viewed within the viewport. For example, you can hide all particles, all landscape objects, or something a bit more advanced like turning off all of the sprite icons.
 
