#UE4doc #docs: ##Overview removed as according to the instructional guide. Changed title of How To to match instructional guide. Altered images to match the instructional guide: green -> yellow. Fixed a spelling error: seperate -> separate.

diff --git a/Engine/Documentation/Source/Gameplay/HowTo/Networking/TestMultiplayer/TestMultiplayer.INT.udn b/Engine/Documentation/Source/Gameplay/HowTo/Networking/TestMultiplayer/TestMultiplayer.INT.udn
--- a/Engine/Documentation/Source/Gameplay/HowTo/Networking/TestMultiplayer/TestMultiplayer.INT.udn
+++ b/Engine/Documentation/Source/Gameplay/HowTo/Networking/TestMultiplayer/TestMultiplayer.INT.udn
@@ -1,13 +1,11 @@
 Availability: Public
-Title:How to Test Multiplayer
+Title:Testing Multiplayer
 Crumbs: %ROOT%, Gameplay
 Description:Demonstrates how you can set up the Unreal Editor for testing multiplayer games. 
 Related: Resources/ContentExamples/Networking
 
 [TOC (start:2 end:2)]
 
-## Overview
-
 This page shows you how you can change some of the settings for the editor for testing different multiplayer scenarios. 
 
 ## Setting Number of Players
@@ -37,7 +35,7 @@ While it is perfectly okay to use the Editor's Viewport as a play window, for cl
 	![](NewEditorWindow.png)
 	[/REGION]
 
-While simulated clients have their own windows, this setting will also create a seperate window for the simulated server. 
+While simulated clients have their own windows, this setting will also create a separate window for the simulated server. 
 
 ## Advanced Settings
 
