#UE4doc #docs:
Bolded a few terms
Replaced an arrow due to color and style
Two gifs on page that need to be replaced.

diff --git a/Engine/Documentation/Source/Engine/Paper2D/Flipbooks/Components/Blueprints/Components_BP.INT.udn b/Engine/Documentation/Source/Engine/Paper2D/Flipbooks/Components/Blueprints/Components_BP.INT.udn
--- a/Engine/Documentation/Source/Engine/Paper2D/Flipbooks/Components/Blueprints/Components_BP.INT.udn
+++ b/Engine/Documentation/Source/Engine/Paper2D/Flipbooks/Components/Blueprints/Components_BP.INT.udn
@@ -13,8 +13,6 @@ Related: Engine/Blueprints
 
 [TOC (start:2 end:2)]
 
-## Overview
-
 Working with **Flipbook Components** in **Blueprints** is very similar to working with other types of [Components](Engine/Blueprints/UserGuide/Components) in Blueprints. After adding a Flipbook Component to a Blueprint, you can specify the **Source Flipbook** asset to use, define the Flipbook's **Play Rate** or adjust the **Sprite Color** used in addition to standard options that are available to all components such as **Transforming**. 
 
 Inside the [Graph](Engine/Blueprints/UserGuide/Graphs) tab, you can directly access and modify the properties of the Flipbook Component through your script. This includes things like adjusting the Play Rate, Playing, Stopping or Reversing the Flipbook, or even changing the Source Flipbook to a different Flipbook asset (useful for when you want to switch between Flipbook animations for a character for example). 
