#UE4doc #docs: Removed extraneous bolding as according to the instructional guide. ##Overview removed as according to the instructional guide. Changed title of How To to match instructional guide.

diff --git a/Engine/Documentation/Source/Gameplay/HowTo/PossessPawns/GHT_4_Overview.INT.udn b/Engine/Documentation/Source/Gameplay/HowTo/PossessPawns/GHT_4_Overview.INT.udn
--- a/Engine/Documentation/Source/Gameplay/HowTo/PossessPawns/GHT_4_Overview.INT.udn
+++ b/Engine/Documentation/Source/Gameplay/HowTo/PossessPawns/GHT_4_Overview.INT.udn
@@ -1,20 +1,17 @@
 Availability: Public
-Title:How to Possess Pawns
+Title:Possessing Pawns
 Crumbs: %ROOT%, Gameplay
 Description:A How To Guide for possessing different pawns in Unreal Engine 4. 
 
+If you want to switch between multiple Playable Characters in your game or you have a situation where you want a Character to enter and control a Vehicle (car, plane, tank, etc.) or other Pawn that has unique controls (like a mounted machine gun a player can control), you will want to take Possession of the Character or Pawn in order to pass input information to it. 
 
-##Overview
+The Player Controller is what is responsible for directing a Pawn and Possession of a Pawn requires a Player Controller to be specified. Pawns themselves do not need to be a humanoid character and can be anything that you want to apply basic movement to and allow a player to control. Characters on the other hand are a form of Pawn that includes Collision and a CharacterMovementComponent by default which allows it to do basic human-like movement. 
 
-If you want to switch between multiple Playable Characters in your game or you have a situation where you want a Character to enter and control a Vehicle (car, plane, tank, etc.) or other Pawn that has unique controls (like a mounted machine gun a player can control), you will want to take **Possession** of the **Character** or **Pawn** in order to pass input information to it. 
-
-The **Player Controller** is what is responsible for directing a Pawn and Possession of a Pawn requires a Player Controller to be specified. Pawns themselves do not need to be a humanoid character and can be anything that you want to apply basic movement to and allow a player to control. Characters on the other hand are a form of Pawn that includes Collision and a CharacterMovementComponent by default which allows it to do basic human-like movement. 
-
-The link below in the **Setup** section will give you examples on **Pawn Possession** using Blueprints.
+The link below in the Setup section will give you examples on Pawn Possession using Blueprints.
 
 ## Setup
 
-To learn more about **Possessing Pawns** using Blueprints, follow the link below:
+To learn more about Possessing Pawns using Blueprints, follow the link below:
 
 [REGION:buttonlist]
 * [](Gameplay\HowTo\PossessPawns\Blueprints)
