#UE4doc #docs: ##Overview removed as according to the instructional guide. Altered images to match the instructional guide: green -> yellow. Unbolded bolded hyperlink according to instructional guide. Fixed spelling error: particluar -> particular.

diff --git a/Engine/Documentation/Source/Gameplay/HowTo/PossessPawns/Blueprints/GHT_4B.INT.udn b/Engine/Documentation/Source/Gameplay/HowTo/PossessPawns/Blueprints/GHT_4B.INT.udn
--- a/Engine/Documentation/Source/Gameplay/HowTo/PossessPawns/Blueprints/GHT_4B.INT.udn
+++ b/Engine/Documentation/Source/Gameplay/HowTo/PossessPawns/Blueprints/GHT_4B.INT.udn
@@ -22,15 +22,11 @@ For the purposes of this guide we are using the **Blueprint Side Scroller** temp
 
 	![](Pawns1b.png)(w:860)
 
-	Here we added two additional **SideScrollerCharacter** Blueprints from the **Content > SideScrollerBP > Blueprints** folder to the level giving us a total of three characters in our level. With this particluar template project, a character is already placed inside the level by default. 
+	Here we added two additional **SideScrollerCharacter** Blueprints from the **Content > SideScrollerBP > Blueprints** folder to the level giving us a total of three characters in our level. With this particular template project, a character is already placed inside the level by default. 
 	
 1.	In the **World Outliner**, select each of the characters you want to control. 
 
 	![](Pawns2b.png)
-
-1. In the **World Outliner**, select each of the characters you want to control. 
-
-	![](Pawns2b.png)
 	
 1. From the Main Toolbar, click the **Blueprints** button and **Open Level Blueprint**. 
 
@@ -54,7 +50,7 @@ For the purposes of this guide we are using the **Blueprint Side Scroller** temp
 
 	![](Pawns6b.png)(w:640)
 
-	The [**Player Controller**](Gameplay/Framework/Controller/PlayerController) is used to take the input from a human player and translate that into actions for a Pawn. 
+	The [Player Controller](Gameplay/Framework/Controller/PlayerController) is used to take the input from a human player and translate that into actions for a Pawn. 
 
 1. **Left-click** and drag off the **Return Value** of the Get Player Controller node and search for and add the **Possess** node. 
 
