#UE4doc #docs: Removed extraneous bolding as according to the instructional guide. ##Overview removed as according to the instructional guide. Changed title of How To to match instructional guide.

diff --git a/Engine/Documentation/Source/Gameplay/HowTo/ReferenceAssets/GHT_5_Overview.INT.udn b/Engine/Documentation/Source/Gameplay/HowTo/ReferenceAssets/GHT_5_Overview.INT.udn
--- a/Engine/Documentation/Source/Gameplay/HowTo/ReferenceAssets/GHT_5_Overview.INT.udn
+++ b/Engine/Documentation/Source/Gameplay/HowTo/ReferenceAssets/GHT_5_Overview.INT.udn
@@ -1,10 +1,8 @@
 Availability: Public
-Title:How to Reference Actors
+Title:Referencing Actors
 Crumbs: %ROOT%, Gameplay
 Description:A How To Guide for Referencing Actors in Unreal Engine 4. 
 
-
-##Overview
 [EXCERPT:ReferenceOverview]
 By establishing reference to an Actor, you can access its Properties, Variables, Events, or Functions (if applicable) during gameplay through Blueprints or C++ and alter them to suit your gameplay needs. 
 
@@ -13,11 +11,11 @@ Say for example you want a light to turn on when a player enters a Trigger Box,
 Getting proper reference to an Actor is also important for communication between Actors, as incorrectly specifying a reference will result in failed communication and undesirable results. 
 [/EXCERPT:ReferenceOverview]
 
-The link below in the **Setup** section will give you examples on **How to Reference Actors** using Blueprints.
+The link below in the Setup section will give you examples on How to Reference Actors using Blueprints.
 
 ## Setup
 
-To learn more about **Referencing Actors** using Blueprints, follow the link below:
+To learn more about Referencing Actors using Blueprints, follow the link below:
 
 [REGION:buttonlist]
 * [](Gameplay\HowTo\ReferenceAssets\Blueprints)
