#UE4doc #docs: Removed extraneous bolding as according to the instructional guide. ##Overview removed as according to the instructional guide. Changed title of How To to match instructional guide. Altered images to match the instructional guide: green -> yellow.

diff --git a/Engine/Documentation/Source/Gameplay/HowTo/ReferenceAssets/Blueprints/GHT_5B.INT.udn b/Engine/Documentation/Source/Gameplay/HowTo/ReferenceAssets/Blueprints/GHT_5B.INT.udn
--- a/Engine/Documentation/Source/Gameplay/HowTo/ReferenceAssets/Blueprints/GHT_5B.INT.udn
+++ b/Engine/Documentation/Source/Gameplay/HowTo/ReferenceAssets/Blueprints/GHT_5B.INT.udn
@@ -1,19 +1,17 @@
 Availability: Public
-Title:How to Reference Actors in Blueprints
+Title:Referencing Actors in Blueprints
 Crumbs: %ROOT%, Gameplay, Gameplay/HowTo/ReferenceAssets
 Description:A How To Guide for referencing Actors in Blueprints. 
 
 [TOC (start:2)]
 
-## Overview
-
 On this page, you will learn multiple ways you can **Reference Actors** in Blueprints. Referencing Actors properly is important if you want to change aspects of the Actor during gameplay through Blueprints or pass the reference from one Blueprint to another so that you may access and modify it. 
  
 ## Level Blueprint Reference
 
 The most common method of creating a reference is by referencing an Actor that is placed in your level inside the **Level Blueprint**. To do this, simply follow the steps below. 
 
-1.	Inside your project, place an Actor or **Blueprint** Actor in your level.
+1.	Inside your project, place an Actor or Blueprint Actor in your level.
 
 	![](GHT5B_1.png)
 
@@ -37,7 +35,7 @@ The most common method of creating a reference is by referencing an Actor that i
 
 	![](GHT5B_4.png)
 
-	Here we are printing to the screen when the player overlaps the trigger. When they exit the trigger, we are disabling its collision so it cannot be triggered again (unless we re-enable collision on it). While we used a TriggerBox, you could have a **Blueprint** Actor in your level which has its own Variables, Events, or Functions within it that you could access or call by using this method of getting reference. 
+	Here we are printing to the screen when the player overlaps the trigger. When they exit the trigger, we are disabling its collision so it cannot be triggered again (unless we re-enable collision on it). While we used a TriggerBox, you could have a Blueprint Actor in your level which has its own Variables, Events, or Functions within it that you could access or call by using this method of getting reference. 
 
 
 
@@ -59,7 +57,7 @@ Another method of getting reference to an Actor can be done when spawning an Act
 
 	![](GHT5B_6.png)
 
-	This will create a variable out of the spawned Actor (which we named **SpawnedActor**). You could also drag off the **Return Value** to directly access the Variables, Events, or Functions inside your selected Actor that you are spawning, but we are making it a variable so that we may access it later. 
+	This will create a variable out of the spawned Actor (which we named **SpawnedActor**). You could also drag off the Return Value to directly access the Variables, Events, or Functions inside your selected Actor that you are spawning, but we are making it a variable so that we may access it later. 
 
 1.	You can now hold **Control** and drag the variable into the graph, then drag off it to access the properties inside it. 
 
@@ -76,7 +74,7 @@ There may be times when you have several Blueprint Actors in your level and you
 
 	![](GHT5B_8.png)
 
-	Here we are using the **Blueprint_Effect_Fire** and **Blueprint_Effect_Sparks** Actors included with **Starter Content**. 
+	Here we are using the **Blueprint_Effect_Fire** and **Blueprint_Effect_Sparks** Actors included with Starter Content. 
 
 	Say we want the fire effect to deactivate the spark effect after a few seconds when the game starts. 
 
@@ -86,7 +84,7 @@ There may be times when you have several Blueprint Actors in your level and you
 
 	![](GHT5B_9.png)
 
-	Reference Blueprint Actors are denoted with an **_C** following their name as seen above. 
+	Reference Blueprint Actors are denoted with a **_C** following their name as seen above. 
 
 1.	Hold **Control** and drag the new variable into the graph, then drag off it to access the Blueprint's properties. 
 
@@ -106,11 +104,11 @@ There may be times when you have several Blueprint Actors in your level and you
 
 	![](GHT5B_13.png)
 
-	The error states that we have **Accessed None** and is a common error that occurs when trying to communicate between Blueprints and get reference to Actors for the first time. 
+	The error states that we have Accessed None and is a common error that occurs when trying to communicate between Blueprints and get reference to Actors for the first time. 
 
-	In the **Blueprint_Effect_Fire** Blueprint, we identified our target Blueprint and created a variable for it and deactivated the sparks, however we need to specify the actual **Instance** of that Blueprint that we want to apply the change to. 
+	In the Blueprint_Effect_Fire Blueprint, we identified our target Blueprint and created a variable for it and deactivated the sparks, however we need to specify the actual Instance of that Blueprint that we want to apply the change to. 
 
-	We can do that by making the variable **Public** which will allow us to access it inside the Main Editor window where we can select which **Instance** should be affected. 
+	We can do that by making the variable Public which will allow us to access it inside the Main Editor window where we can select which Instance should be affected. 
 
 1.	In the **Blueprint_Effect_Fire** Blueprint, click the eye icon next to the **TargetBlueprint** variable. 
 
@@ -124,20 +122,20 @@ There may be times when you have several Blueprint Actors in your level and you
 
 	![](GHT5B_16.png)
 
-	You can use the drop-down menu to select the target Actor, or you can use the **Eye Dropper** and select the target Actor in the scene. 
+	You can use the drop-down menu to select the target Actor, or you can use the Eye Dropper and select the target Actor in the scene. 
 
 	![](GHT5B_17.png)
 
 1.	With the target Actor set, if you **Play** in the editor now, the spark effect will deactivate after two seconds. 
 
 	[REGION:note]
-	You can refer to [](Engine/Blueprints/UserGuide/BlueprintComms) and [](Engine\Blueprints\UserGuide\BlueprintCommsUsage) for more information on using **Direct Blueprint Communication**.
+	You can refer to [](Engine/Blueprints/UserGuide/BlueprintComms) and [](Engine\Blueprints\UserGuide\BlueprintCommsUsage) for more information on using Direct Blueprint Communication.
 	[/REGION]
 
 
 ## Cast To Referencing
 
-While the **Direct Blueprint Communication** method works for getting reference between Actors that exist in your level, there may be instances when you want to get a reference to an Actor in your level and pass it to a Blueprint Actor that is not in the level yet, for example the player character (or other Actor) that does not exist in the level until the game starts. In this case, you can use a **Cast To** node to send the reference to your **Target Blueprint**. 
+While the Direct Blueprint Communication method works for getting reference between Actors that exist in your level, there may be instances when you want to get a reference to an Actor in your level and pass it to a Blueprint Actor that is not in the level yet, for example the player character (or other Actor) that does not exist in the level until the game starts. In this case, you can use a Cast To node to send the reference to your Target Blueprint. 
 
 1.	Place an Actor in your level that you want to reference. 
 
@@ -157,7 +155,7 @@ While the **Direct Blueprint Communication** method works for getting reference
 
 	![](GHT5B_21.png)
 
-	Here we created a variable of the **Blueprint_Effect_Fire_C** type called **TargetBlueprint** (as the fire effect is what we want to reference) inside the **ThirdPersonCharacter** Blueprint. When the game is started the **Level Blueprint** will get the reference to the Actor in the level then pass it to the **ThirdPersonCharacter** Blueprint so that we may access its properties. 
+	Here we created a variable of the Blueprint_Effect_Fire_C type called TargetBlueprint (as the fire effect is what we want to reference) inside the ThirdPersonCharacter Blueprint. When the game is started the Level Blueprint will get the reference to the Actor in the level then pass it to the ThirdPersonCharacter Blueprint so that we may access its properties. 
 
 1.	**Control** drag the variable into the graph, then drag off it to access its properties. 
 
@@ -180,7 +178,7 @@ While the **Direct Blueprint Communication** method works for getting reference
 	If you followed this example, you should be able to deactivate the fire effect by pressing the **F** key. 
 
 	[REGION:note]
-	You can refer to [Casting to Blueprint Types](Engine/Blueprints/UserGuide/BlueprintComms) and [](Engine\Blueprints\UserGuide\BlueprintCommsUsage) for more information on **Actor Casting**.
+	You can refer to [Casting to Blueprint Types](Engine/Blueprints/UserGuide/BlueprintComms) and [](Engine\Blueprints\UserGuide\BlueprintCommsUsage) for more information on Actor Casting.
 	[/REGION]
 
 
