#UE4doc #docs:
Bolded and unbolded a few terms
Replaced a couple pictures due to the color/style of arrow being used

diff --git a/Engine/Documentation/Source/Engine/Paper2D/Flipbooks/Components/Blueprints/Components_BP.INT.udn b/Engine/Documentation/Source/Engine/Paper2D/Flipbooks/Components/Blueprints/Components_BP.INT.udn
--- a/Engine/Documentation/Source/Engine/Paper2D/Flipbooks/Components/Blueprints/Components_BP.INT.udn
+++ b/Engine/Documentation/Source/Engine/Paper2D/Flipbooks/Components/Blueprints/Components_BP.INT.udn
@@ -13,7 +13,7 @@ Related: Engine/Blueprints
 
 [TOC (start:2 end:2)]
 
-Working with **Flipbook Components** in **Blueprints** is very similar to working with other types of [Components](Engine/Blueprints/UserGuide/Components) in Blueprints. After adding a Flipbook Component to a Blueprint, you can specify the **Source Flipbook** asset to use, define the Flipbook's **Play Rate** or adjust the **Sprite Color** used in addition to standard options that are available to all components such as **Transforming**. 
+Working with **Flipbook Components** in Blueprints is very similar to working with other types of [Components](Engine/Blueprints/UserGuide/Components) in Blueprints. After adding a Flipbook Component to a Blueprint, you can specify the Source Flipbook asset to use, define the Flipbook's Play Rate or adjust the Sprite Color used in addition to standard options that are available to all components such as Transforming. 
 
 Inside the [Graph](Engine/Blueprints/UserGuide/Graphs) tab, you can directly access and modify the properties of the Flipbook Component through your script. This includes things like adjusting the Play Rate, Playing, Stopping or Reversing the Flipbook, or even changing the Source Flipbook to a different Flipbook asset (useful for when you want to switch between Flipbook animations for a character for example). 
 
@@ -23,7 +23,7 @@ This page will show you how to add, access and modify Flipbook Components in a B
 
 You can add a Flipbook Component to a Blueprint from the **Components** window of a Blueprint. 
 
-1. Inside a Blueprint, click the **Add Component** button inside the **Components** window and search for and add a **Paper Flipbook**. 
+1. Inside a Blueprint, **Left-Click** the **Add Component** button inside the **Components** window and search for and add a **Paper Flipbook**. 
 
 	![](Components1.png)	
 
@@ -31,7 +31,7 @@ You can add a Flipbook Component to a Blueprint from the **Components** window o
 
 	![](Components2.png)
 
-	You can then begin adjusting the **Properties** of the Flipbook Component inside the **Details** panel. 
+	You can then begin adjusting the Properties of the Flipbook Component inside the **Details** panel. 
 
 You can also add Flipbook Components to Blueprints by using **Blueprint/Add Script** option. 
 
@@ -87,11 +87,11 @@ When you are ready to start working with your Flipbook Component, you can head t
 
 ![](MyBlueprint1.png)
 
-To access or affect the Flipbook Component in anyway, you will need to hold **Ctrl** and drag-drop the Flipbook Component into a graph. 
+To access or affect the Flipbook Component in anyway, you will need to hold **Ctrl** and **Drag-Drop** the Flipbook Component into a graph. 
 
 ![](Graph1.png)
 
-This will **Get** the Flipbook Component at which point you can drag off it and call any number of Events or Functions on it. 
+This will Get the Flipbook Component at which point you can drag off it and call any number of Events or Functions on it. 
 
 ![](Graph2.png)(w:600)
 
diff --git a/Engine/Documentation/Source/Engine/Paper2D/Flipbooks/Editor/FlipbookEditor.INT.udn b/Engine/Documentation/Source/Engine/Paper2D/Flipbooks/Editor/FlipbookEditor.INT.udn
--- a/Engine/Documentation/Source/Engine/Paper2D/Flipbooks/Editor/FlipbookEditor.INT.udn
+++ b/Engine/Documentation/Source/Engine/Paper2D/Flipbooks/Editor/FlipbookEditor.INT.udn
@@ -10,24 +10,23 @@ Related: Engine/Content/Types/Textures
 
 [TOC (start:2 end:2)]
 
-## Flipbook Editor Overview
 
 The **Flipbook Editor** enables setup and editing of individual Flipbook assets and is where you can edit the sequintal playable of your sprites, tune or tweak the duration in which a sprite is displayed or determine the overall playback speed of the Flipbook animation through the **Frames Per Second** option. 
 
 ## Opening the Flipbook Editor
 
-After you've created a Flipbook, you can then open it for editing in the **Flipbook Editor** in the follow ways: 
+After you've created a Flipbook, you can then open it for editing in the Flipbook Editor in the follow ways: 
 
 1. **Double-click** on a Flipbook inside the **Content Browser** (or **Right-click** on a Flipbook and select **Edit**). 
 
 	![](EditFlipbook.png)
 
-1.	This will open the **Flipbook Editor** window. 
+1.	This will open the Flipbook Editor window. 
 
 	![](FlipbookEditorMain.png)(w:800)
 
 [REGION:note]
-With a Flipbook Actor in your level and the Actor selected, you can press **Ctrl+E** to open it in the **Flipbook Editor** as well. 
+With a Flipbook Actor in your level and the Actor selected, you can press **Ctrl+E** to open it in the Flipbook Editor as well. 
 [/REGION]
 
 ## Flipbook Editor Interface
@@ -49,10 +48,10 @@ Each of the **Tool Bar** options are described in the table below:
 |----|----|
 |![](Toolbar1.png)(w:60)| Saves the Flipbook asset. |
 |![](Toolbar2.png)(w:60)| Locates and selectes the current Flipbook asset inside the **Content Browser** (**Ctrl+B**). |
-|![](Toolbar3.png)(w:60)| This toggles the grid display inside the viewport. |
-|![](Toolbar4.png)(w:60)| Toggles the display of the bounds of the Static Mesh in the viewport. |
+|![](Toolbar3.png)(w:60)| This toggles the grid display inside the **Viewport**. |
+|![](Toolbar4.png)(w:60)| Toggles the display of the bounds of the Static Mesh in the **Viewport**. |
 |![](Toolbar5.png)(w:60)| Toggles the display of the simplified collision mesh of the static mesh, if one has been assigned (**Alt+C**). |
-|![](Toolbar6.png)(w:60)| Displays the pivot location of the sprite in the viewport. |
+|![](Toolbar6.png)(w:60)| Displays the pivot location of the sprite in the **Viewport**. |
 
 ### Viewport
 
@@ -74,7 +73,7 @@ Mousing over a sprite in the Timeline will also give you more information about
 
 ![](Timeline3.png)
 
-Above the name of the **Sprite** that is being used, its **Index** (location in the sequence) and **Duration** in frames are displayed. 
+Above the name of the Sprite that is being used, its **Index** (location in the sequence) and **Duration** in frames are displayed. 
 
 You can also **Right-click** on a key frame to bring up a context menu with additional key frame actions. 
 
@@ -88,7 +87,7 @@ When you insert a key frame before or after a key, an empty key frame will be ad
 
 You will need to specify the sprite to be used for those key frames, otherwise the Flipbook will render nothing in those frames. 
 
-**Sprites** or **Flipbook** assets can also be drag-dropped onto the Timeline to append them as additional key frames to the end of the Timeline. 
+Sprites or Flipbook assets can also be drag-dropped onto the Timeline to append them as additional key frames to the end of the Timeline. 
 
 ![](Timeline6.png)
 
@@ -100,7 +99,7 @@ The **Scrub Bar/Playback Controls** allow you to scrub through or use the playba
 
 ### Details
 
-The **Details** panel inside the **Flipbook Editor** is where you determine the animation playrate through **Frames Per Second**, assign **Key Frames** and their order as well as the **Sprites** to use at each key frame and their duration with the **Frame Run** setting. There's also an option to assign a **Default Material** to use for the Flipbook player instance (there is a Default Material supplied here that you can use or override with your own Material if you wish). 
+The **Details** panel inside the Flipbook Editor is where you determine the animation playrate through **Frames Per Second**, assign **Key Frames** and their order as well as the Sprites to use at each key frame and their duration with the **Frame Run** setting. There's also an option to assign a **Default Material** to use for the Flipbook player instance (there is a Default Material supplied here that you can use or override with your own Material if you wish). 
 
 ![](DetailsPanel.png)
 
@@ -131,7 +130,7 @@ If you are working with a blank Flipbook, you can add sprites to the Flipbook in
 
 You can **Drag-drop** a Sprite (or Flipbook) asset onto the Timeline from the **Content Browser**. 
 
-![](Adding1.png)
+![](Timeline6.png)
 
 Which will add the Sprite(s) to the Timeline as well as update the **Key Frames** section inside the **Details** panel. 
 
