#UE4doc #docs: Removed extraneous bolding as according to the instructional guide. ##Overview removed as according to the instructional guide. Changed title of How To to match instructional guide.

diff --git a/Engine/Documentation/Source/Gameplay/HowTo/RespawnPlayer/GHT_3_Overview.INT.udn b/Engine/Documentation/Source/Gameplay/HowTo/RespawnPlayer/GHT_3_Overview.INT.udn
--- a/Engine/Documentation/Source/Gameplay/HowTo/RespawnPlayer/GHT_3_Overview.INT.udn
+++ b/Engine/Documentation/Source/Gameplay/HowTo/RespawnPlayer/GHT_3_Overview.INT.udn
@@ -1,16 +1,13 @@
 Availability: Public
-Title:How to Respawn a Player
+Title:Respawning a Player
 Crumbs: %ROOT%, Gameplay
 Description:A How To Guide for respawning player characters in Unreal Engine 4. 
 
-
-##Overview
-
-In most games, there comes a time when the player character is killed and you want to respawn them in the world. In this **How To** guide, you will learn how to set up a simple respawn system for single player purposes. 
+In most games, there comes a time when the player character is killed and you want to respawn them in the world. In this How To guide, you will learn how to set up a simple respawn system for single player purposes. 
 
 ## Setup
 
-To learn more about **Respawning a Player** using Blueprints, follow the link below:
+To learn more about Respawning a Player using Blueprints, follow the link below:
 
 [REGION:buttonlist]
 * [](Gameplay\HowTo\RespawnPlayer\Blueprints)
