#UE4doc #docs:
Changed value provided by a step due to it being inaccurate
Changed picture to reflect previous change

diff --git a/Engine/Documentation/Source/Engine/AI/EnvironmentQuerySystem/QuickStart/12/12.INT.udn b/Engine/Documentation/Source/Engine/AI/EnvironmentQuerySystem/QuickStart/12/12.INT.udn
--- a/Engine/Documentation/Source/Engine/AI/EnvironmentQuerySystem/QuickStart/12/12.INT.udn
+++ b/Engine/Documentation/Source/Engine/AI/EnvironmentQuerySystem/QuickStart/12/12.INT.udn
@@ -119,7 +119,7 @@ Just a quick reminder:
 * **Space Between** = `150.0`
 * **Navigation Filer** = `RecastFilter_UseDefaultArea`
 * **Project Down** = `2048.0`
-* **Post Projection Vertical Offset** = `120.0`
+* **Post Projection Vertical Offset** = `32.0`
 [/EXCERPT]
 
 [EXCERPT:eqsqs122]
