#UE4doc #docs:
Unbolded multiple terms

diff --git a/Engine/Documentation/Source/Engine/AI/EnvironmentQuerySystem/QuickStart/14/14.INT.udn b/Engine/Documentation/Source/Engine/AI/EnvironmentQuerySystem/QuickStart/14/14.INT.udn
--- a/Engine/Documentation/Source/Engine/AI/EnvironmentQuerySystem/QuickStart/14/14.INT.udn
+++ b/Engine/Documentation/Source/Engine/AI/EnvironmentQuerySystem/QuickStart/14/14.INT.udn
@@ -24,13 +24,13 @@ This page covers a number of ideas for improvements you can make now that you're
 
 ## Ideas
 
-* **TopDownCharacter** is far to specific, update everything to a more flexible system.
-* The **Behavior Tree** shows one way to set a system like this up:
-	* Try using a **Wait** node instead of the **Cooldown** Decorator. This will require a **Bool** Blackboard entry and some extra work in the AI Controller.
+* TopDownCharacter is far to specific, update everything to a more flexible system.
+* The Behavior Tree shows one way to set a system like this up:
+	* Try using a Wait node instead of the Cooldown Decorator. This will require a Bool Blackboard entry and some extra work in the AI Controller.
 	* Add a secondary Behavior that, after a time, the AI will begin to wander around.
-	* **Bonus** - Store the last know location(s) of the Player(s) and use those in an Environment Query to find a better hiding spot.
-	* **Double Bonus** - Have the AI attempt to approach a Player from behind without entering the forward view arc of the Player(s).
-* Did you know you can create your own **Generators** in Blueprint? Maybe there's a better way to find the possible points for the AI to hide.
+	* Bonus - Store the last know location(s) of the Player(s) and use those in an Environment Query to find a better hiding spot.
+	* Double Bonus - Have the AI attempt to approach a Player from behind without entering the forward view arc of the Player(s).
+* Did you know you can create your own Generators in Blueprint? Maybe there's a better way to find the possible points for the AI to hide.
 * [REGION:raw]The ' key (single quote) turns on a host of Gameplay Debugging tools. Check out all of the options available in the **Editor Preferences** %arrow% **Gameplay Debugger**[/REGION]
 
 [OBJECT:TopicButtonList]
