#UE4doc #docs: Removed ôOverviewö header

diff --git a/Engine/Documentation/Source/Engine/UI/LevelEditor/Viewports/ViewportToolbar/ViewportToolbar.INT.udn b/Engine/Documentation/Source/Engine/UI/LevelEditor/Viewports/ViewportToolbar/ViewportToolbar.INT.udn
--- a/Engine/Documentation/Source/Engine/UI/LevelEditor/Viewports/ViewportToolbar/ViewportToolbar.INT.udn
+++ b/Engine/Documentation/Source/Engine/UI/LevelEditor/Viewports/ViewportToolbar/ViewportToolbar.INT.udn
@@ -46,7 +46,7 @@ Description:A rundown of the various tools available within the viewports integr
 
 [TOC(start:2)]
 
-## Overview
+
 
 The Viewport Toolbar contains the transform tools that you will use throughout the level design process, as well as snapping controls for those tools. You will also find access to camera speed, as well as the ability to split the viewport into a 4-view.
 
