#UE4doc #docs: Removed ôOverviewö header & Adjusted bold text in accordance with the Instructional Guide

diff --git a/Engine/Documentation/Source/Engine/UI/LevelEditor/Viewports/ViewportControls/ViewportControls.INT.udn b/Engine/Documentation/Source/Engine/UI/LevelEditor/Viewports/ViewportControls/ViewportControls.INT.udn
--- a/Engine/Documentation/Source/Engine/UI/LevelEditor/Viewports/ViewportControls/ViewportControls.INT.udn
+++ b/Engine/Documentation/Source/Engine/UI/LevelEditor/Viewports/ViewportControls/ViewportControls.INT.udn
@@ -45,7 +45,7 @@ Description:Learn about the various control schemes for the editor's viewports.
 
 [TOC(start:2 end:3)]
 
-## Overview
+
 
 There are various controls to enable you to navigate the scene, select and manipulate Actors, and change display 
 options while working in the viewports.
