#UE4doc #docs: Removed extraneous bolding as according to the instructional guide. Changed title of How To to match instructional guide.

diff --git a/Engine/Documentation/Source/Gameplay/HowTo/RespawnPlayer/Blueprints/GHT_3B.INT.udn b/Engine/Documentation/Source/Gameplay/HowTo/RespawnPlayer/Blueprints/GHT_3B.INT.udn
--- a/Engine/Documentation/Source/Gameplay/HowTo/RespawnPlayer/Blueprints/GHT_3B.INT.udn
+++ b/Engine/Documentation/Source/Gameplay/HowTo/RespawnPlayer/Blueprints/GHT_3B.INT.udn
@@ -1,13 +1,11 @@
 Availability: Public
-Title:How to Respawn a Player in Blueprints
+Title:Respawning a Player in Blueprint
 Crumbs: %ROOT%, Gameplay, Gameplay/HowTo/RespawnPlayer
 Description:A How To Guide for respawning a player character through Blueprints in Unreal Engine 4. 
 
 [TOC (start:2)]
 
-## Overview
-
- Here you will learn how to respawn a player character in a single player game in **Blueprints**. 
+ Here you will learn how to respawn a player character in a single player game in Blueprints. 
 
 ## Respawning the Player (Single Player)
 
@@ -23,15 +21,15 @@ _The steps below are based off using the **Blueprint Third Person Template** pro
 
 	![](GHT3B_1.png)
 
-	You want to **Cast To** your Player Character Blueprint, which is **ThirdPersonCharacter** in this example. This will allow us to access it and determine when the player character is killed by using an **On Destroyed** event, at which point we can spawn a new player character. 
+	You want to Cast To your Player Character Blueprint, which is ThirdPersonCharacter in this example. This will allow us to access it and determine when the player character is killed by using an On Destroyed event, at which point we can spawn a new player character. 
 
-1.	Drag off the **As ThirdPersonCharacter C** pin and add a **Assign On Destroyed** node. 
+1.	Drag off the **As ThirdPersonCharacter C** pin and add an **Assign On Destroyed** node. 
 
 1.	Off the newly added **OnDestroyed_Event** node, add a **Delay** node set to whatever time you wish to delay the respawn (optional). 
 
 	![](GHT3B_2.png)
 
-	Here we are Binding an Event to another Event from another Blueprint. This allows the Event in this Blueprint to fire off when the Event occurs from the other Blueprint (the OnDestroy Event from **ThirdPersonCharacter**).  This is called an **Event Dispatcher** and can also be used to communicate with the **Level Blueprint** to tell the world an Event has occurred. 
+	Here we are Binding an Event to another Event from another Blueprint. This allows the Event in this Blueprint to fire off when the Event occurs from the other Blueprint (the OnDestroy Event from ThirdPersonCharacter).  This is called an **Event Dispatcher** and can also be used to communicate with the **Level Blueprint** to tell the world an Event has occurred. 
 	
 	[REGION:note]
 	Refer to the [](Engine/Blueprints/UserGuide/EventDispatcher) documentation for more information.
