#UE4doc #docs: Adjusted bold text in accordance with the Instructional Guide

diff --git a/Engine/Documentation/Source/Engine/UMG/UserGuide/UserGuide.INT.udn b/Engine/Documentation/Source/Engine/UMG/UserGuide/UserGuide.INT.udn
--- a/Engine/Documentation/Source/Engine/UMG/UserGuide/UserGuide.INT.udn
+++ b/Engine/Documentation/Source/Engine/UMG/UserGuide/UserGuide.INT.udn
@@ -50,7 +50,7 @@ Description:A guide to using Unreal Motion Graphics to create UI elements.
 ## Overview
 
 [EXCERPT:umgheader]
-**Unreal Motion Graphics UI Designer (UMG)**  is a visual UI authoring tool which can be used to create UI elements such as in-game HUDs, menus or other interface related graphics you wish to present to your users. At the core of **UMG** are **Widgets**, which are a series of pre-made functions that can be used to construct your interface (things like buttons, checkboxes, sliders, progress bars, etc.). These **Widgets** are edited in a specialized **Widget Blueprint**, which uses two tabs for construction: the **Designer** tab allows for the visual layout of the interface and basic functions while the **Graph** tab provides the functionality behind the **Widgets** used. 
+**Unreal Motion Graphics UI Designer (UMG)**  is a visual UI authoring tool which can be used to create UI elements such as in-game HUDs, menus or other interface related graphics you wish to present to your users. At the core of UMG are Widgets, which are a series of pre-made functions that can be used to construct your interface (things like buttons, checkboxes, sliders, progress bars, etc.). These Widgets are edited in a specialized Widget Blueprint, which uses two tabs for construction: the Designer tab allows for the visual layout of the interface and basic functions while the Graph tab provides the functionality behind the Widgets used. 
 [/EXCERPT:umgheader]
 
 [OBJECT:TopicButtonList]
