#UEDoc #docs - updated doc and images.

diff --git a/Engine/Documentation/Source/Engine/Matinee/HowTo/MHT_8/MHT_8.INT.udn b/Engine/Documentation/Source/Engine/Matinee/HowTo/MHT_8/MHT_8.INT.udn
--- a/Engine/Documentation/Source/Engine/Matinee/HowTo/MHT_8/MHT_8.INT.udn
+++ b/Engine/Documentation/Source/Engine/Matinee/HowTo/MHT_8/MHT_8.INT.udn
@@ -22,7 +22,7 @@ In the following steps, we will create a rock that explodes after a timer count
 
 	![](Engine\Matinee\HowTo\MHT_1\MHT1_AddMatinee.png)(w:720)
 
-1. If the **Matinee** undo warning message appears, click the **Continue** button. 
+1. If the Matinee undo warning message appears, **click** the **Continue** button. 
 
 	![](Engine\Matinee\HowTo\MHT_1\MHT1_MatineeWarning.png)
 	
@@ -34,61 +34,61 @@ In the following steps, we will create a rock that explodes after a timer count
 
 	_Click for Full View._
 
-1.  Minimize the **Matinee** and in the **Content Browser** under the **Game/Props** folder, **Right-click** on the **SM_Rock** and select **Create Destructible Mesh**. 
+1.  **Minimize** the Matinee and in the **Content Browser** under the **Game/Props** folder, **Right-click** on the **SM_Rock** and **select Create Destructible Mesh**. 
 
 	![](MHT8_CreateDM.png)
 
-1.  The **Destructible Mesh** window will open, click the **Fracture Mesh** button from the tool bar. 
+1.  The **Destructible Mesh** window will open, **click** the **Fracture Mesh** button from the tool bar. 
 
 	![](MHT8_FractureMesh.png)(w:720)
 
-1.  Under **Destructible Settings**, expand **Fracture Effects** and **0**, then under **Particle System** click the drop-down menu and select **P_Explosion**. 
+1.  Under **Destructible Settings**, **expand Fracture Effects** and **0**, then under **Particle System** click the drop-down menu and **select P_Explosion**. 
 
 	![](MHT8_AssignParticleEffct.png)
 
-1.  In the same section, click the **Sound** drop-down menu and select the **Explosion01** asset. 
+1.  In the same section, **click** the **Sound** drop-down menu and **select** the **Explosion01** asset. 
 
 	![](MHT8_AssignSound.png)
 
-1.  Click the **Save** button from the tool bar then close the window. 
+1.  **Click** the **Save** button from the tool bar then close the window. 
 
-1.  In the **Content Browser**, drag the **SM_Rock_DM** asset into the level. 
+1.  In the Content Browser, **drag** the **SM_Rock_DM** asset into the level. 
 
 	![](MHT8_DragRockIn.png)
 
-1.  In the **Content Browser** under **Game/Particles**, drag-and-drop the **P_Fire** effect into the level, inside the rock. 
+1.  In the Content Browser under **Game/Particles**, **drag-and-drop** the **P_Fire** effect into the level, inside the rock. 
 
 	![](MHT8_AddFire.png)
 
-1.  In the **Details** panel for **P_Fire**, scroll down and find the **Activation** section and un-check **Auto Activate**. 
+1.  In the Details panel for P_Fire, scroll down and find the **Activation** section and **un-check Auto Activate**. 
 
 	![](MHT8_UncheckAutoActivate.png)
 
 	_This will keep the asset in the level, but not activate it until we tell it to which we will do through **Matinee**._
 
-1.	Select the **TextRenderActor** text and move it to the little box as shown below (press the **E** key to enter rotation mode and rotate it forward). 
+1.	**Select** the **TextRenderActor** text and move it to the little box as shown below (press the **E** key to enter rotation mode and rotate it forward). 
 
 	![](MHT8_MoveTemplateText.png)
 
-1.	In the **Details** panel for the **TextRenderActor**, under the **Text** section, in the **Text** field replace the text with **3**. 
+1.	In the Details panel for the TextRenderActor, under the **Text** section, in the Text field replace the text with **3**. 
 
 	![](MHT8_3.png)
 
 	_This is going to serve as our timer and countdown to the explosion we will trigger._
 
-1.	Also inside the **Details** panel, make a note of what the **Location** values are under **Transform** for the text as we will need them in a moment. 
+1.	Also, inside the Details panel make a note of what the **Location** values are under **Transform** for the text as we will need them in a moment. 
 
 	![](MHT8_TextLocation.png)
 
-1.	In the level viewport with the text still selected, press **Control+W** to duplicate the Text Render Actor. 
+1.	In the level viewport with the text still selected, **press Control+W** to duplicate the Text Render Actor. 
 
-1.	Inside the **Details** panel for the new Text Render Actor, set **Location** to the same as the other Text Render Actor and set the **Text** to **2**. 
+1.	Inside the Details panel for the new Text Render Actor, **set Location** to the same as the other Text Render Actor and set the **Text** to **2**. 
 
 	![](MHT8_2.png)
 
-1.	Repeat the previous two steps to create another Text Render Actor, set it in the same location as the others, and set the text to **1**. 
+1.	**Repeat the previous two steps** to create another Text Render Actor. Set it in the same location as the others and set the text to **1**. 
 
-1.	For each of the Text Render Actors, in the **Details** panel under **Rendering**, check the box for **Actor Hidden In Game**. 
+1.	For each of the Text Render Actors in the Details panel under **Rendering**, **check** the box for **Actor Hidden In Game**. 
 
 	![](MHT8_ActorHiddenInGame.png)
 
@@ -103,51 +103,51 @@ In the following steps, we will create a rock that explodes after a timer count
 	[/PARAM]
 [/OBJECT]
 
-1.  Select the **Matinee** in the **World Outliner** or viewport, then in the **Details** panel select **Open Matinee** and minimize it when it opens. 
+1.  **Select** the Matinee in the **World Outliner** or viewport, then in the Details panel **select Open Matinee** and minimize it when it opens. 
 
-1.  Select the rock in the viewport, then back in **Matinee** in the **Tracks Window**, **Right-click** and select **Add New Empty Group**. 
+1.  **Select** the rock in the viewport, then back in Matinee in the **Tracks Window**, **Right-click** and **select Add New Empty Group**. 
 
 	![](MHT8_AddEmptyGroup.png)
 
 	_Give the new group a name when prompted, in this example we are simply calling it the **Rock** group._
 
-1.  **Right-click** on the group and select **Add New Event Track** from the context menu. 
+1.  **Right-click** on the group and **select Add New Event Track** from the context menu. 
 
 	![](MHT8_AddEventTrack.png)
 
-1.  With the **Event Track** selected, move the **Time Bar** to **4.00** and press **Enter** to add a keyframe, then name the event **Explode**.
+1.  With the Event Track selected, **move** the **Time Bar** to **4.00** and **press Enter** to add a keyframe, then name the event **Explode**.
 
 	![](MHT8_EventTrackComplete.png)
 
-1.  **Right-click** on the red ending marker at **5.00** and choose **Move To Longest Track Endpoint** to shorten the sequence to **4.00**. 
+1.  **Right-click** on the red ending marker at **5.00** and **choose Move To Longest Track Endpoint** to shorten the sequence to **4.00**. 
 
 	![](MHT8_SetLength.png)
 
-1.  Minimize **Matinee** and select the Text Render Actor who is set to **3** in the **World Outliner**. 
+1.  **Minimize** Matinee and **select** the Text Render Actor who is set to **3** in the World Outliner. 
 
-1.	Back in **Matinee** in the **Tracks Window**, **Right-click** and select **Add New Empty Group** and name the group something like **Text-3**. 
+1.	Back in Matinee in the Tracks Window, **Right-click** and **select** **Add New Empty Group** and name the group something like **Text-3**. 
 
 	![](MHT8_Text3Added.png)
 
-1.	**Right-click** on the group and select **Add New Visibility Track**. 
+1.	**Right-click** on the group and **select** **Add New Visibility Track**. 
 
 	![](MHT8_AddNewVisTrack.png)
 
-1.	Repeat the previous two steps and create two new groups (**Text-2** and **Text-1**), both with **Visibility** tracks. 
+1.	**Repeat the previous two steps** and create two new groups (**Text-2** and **Text-1**), both with **Visibility** tracks. 
 
 	![](MHT8_MatineeSetupComplete.png)
 
-1.	Click on the **Visibility** track for **Track-3**, move the **Time Bar** to **1.00** and press **Enter**, then select **Show** from the pop-up menu to add a keyframe. 
+1.	**Click** on the Visibility track for **Track-3**, move the Time Bar to **1.00** and **press** **Enter**, then **select** Show from the pop-up menu to add a keyframe. 
 
 	![](MHT8_Show3.png)
 
-1.	Move the **Time Bar** to **2.00** and press **Enter** and select **Hide** from the pop-up menu. 
+1.	**Move** the Time Bar to **2.00** and **press** **Enter** and **select** Hide from the pop-up menu. 
 
 	![](MHT8_Hide3.png)
 
 	_A bar will be created between **1.00** and **2.00** which indicates when the text **3** will be shown, which is for 1 second._
 
-1.	Repeat the previous two steps to **Show/Hide** the text **2** from **2.00** to **3.00** and the text **1** from **3.00** to **4.00**, then close the **Matinee**. 
+1.	Repeat the previous two steps to Show/Hide the text **2** from **2.00** to **3.00** and the text **1** from **3.00** to **4.00**, then close the Matinee. 
 
 	![](MHT8_MatineeFinished.png)
 
@@ -161,55 +161,55 @@ In the following steps, we will create a rock that explodes after a timer count
 	[/PARAM]
 [/OBJECT]
 
-1.	Select the **MatineeActor** in the level viewport or **World Outliner** then click the **Blueprints** button from the main tool bar and select **Open Level Blueprint**. 
+1.	**Select** the **MatineeActor** in the level viewport or World Outliner then **click** the **Blueprints** button from the main tool bar and select **Open Level Blueprint**. 
 
 	![](MHT8_OpenLevelBlueprint.png)
 
-1.	**Right-click** in the graph window and select the **Add MatineeController** option. 
+1.	**Right-click** in the graph window and **select** the **Add MatineeController** option. 
 
 	![](MHT8_AddMatineeController.png)
 
-1.	Minimize the **Level Blueprint** and select the rock in the level, then back in the **Level Blueprint**, **Right-click** in the graph and choose **Add Reference to SM_Rock_DM**. 
+1.	**Minimize** the **Level Blueprint** and **select** the rock in the level, then back in the Level Blueprint, **Right-click** in the graph and choose **Add Reference to SM_Rock_DM**. 
 
 	![](MHT8_AddReferenceToRock.png)
 
-1.	Minimize the **Level Blueprint** again and select the **P_Fire** asset, then back in the **Level Blueprint**, **Right-click** and choose **Add Reference to P_Fire**.
+1.	**Minimize** the Level Blueprint again and **select** the **P_Fire** asset, then back in the Level Blueprint, **Right-click** and choose **Add Reference to P_Fire**.
 
 	![](MHT8_ReferencesAdded.png)
 
-1.	**Left-click** on and drag off the out pin of the **SM_Rock_DM** node and in the window that appears, enter the text **Apply** then click **Apply Radius Damage**. 
+1.	**Left-click** on and **drag off** the out pin of the **SM_Rock_DM** node and in the window that appears, **enter** the text **Apply** then **click** **Apply Radius Damage**. 
 
 	![](MHT8_ApplyRadiusDamage.png)
 
-1.	**Left-click** on and drag off the **Explode** out pin for the **MatineeActor** and connect it to the in pin of the **Apply Radius Damage** node, then set each of the float values to **12000**. 
+1.	**Left-click** on and **drag off** the **Explode** out pin for the **MatineeActor** and connect it to the in pin of the Apply Radius Damage node, then **set** each of the float values to **12000**. 
 
 	![](MHT8_ExplodeComplete.png)
 
-1.	**Left-click** on and drag off the out pin of the **SM_Rock_DM** node and in the window that appears, enter the text **Location** and choose **Get Actor Location**. 
+1.	**Left-click** on and **drag off** the out pin of the SM_Rock_DM node and in the window that appears, **enter** the text **Location** and choose **Get Actor Location**. 
 
 	![](MHT8_GetActorLocation.png)
 
-1.	**Left-click** on and drag off the out pin of the **Get Actor Location** node and connect it to the **Hurt Origin** in pin of the **Apply Radius Damage** node. 
+1.	**Left-click** on and **drag off** the out pin of the Get Actor Location node and **connect it** to the **Hurt Origin** in pin of the Apply Radius Damage node. 
 
 	![](MHT8_HurtOrigin.png)
 
-1.	**Left-click** on and drag off the out pin of the **P_Fire** node and in the window that appears, enter the text **Activate** and choose the **Activate** option. 
+1.	**Left-click** on and **drag off** the out pin of the **P_Fire** node and in the window that appears, **enter** the text **Activate** and choose the **Activate** option. 
 
 	![](MHT8_ActivateNode.png)
 
-1.	**Left-click** on and drag off the **Finished** out pin for the **MatineeActor** and connect it to the in pin of the **Activate** node. 
+1.	**Left-click** on and **drag off** the **Finished** out pin for the **MatineeActor** and connect it to the in pin of the Activate node. 
 
 	![](MHT8_BlueprintFinished.png)
 
-1.	Click the **Compile** button on the tool bar and close the **Level Blueprint**. 
+1.	**Click** the **Compile** button on the tool bar and **close** the Level Blueprint. 
 
 	![](MHT8_Compile.png)
 
-1.	In the **World Outliner**, click the **MatineeActor**, then in the **Details** panel click the **Play on Level Load** option. 
+1.	In the World Outliner, click the **MatineeActor**, then in the Details panel **click** the **Play on Level Load** option. 
 
 	![](MHT8_PlayOnLoad.png)
 
-1.	Click the **Play** button from the main tool bar to play in the editor. 
+1.	**Click** the **Play** button from the main tool bar to play in the editor. 
 
 	![](MHT8_Gif.gif)(convert:false)
 
