#UE4doc #docs: Removed extraneous bolding as according to the instructional guide. ##Overview removed as according to the instructional guide. Changed title of How To to match instructional guide.

diff --git a/Engine/Documentation/Source/Gameplay/HowTo/SettingUpAGameMode/GameModeOverview.INT.udn b/Engine/Documentation/Source/Gameplay/HowTo/SettingUpAGameMode/GameModeOverview.INT.udn
--- a/Engine/Documentation/Source/Gameplay/HowTo/SettingUpAGameMode/GameModeOverview.INT.udn
+++ b/Engine/Documentation/Source/Gameplay/HowTo/SettingUpAGameMode/GameModeOverview.INT.udn
@@ -1,18 +1,15 @@
 Availability: Public
-Title:How to Set Up a Game Mode
+Title:Setting Up a Game Mode
 Crumbs: %ROOT%, Gameplay
 Description:A How To Guide for Setting Up a Game Mode in Unreal Engine 4. 
 
-
-##Overview
-
 The **GameMode** defines the game's set of rules. The rules can include how players join the game, game pausing, and level transition, as well as any game-specific behavior such as win conditions. The GameMode is set for each level, and GameModes can be reused in multiple levels. 
 
-The link below in the **Setup** section will give you examples on **How to Set Up a Game Mode** using Blueprints. 
+The link below in the Setup section will give you examples on How to Set Up a Game Mode using Blueprints. 
 
 ## Setup
 
-To learn more about **Setting Up a Game Mode** using Blueprints, follow the link below:
+To learn more about Setting Up a Game Mode using Blueprints, follow the link below:
 
 [REGION:buttonlist]
 * [](Gameplay\HowTo\SettingUpAGameMode\Blueprints)
