#UE4doc #docs: Removed extraneous bolding as according to the instructional guide. Changed title of How To to match instructional guide. Altered images to match the instructional guide: green -> yellow.

diff --git a/Engine/Documentation/Source/Gameplay/HowTo/RespawnPlayer/Blueprints/GHT_3B.INT.udn b/Engine/Documentation/Source/Gameplay/HowTo/RespawnPlayer/Blueprints/GHT_3B.INT.udn
--- a/Engine/Documentation/Source/Gameplay/HowTo/RespawnPlayer/Blueprints/GHT_3B.INT.udn
+++ b/Engine/Documentation/Source/Gameplay/HowTo/RespawnPlayer/Blueprints/GHT_3B.INT.udn
@@ -1,5 +1,5 @@
 Availability: Public
-Title:Respawning a Player in Blueprint
+Title:Respawning a Player in Blueprints
 Crumbs: %ROOT%, Gameplay, Gameplay/HowTo/RespawnPlayer
 Description:A How To Guide for respawning a player character through Blueprints in Unreal Engine 4. 
 
diff --git a/Engine/Documentation/Source/Gameplay/HowTo/SettingUpAGameMode/Blueprints/GameMode_BP.INT.udn b/Engine/Documentation/Source/Gameplay/HowTo/SettingUpAGameMode/Blueprints/GameMode_BP.INT.udn
--- a/Engine/Documentation/Source/Gameplay/HowTo/SettingUpAGameMode/Blueprints/GameMode_BP.INT.udn
+++ b/Engine/Documentation/Source/Gameplay/HowTo/SettingUpAGameMode/Blueprints/GameMode_BP.INT.udn
@@ -1,13 +1,11 @@
 Availability: Public
-Title:How to Set Up a Game Mode in Blueprints
+Title:Setting Up a Game Mode in Blueprints
 Crumbs: %ROOT%, Gameplay, Gameplay\HowTo\SettingUpAGameMode
 Description:A How To Guide for Setting Up a Game Mode in Blueprints.
 
 [TOC (start:2 end:2)]
 
-## Overview
-
-This page covers how to create a **GameMode Blueprint** and set default values for it, how to assign a **Default GameMode** for your game, and how to override the Default GameMode through the **World Settings** and **GameMode Override** option. 
+This page covers how to create a GameMode Blueprint and set default values for it, how to assign a Default GameMode for your game, and how to override the Default GameMode through the World Settings and GameMode Override option. 
 
 ## Creating a Game Mode Blueprint
 
@@ -27,16 +25,16 @@ In the **Pick Parent Class** window, select the **Game Mode** Class.
 
 	![](GameMode_5.png)
 
-	Here we are assigning the Character Blueprint called **ThirdPersonCharacter** as the **Default Pawn Class** for players to use in the game. 
+	Here we are assigning the Character Blueprint called ThirdPersonCharacter as the Default Pawn Class for players to use in the game. 
 
 	[REGION:note]
-	The **Game Mode** Blueprint points to existing Blueprints of the **Character**, **HUD**, **PlayerController**, **Spectator**, and **Game State** Classes. You will need to create these separately then specify them for use in the **Game Mode** Blueprint in order to actually use them in your game. 
+	The **Game Mode** Blueprint points to existing Blueprints of the Character, HUD, PlayerController, Spectator, and Game State Classes. You will need to create these separately then specify them for use in the Game Mode Blueprint in order to actually use them in your game. 
 	[/REGION]
 
 
 ## Assigning a Default Game Mode
 
-In the previous section, we created a **Game Mode** Blueprint. Once you have a **Game Mode** Blueprint, you can assign it as the **Default Game Mode** to use in your game. The steps below will guide you through assigning the **Default Game Mode** through the **Project Settings** option. 
+In the previous section, we created a Game Mode Blueprint. Once you have a Game Mode Blueprint, you can assign it as the Default Game Mode to use in your game. The steps below will guide you through assigning the Default Game Mode through the Project Settings option. 
 
 1.	From the Main Editor Window, click the **Edit** button from the Menu Bar, then select **Project Settings**. 
 
@@ -61,7 +59,7 @@ In the previous section, we created a **Game Mode** Blueprint. Once you have a *
 	
 ## Overriding the Default Game Mode
 
-When you have a **Default Game Mode** assigned, you can overwrite it on a per level basis through the **World Settings** menu under the **GameMode Override** section. The steps below will show you how to override the default **Game Mode**. 
+When you have a Default Game Mode assigned, you can overwrite it on a per level basis through the World Settings menu under the GameMode Override section. The steps below will show you how to override the default Game Mode. 
 
 1.	From the Main Editor Window, click the **World Settings** button from the Main Toolbar. 
 
