#UE4doc #docs: Removed extraneous bolding as according to the instructional guide. ##Overview removed as according to the instructional guide. Changed title of How To to match instructional guide.

diff --git a/Engine/Documentation/Source/Gameplay/HowTo/SpawnAndDestroyActors/GHT_1_Overview.INT.udn b/Engine/Documentation/Source/Gameplay/HowTo/SpawnAndDestroyActors/GHT_1_Overview.INT.udn
--- a/Engine/Documentation/Source/Gameplay/HowTo/SpawnAndDestroyActors/GHT_1_Overview.INT.udn
+++ b/Engine/Documentation/Source/Gameplay/HowTo/SpawnAndDestroyActors/GHT_1_Overview.INT.udn
@@ -1,12 +1,9 @@
 Availability: Public
-Title:How to Spawn/Destroy an Actor Overview
+Title:Spawning/Destroying an Actor Overview
 Crumbs: %ROOT%, Gameplay
 Description:A How To Guide for Spawning Actors in Unreal Engine 4. 
 
-
-##Overview
-
-Often times in your game you will find that you want to add items to your levels dynamically during gameplay, this is called **Spawning** an **Actor**. Items you want to spawn could be anything from other players or enemies to pick-ups such as Health, Weapons, or Rewards. On this page, you will find information regarding how to spawn an **Actor** through **Blueprints**.<!-- as well as other information regarding **Actor Spawning** and links to more documentation.--> 
+Often times in your game you will find that you want to add items to your levels dynamically during gameplay, this is called Spawning an Actor. Items you want to spawn could be anything from other players or enemies to pick-ups such as Health, Weapons, or Rewards. On this page, you will find information regarding how to spawn an Actor through Blueprints.<!-- as well as other information regarding **Actor Spawning** and links to more documentation.--> 
 
 ## Setup
 
