#UEDoc #docs - updated doc grammar and instructional guide pass

diff --git a/Engine/Documentation/Source/Engine/Matinee/TrackReference/TrackReference.INT.udn b/Engine/Documentation/Source/Engine/Matinee/TrackReference/TrackReference.INT.udn
--- a/Engine/Documentation/Source/Engine/Matinee/TrackReference/TrackReference.INT.udn
+++ b/Engine/Documentation/Source/Engine/Matinee/TrackReference/TrackReference.INT.udn
@@ -5,27 +5,34 @@ Description:Reference covering all available track types in the Matinee animatio
 
 [TOC(start:2 end:2)]
 
-## Overview
-_Matinee_ provides many built-in **Track Types** that allow for animating different types of data on various types of _Actors_. Through the use of _Tracks_ you can dictate movement of _Actors_ around in a scene, fire off events in _Blueprints_, adjust color or material properties of an _Actor_, as well as several other elements which are covered on this page. 
+
+_Matinee_ provides many built-in **Track Types** that allow for animating different types of data on various types of _Actors_. 
+Through the use of _Tracks_ you can dictate movement of _Actors_ around in a scene, fire off events in _Blueprints_, adjust color or material properties of an _Actor_, as well as several other elements which are covered on this page. 
 
 
 ## Movement Track
-One of the most important _Track_ classes is the **Movement Track** which is used to move an _Actor_ around in the scene over time. A prime example of this would be to generate a path for an _Actor_ to follow. Creating a path is simply a matter of selecting the _Movement Track_ you want to modify and doing the following:
+One of the most important _Track_ classes is the **Movement Track** which is used to move an _Actor_ around in the scene over time. A prime example of this would be to 
+generate a path for an _Actor_ to follow. Creating a path is simply a matter of selecting the _Movement Track_ you want to modify and doing the following:
 
 * Scrub to the time when you want to add the new keyframe.
-* Press the **Enter** key or the **Add Key** button to create a new keyframe at that time.
-* Move the _Actor_ to the desired position for that frame.
-* Repeat this for each keyframe and desired position in order to create your path. 
+* **Press** the **Enter** key or the **Add Key** button to create a new keyframe at that time.
+* **Move** the _Actor_ to the desired position for that frame.
+* **Repeat this for each keyframe** and desired position in order to create your path. 
 
-It is a very good idea to create an initial keyframe for the _Actor_ at its initial position at level load and then not move that keyframe. Otherwise you may see a "pop" when playback begins as the _Actor_ moves to the position proscribed by the animation data.
+It is a very good idea to create an initial keyframe for the _Actor_ at its initial position at level load and then not move that keyframe. 
+Otherwise you may see a "pop" when playback begins as the _Actor_ moves to the position proscribed by the animation data.
 
 [REGION:fullwidth]
 ![move_3dpath.png](move_3dpath.png)
 [/REGION] 
 
-You will see the path of the _Actor_ in the 3D viewport. Paths are drawn in yellow when the _Group_ is selected and are set to the _Editor Group_ color when unselected. You can **Left-click** on the large keys (larger boxes) in the 3D viewport to select a key and allow you to edit it in the same way as you would select them in the main _Matinee_ window. The smaller boxes (or Ticks) indicate 0.1 second intervals, and so give you an idea of the speed of the object.
+You will see the path of the _Actor_ in the 3D viewport. Paths are drawn in yellow when the _Group_ is selected and are set to the _Editor Group_ color when unselected. 
+You can **Left-click** on the large keys (larger boxes) in the 3D viewport to select a key and allow you to edit it in the same way as you would select them 
+in the main _Matinee_ window. The smaller boxes (or Ticks) indicate 0.1 second intervals, and so give you an idea of the speed of the object.
 
-When a keyframe is selected, you will see green handles in the level viewports which allow you to control the velocity at that key and reshape the curve. By default the interpolation mode between keyframes is a curve, but you can change the interpolation mode by **Right-clicking** on a key in the main _Matinee_ window and using the **Interp Mode** menu. 
+When a keyframe is selected, you will see green handles in the level viewports which allow you to control the velocity at that key and reshape the curve. 
+By default the interpolation mode between keyframes is a curve, but you can change the interpolation mode by **Right-clicking** on a key in the main _Matinee_ window 
+and using the **Interp Mode** menu. 
 
 ![interp_mode_menu.png](interp_mode_menu.png)
 
@@ -44,23 +51,23 @@ Below are the _Properties_ that can be adjusted for this _Track_ from the **Deta
 | Lin Curve Tension | Controls the tightness of the curve for the translation path. |
 | Rot Mode | Allows you to choose the method for determining the rotation of the controlled _Actor_. |
 | Show Arrow At Keys | Displays an arrow at each keyframe to allow you to see the orientation at each keyframe. |
-| Show Rotation On Curve Ed | If _false_, when this track is displayed on the **Curve Editor** in _Matinee_, do not show the Rotation tracks. |
-| Show Translation On Curve Ed | If _false_, when this track is displayed on the **Curve Editor** in _Matinee_, do not show the Translation tracks. |
+| Show Rotation On Curve Ed | If _false_, when this track is displayed on the Curve Editor in _Matinee_, do not show the Rotation tracks. |
+| Show Translation On Curve Ed | If _false_, when this track is displayed on the Curve Editor in _Matinee_, do not show the Translation tracks. |
 | Use Quat Interpolation | Use a _Quaternion_ linear interpolation between keys. |
 
 [REGION:tip]
-Refer to [](Engine\Matinee\HowTo\MHT_1) and [](Engine\Matinee\HowTo\MHT_2) for examples of using the **Movement Track** to move Actors in a scene.
+Refer to [](Engine\Matinee\HowTo\MHT_1) and [](Engine\Matinee\HowTo\MHT_2) for examples of using the Movement Track to move Actors in a scene.
 [/REGION]
 
 
 ### Rotation Interpolation
 _Matinee_ uses two different rotation interpolation schemes to rotate _Actors_ between keyframes: **Euler** and **Quaternion**.
 
-* _Euler_ - This is the default and works by interpolating the Yaw, Pitch, and Roll of the _Actor_ between the values at keyframes. This allows you to edit the 3 curves for orientation over time in the **Curve Editor** and control ease-in/ease-out between keyframes. It also supports winding - that is, if you rotate an object multiple times around, the keyframe will store the number of complete revolutions instead of throwing it away. This makes it easy to keyframe something like a screw going into a wall, and also gives you complete control over which direction an _Actor_ takes when moving between two orientations.
+* _Euler_ - This is the default and works by interpolating the Yaw, Pitch, and Roll of the _Actor_ between the values at keyframes. This allows you to edit the 3 curves for orientation over time in the Curve Editor and control ease-in/ease-out between keyframes. It also supports winding - that is, if you rotate an object multiple times around, the keyframe will store the number of complete revolutions instead of throwing it away. This makes it easy to keyframe something like a screw going into a wall, and also gives you complete control over which direction an _Actor_ takes when moving between two orientations.
 
 * _Quaternion_ - This method of interpolation does not support winding, and is essentially linear between keyframes with no support for ease in/out. However, it is more robust and will always find the "shortest" route between two different orientations.
 
-You can control which method is being used by toggling the **Use Quat Interpolation** setting from the **Details** panel.  
+You can control which method is being used by toggling the **Use Quat Interpolation** setting from the Details panel.  
 
 
 ### Rotation Modes
@@ -73,9 +80,11 @@ _Look At Group_ is useful when using a camera, as you can have it always look at
 
 
 ### Movement Relative To Other Groups
-It is possible to have a _Movement Track_ keyframe lookup its transformation info from another _Group_ in the same _Matinee_ sequence.  This means that a keyframe can point to a moving object and have the _Movement Track_ interpolate to that position.  An example usage case of this would be animating a projectile that always hits a character moving in a non-predictable fashion.
+It is possible to have a _Movement Track_ keyframe lookup its transformation info from another _Group_ in the same _Matinee_ sequence.  
+This means that a keyframe can point to a moving object and have the _Movement Track_ interpolate to that position.  
+An example usage case of this would be animating a projectile that always hits a character moving in a non-predictable fashion.
 
-To setup a keyframe to lookup its position from another _Group_, **Right-click** on the keyframe and choose **Select Transform Lookup Group...**.
+To setup a keyframe to lookup its position from another _Group_, **Right-click** on the keyframe and **choose** **Select Transform Lookup Group...**.
 
 ![lookup_track1.png](lookup_track1.png)
 
@@ -83,7 +92,8 @@ This will show a dialog that lets you choose a _Group_ to lookup transform infor
 
 ![lookup_track2.png](lookup_track2.png)
 
-After selecting a _Group_, the keyframe will have the lookup _Group's_ name displayed above it.  The image below displays 4 key frames that are using the _Group's_ **Box1** and **Box2** to lookup their positioning info.
+After selecting a _Group_, the keyframe will have the lookup _Group's_ name displayed above it.  
+The image below displays 4 key frames that are using the _Group's_ **Box1** and **Box2** to lookup their positioning info.
 
 [REGION:fullwidth]
 ![lookup_track3.png](lookup_track3.png)
@@ -95,9 +105,11 @@ To clear a lookup _Group_ for a keyframe, **Right-click** on the keyframe and ch
 
 
 ### Separating Translation and Rotation
-It is possible to "split" a _Movement Track's_ translation and rotation components so they can be keyframed individually on each axis.  This allows finer control over _Movement Track_ curves.  Splitting a _Movement Track_ is an opt-in process.  Initially, all _Movement Tracks_ are combined which means translation and rotation curves have the same number of keyframes and all keyframes exist at the same location in time.  
+It is possible to "split" a _Movement Track's_ translation and rotation components so they can be keyframed individually on each axis.  
+This allows finer control over _Movement Track_ curves.  Splitting a _Movement Track_ is an opt-in process.  
+Initially, all _Movement Tracks_ are combined which means translation and rotation curves have the same number of keyframes and all keyframes exist at the same location in time.  
 
-To split a _Movement Track_, simply select **Split Translation and Rotation** from the context menu when **Right-clicking** on a _Movement Track_.
+To split a _Movement Track_, simply **select** **Split Translation and Rotation** from the context menu when **Right-clicking** on a _Movement Track_.
 
 ![MoveTrackMenu.png](MoveTrackMenu.png)
 
@@ -109,7 +121,8 @@ Once you split a _Movement Track_, it will look like this:
 
 Each axis operates like a unique _Track_ so you can have any number of keyframes in them.
 
-Translation and rotation components are grouped and can be collapsed when not in use.  The vertical bars drawn adjacent to the translation and rotation groups represent a keyframe in one or more of the axes.  Selecting a vertical bar will select all keyframes below it.  This allows you to select multiple keyframes at once.  
+Translation and rotation components are grouped and can be collapsed when not in use.  The vertical bars drawn adjacent to the translation and rotation groups represent 
+a keyframe in one or more of the axes.  Selecting a vertical bar will select all keyframes below it.  This allows you to select multiple keyframes at once.  
 
 [REGION:note]
 Quaternion interpolation mode cannot be used when translation and rotation are separated.
@@ -117,38 +130,54 @@ Quaternion interpolation mode cannot be used when translation and rotation are s
 
 
 ### Normalize Velocity 
-When working with Bezier curves, it is often possible to create situations where objects traveling along a curve speed up or slow down depending on the shape and spacing of keyframes.  The **Normalize Velocity** option modifies a _Movement Track's_ translation curves in order to make objects traveling along it move at constant velocity.  It is important to understand that constant velocity uses the 3D path that your _Movement Track_ curves generate.  When normalize velocity is used, the 2D curves for each axis of translation may change shape but the 3D path that it generates should not change from the original path.  
+When working with Bezier curves, it is often possible to create situations where objects traveling along a curve speed up or slow down depending on the shape and spacing of 
+keyframes.  The **Normalize Velocity** option modifies a _Movement Track's_ translation curves in order to make objects traveling along it move at constant velocity.  
+It is important to understand that constant velocity uses the 3D path that your _Movement Track_ curves generate.  When normalize velocity is used, the 2D curves for each axis 
+of translation may change shape but the 3D path that it generates should not change from the original path.  
 
-_Normalize Velocity_ does not force a curve into constant velocity after the menu option is selected.  It simply does a one time adjustment by examining the original curves and modifying them to maintain constant velocity.  Any subsequent changes to the curve may cause it to no longer have constant velocity. _Normalize Velocity_ is also only available on _Movement Tracks_ which have translation and rotation separated (see _Separating Translation and Rotation_ above).  
+_Normalize Velocity_ does not force a curve into constant velocity after the menu option is selected.  
+It simply does a one time adjustment by examining the original curves and modifying them to maintain constant velocity.  
+Any subsequent changes to the curve may cause it to no longer have constant velocity. _Normalize Velocity_ is also only available on _Movement Tracks_ which 
+have translation and rotation separated (see _Separating Translation and Rotation_ above).  
 
-It is possible to _Normalize Velocity_ on segments of a curve rather than the full interval. When this happens, the _Normalize Velocity_ routine adds a keyframe right before and right after the section being normalized.  This is done so that the parts of the curve not being normalized remain intact.  You may notice "jerky" movements during playback when an object traveling along the curve transitions to and from the _Normalized Velocity_ section.  To fix this, you will have to manually adjust the keyframes before and after the normalized segment.
+It is possible to _Normalize Velocity_ on segments of a curve rather than the full interval. When this happens, the _Normalize Velocity_ routine adds a keyframe 
+right before and right after the section being normalized. This is done so that the parts of the curve not being normalized remain intact. You may notice "jerky" movements 
+during playback when an object traveling along the curve transitions to and from the _Normalized Velocity_ section.  To fix this, you will have to manually adjust the keyframes 
+before and after the normalized segment.
 
 
 
 ## The Director Group
-The **Director Group** is a special type of group that does not work on an arbitrary _Actor_ in the scene like a normal _Actor_, but implicitly works with the "Player". It is used when you want your sequence to be a cinematic, and you want to take control of the players view.
+The **Director Group** is a special type of group that does not work on an arbitrary _Actor_ in the scene like a normal _Actor_, but implicitly works with the "Player". 
+It is used when you want your sequence to be a cinematic, and you want to take control of the players view.
 
-You can add a _Director Group_ to your _Matinee_ sequence in the same way as a regular group, but choose **Add New Director Group** instead. The _Director Group_ header is a different color from a regular group header to make it easier to find. 
+You can add a _Director Group_ to your _Matinee_ sequence in the same way as a regular group, but choose **Add New Director Group** instead. The _Director Group_ header is 
+a different color from a regular group header to make it easier to find. 
 
 The _Director Group_ has a set of unique _Track Types_ that can be added which are outlined below. 
 
 
 ### Director Track
-The **Director Track** allows you to attach the players view to different _Actors_ as the sequence proceeds. When you add a new keyframe to the _Director Track_, a combo box will allow you to choose which _Group_ to assign the view to at that point. 
+The **Director Track** allows you to attach the players view to different _Actors_ as the sequence proceeds. When you add a new keyframe to the _Director Track_, a combo box 
+will allow you to choose which _Group_ to assign the view to at that point. 
 
-Colored bars on the _Director Track_ allow you to see which _Group_ the camera is viewing through at each point in the sequence - the color corresponds to the _Editor Group_ color of each group. 
+Colored bars on the _Director Track_ allow you to see which _Group_ the camera is viewing through at each point in the sequence - the color corresponds to the _Editor Group_ 
+color of each group. 
 
-If you want to return the camera to the player, simply add a cut back to the _Director Group_ itself (think of the _Director Group_ as a special group relating to the "Player"). At the end of the sequence, the view will automatically return to the player so a cut back to the _Director Group_ is not necessary.
+If you want to return the camera to the player, simply add a cut back to the _Director Group_ itself (think of the _Director Group_ as a special group relating to the "Player"). 
+At the end of the sequence, the view will automatically return to the player so a cut back to the _Director Group_ is not necessary.
 
 ![dir_track.png](dir_track.png)
 
 
 #### Shots
-_Director Tracks_ are broken up into named **Shots**, just as you would encounter in standard film-making. Each key denotes a new shot and the shot lasts for the duration of that key.
+_Director Tracks_ are broken up into named **Shots**, just as you would encounter in standard film-making. Each key denotes a new shot and the shot lasts for the duration of 
+that key.
 
 ![shot_label.png](shot_label.png)
 
-By default, shots are numbered according to their position in the timeline. Adding a new key between two existing keys will cause the shot associated with the new key to be numbered between the existing shots.
+By default, shots are numbered according to their position in the timeline. Adding a new key between two existing keys will cause the shot associated with the new key to be 
+numbered between the existing shots.
 
 **Existing keys:**
 [REGION:fullwidth]
@@ -164,7 +193,7 @@ Shots can be renamed through the context menu.
 
 1. The existing shot name:   
     ![shot_label.png](shot_label.png)
-1. Select the key the shot is associated with, **Right-click** on it, and choose **Rename** in the context menu.   
+1. **Select** the key the shot is associated with, **Right-click** on it, and choose **Rename** in the context menu.   
     ![track_menu.png](track_menu.png)
 1. Enter a new number for the shot in the **Set Camera Shot Number** dialog:   
     ![shot_dialog.png](shot_dialog.png)
@@ -173,9 +202,10 @@ Shots can be renamed through the context menu.
 
 
 ### Fade Track
-To control fading of the scene to a particular color, you can use a **Fade Track**. In order to control the amount of fading over time, you must use the [Curve Editor](Engine/UI/CurveEditor). When the Fade amount is 0.0, there is no fading, and when its 1.0, the scene is completely faded out.
+To control fading of the scene to a particular color, you can use a **Fade Track**. In order to control the amount of fading over time, 
+you must use the [Curve Editor](Engine/UI/CurveEditor). When the Fade amount is 0.0, there is no fading, and when its 1.0, the scene is completely faded out.
 
-Below are the _Properties_ that can be adjusted for this _Track_ from the **Details** panel. 
+Below are the _Properties_ that can be adjusted for this _Track_ from the Details panel. 
 
 |Property| Description |
 | ------- | ------------ |
@@ -184,11 +214,13 @@ Below are the _Properties_ that can be adjusted for this _Track_ from the **Deta
 | Persist Fade | Enables and Disables the _Fade Track._ |
 
 ### Slomo Track
-One special effect you can apply to a cinematic is to slow down the game speed by using a **Slomo Track**. This _Track_ works similarly to the _Fade Track_ where you use the [Curve Editor](Engine/UI/CurveEditor) to modify the slow-motion over time. A Slomo factor of 1.0 will use the default game speed, a lower number will reduce game speed and a higher number will speed up the game.
+One special effect you can apply to a cinematic is to slow down the game speed by using a **Slomo Track**. This _Track_ works similarly to the _Fade Track_ where you 
+use the [Curve Editor](Engine/UI/CurveEditor) to modify the slow-motion over time. A Slomo factor of 1.0 will use the default game speed, a lower number will reduce game 
+speed and a higher number will speed up the game.
 
 Slomo affects all game actions, physics, particles, etc. as well as the speed of the sequence itself. 
 
-Below are the _Properties_ that can be adjusted for this _Track_ from the **Details** panel. 
+Below are the _Properties_ that can be adjusted for this _Track_ from the Details panel. 
 
 |Property| Description |
 | ------- | ------------ |
@@ -197,9 +229,10 @@ Below are the _Properties_ that can be adjusted for this _Track_ from the **Deta
 
 
 ### Audio Master Track
-The **Audio Master Track** gives control over the overall volume and pitch of the game's audio. When an _Audio Master Track_ is added to the _Director Group_, it has one default key at the beginning of the sequence whose **InterpMode** is Linear. Adding new keys allows you to modify the volume and/or pitch over the course of the sequence.
+The **Audio Master Track** gives control over the overall volume and pitch of the game's audio. When an _Audio Master Track_ is added to the _Director Group_, it has 
+one default key at the beginning of the sequence whose **InterpMode** is Linear. Adding new keys allows you to modify the volume and/or pitch over the course of the sequence.
 
-Below are the _Properties_ that can be adjusted for this _Track_ from the **Details** panel. 
+Below are the _Properties_ that can be adjusted for this _Track_ from the Details panel. 
 
 |Property| Description |
 | ------- | ------------ |
@@ -210,7 +243,7 @@ Below are the _Properties_ that can be adjusted for this _Track_ from the **Deta
 ### Color Scale Track
 The **Color Scale Track** modifies the color scale of the rendered scene output, allowing you to tint the scene.
 
-Below are the _Properties_ that can be adjusted for this _Track_ from the **Details** panel. 
+Below are the _Properties_ that can be adjusted for this _Track_ from the Details panel. 
 
 |Property| Description |
 | ------- | ------------ |
@@ -219,12 +252,14 @@ Below are the _Properties_ that can be adjusted for this _Track_ from the **Deta
 
 
 ## Property Tracks
-**Property Tracks** consist of several different _Track Types_ and contain a variable inside the _Actor_ the _Track_ is attached to. You can set new variable values for a _Property Track_ by **Right-clicking** on a keyframe and choosing the **Set** action that corresponds to the _Track Type_ selected. 
+**Property Tracks** consist of several different _Track Types_ and contain a variable inside the _Actor_ the _Track_ is attached to. 
+You can set new variable values for a _Property Track_ by **Right-clicking** on a keyframe and choosing the **Set** action that corresponds to the _Track Type_ selected. 
 
 Before a _Property Track _ can be added to _Matinee_ however, it must first be "exposed" so that _Matinee_ can affect it. 
 
 ### Exposing Tracks To Matinee
-Exposing a variable to _Matinee_ can be accomplished by selecting the **Expose to Matinee** checkbox inside the **Details** panel of the _Blueprint_ the variable resides in, as pictured below. 
+Exposing a variable to _Matinee_ can be accomplished by selecting the **Expose to Matinee** checkbox inside the Details panel of the _Blueprint_ the variable resides in, 
+as pictured below. 
 
 ![Expose_ToMatinee.png](Expose_ToMatinee.png)
 
@@ -238,9 +273,10 @@ While most _Tracks_ can be added and worked with instantly, the following _Track
 
 
 ### Bool Property Track
-The **Bool Property Track** provides control over a boolean variable in the _Actor_ connected to the group containing the _Track_. When adding the _Track_, a dialog will appear allowing you to choose from all boolean variables that are connected to the _Actor_ that can be controlled by _Matinee_.
+The **Bool Property Track** provides control over a boolean variable in the _Actor_ connected to the group containing the _Track_. When adding the _Track_, a dialog will 
+appear allowing you to choose from all boolean variables that are connected to the _Actor_ that can be controlled by _Matinee_.
 
-Below are the _Properties_ that can be adjusted for this _Track_ from the **Details** panel. 
+Below are the _Properties_ that can be adjusted for this _Track_ from the Details panel. 
 
 |Property| Description |
 | ------- | ------------ |
@@ -250,7 +286,7 @@ Below are the _Properties_ that can be adjusted for this _Track_ from the **Deta
 ### Float Property Track
 The **Float Property Track** allows you to modify the value of a floating-point property of an _Actor_ over time, such as Brightness for a light, or _DrawScale_ for a _StaticMesh Actor_. When you create the _Track_, it will bring up a combo box showing all modifiable properties of that Group's Actor. Select the one you want and then add/modify keyframes in the same way as you would for a _Movement Track_. You can also use the [CurveEditor](Engine/UI/CurveEditor) to adjust the value of the parameter over time.
 
-Below are the _Properties_ that can be adjusted for this _Track_ from the **Details** panel. 
+Below are the _Properties_ that can be adjusted for this _Track_ from the Details panel. 
 
 |Property| Description |
 | ------- | ------------ |
@@ -260,9 +296,11 @@ Below are the _Properties_ that can be adjusted for this _Track_ from the **Deta
 
 
 ### Vector Property Track
-The **Vector Property Track** allows you to modify the value of a Vector property of an _Actor_ over time, such as _DrawScale3D_ for a _StaticMesh Actor_. When you create the _Track_, it will bring up a combo box showing all modifiable properties of that Group's _Actor_. Select the one you want and then add/modify keyframes in the same way as you would for a _Movement Track_. You can also use the [CurveEditor](Engine/UI/CurveEditor) to adjust the values of the parameter over time.
+The **Vector Property Track** allows you to modify the value of a Vector property of an _Actor_ over time, such as _DrawScale3D_ for a _StaticMesh Actor_. When you create 
+the _Track_, it will bring up a combo box showing all modifiable properties of that Group's _Actor_. Select the one you want and then add/modify keyframes in the same way as 
+you would for a _Movement Track_. You can also use the [CurveEditor](Engine/UI/CurveEditor) to adjust the values of the parameter over time.
 
-Below are the _Properties_ that can be adjusted for this _Track_ from the **Details** panel. 
+Below are the _Properties_ that can be adjusted for this _Track_ from the Details panel. 
 
 |Property| Description |
 | ------- | ------------ |
@@ -276,7 +314,7 @@ The **LinearColor Property Track** is used to modify a linear color variable in
 For example, if you have a _Light Component_ as part of a _Blueprint_, you can assign a linear color variable as the light's color and can then change it using this type of track. 
 The components of a linear color take values in the range 0 to 1 (or higher to overdrive the component).
 
-Below are the _Properties_ that can be adjusted for this _Track_ from the **Details** panel. 
+Below are the _Properties_ that can be adjusted for this _Track_ from the Details panel. 
 
 |Property| Description |
 | ------- | ------------ |
@@ -287,9 +325,10 @@ Below are the _Properties_ that can be adjusted for this _Track_ from the **Deta
 
 ### Color Property Track
 The **Color Property Track** is used to control a byte-based Color variable in the attached _Actor_. 
-For example, unlike the **LinearColor Property Track**, if you were to place a _Light Actor_ directly inside a level, you would then be able to adjust the **Light Color** property via this type of track. 
+For example, unlike the **LinearColor Property Track,** if you were to place a _Light Actor_ directly inside a level, you would then be able to adjust the **Light Color** 
+property via this type of track. 
 
-Below are the _Properties_ that can be adjusted for this _Track_ from the **Details** panel. 
+Below are the _Properties_ that can be adjusted for this _Track_ from the Details panel. 
 
 |Property| Description |
 | ------- | ------------ |
@@ -304,9 +343,10 @@ There are two **Material Parameter Tracks** that can be driven by _Matinee_: **F
 
 
 ### Float Material Parameter Track
-The **Float Material Parameter Track** is used to modify Scalar Parameters within materials during runtime. The _Track_ takes a list of materials and modifies the value of the given parameter in those materials.
+The **Float Material Parameter Track** is used to modify Scalar Parameters within materials during runtime. The _Track_ takes a list of materials and modifies the value 
+of the given parameter in those materials.
 
-Below are the _Properties_ that can be adjusted for this _Track_ from the **Details** panel. 
+Below are the _Properties_ that can be adjusted for this _Track_ from the Details panel. 
 
 |Property| Description |
 | ------- | ------------ |
@@ -316,9 +356,10 @@ Below are the _Properties_ that can be adjusted for this _Track_ from the **Deta
 | Param Name | Name of the parameter in the materials to modify. |
 
 ### Vector Material Parameter Track
-The **Vector Material Parameter Track** is used to modify Vector Parameters within materials during runtime. The _Track_ takes a list of materials and modifies the value of the given parameter in those materials.
+The **Vector Material Parameter Track** is used to modify Vector Parameters within materials during runtime. The _Track_ takes a list of materials and modifies the value of the 
+given parameter in those materials.
 
-Below are the _Properties_ that can be adjusted for this _Track_ from the **Details** panel. 
+Below are the _Properties_ that can be adjusted for this _Track_ from the Details panel. 
 
 |Property| Description |
 | ------- | ------------ |
@@ -332,9 +373,12 @@ Below are the _Properties_ that can be adjusted for this _Track_ from the **Deta
 There are two **Particle Tracks** that can be driven by _Matinee_: **Float Particle Parameter Track** and **Particle Replay Track**. 
 
 ### Float Particle Parameter Track
-The **Float Particle Parameter Track** is used to modify parameter distributions within particle systems during runtime. Any parameters with the given parameter name in the particle system connected to the group containing the _Track_ will be modified. This means you could have a fire particle system with emitters for the flames, sparks, and smoke where each emitter has a parameter controlling the spawn rate named _SpawnRate_. A single _Float Particle Param Track_ could drive all of these parameters to control the fire.
+The **Float Particle Parameter Track** is used to modify parameter distributions within particle systems during runtime. Any parameters with the given parameter name in 
+the particle system connected to the group containing the _Track_ will be modified. This means you could have a fire particle system with emitters for the flames, sparks, 
+and smoke where each emitter has a parameter controlling the spawn rate named _SpawnRate_. A single _Float Particle Param Track_ could drive all of these parameters to control 
+the fire.
 
-Below are the _Properties_ that can be adjusted for this _Track_ from the **Details** panel. 
+Below are the _Properties_ that can be adjusted for this _Track_ from the Details panel. 
 
 |Property| Description |
 | ------- | ------------ |
@@ -344,16 +388,21 @@ Below are the _Properties_ that can be adjusted for this _Track_ from the **Deta
 
 
 ### Particle Replay Track
-The **Particle Replay track** allows you to create and play back pre-recorded particle simulations. This is especially useful when you want the particle to behave the same way in all shots. When you add a new key to the _Particle Replay Track_, you are essentially adding a new clip to record to. For each clip, a bar is displayed on the timeline showing the duration of the clip.
+The **Particle Replay track** allows you to create and play back pre-recorded particle simulations. This is especially useful when you want the particle to behave the same 
+way in all shots. When you add a new key to the _Particle Replay Track_, you are essentially adding a new clip to record to. For each clip, a bar is displayed on the timeline 
+showing the duration of the clip.
 
 ![replay_bar.png](replay_bar.png)
 
-When selected, you have the ability to adjust the starting point and duration of the clip using the handles on either end of the clip's bar in the timeline. **Left-clicking** and dragging the mouse on either handle will adjust the duration of the clip by moving that handle. To move the clip without adjusting its duration, **Ctrl+Click** and drag on the key itself.
+When selected, you have the ability to adjust the starting point and duration of the clip using the handles on either end of the clip's bar in the timeline. **Left-clicking** 
+and dragging the mouse on either handle will adjust the duration of the clip by moving that handle. To move the clip without adjusting its duration, **Ctrl+Click** and **drag** on 
+the key itself.
 [REGION:fullwidth]
 ![replay_handles.png](replay_handles.png)
 [/REGION] 
 
-You can also adjust the starting position and duration of the clip as well as change the ID of the clip by **Right-clicking** on the clip's bar in the timeline and choosing the appropriate option from the menu.
+You can also adjust the starting position and duration of the clip as well as change the ID of the clip by **Right-clicking** on the clip's bar in the timeline and choosing 
+the appropriate option from the menu.
 
 ![replay_menu.png](replay_menu.png)
 
@@ -365,17 +414,21 @@ When recording is active, the clip bars will be highlighted in the timeline.
 
 ![replay_start_bar.png](replay_start_bar.png)
 
-Press the ![replay_play.png](replay_play.png) button in the toolbar to begin playback. The track will record the particles during play back when recording is enabled. Once recording is finished, **Right-click** the track again and choose **Stop Recording Particles**.
+**Press** the ![replay_play.png](replay_play.png) button in the toolbar to begin playback. The track will record the particles during play back when recording is enabled. 
+Once recording is finished, **Right-click** the track again and choose **Stop Recording Particles**.
 
 ![replay_stop.png](replay_stop.png)
 
-Now, when you press the ![replay_play.png](replay_play.png) button, the particle system(s) should play back from the recorded clips and you can scrub through the timeline to move forward and backward through the clips to inspect the recording.
+Now, when you press the ![replay_play.png](replay_play.png) button, the particle system(s) should play back from the recorded clips and you can scrub through the timeline to 
+move forward and backward through the clips to inspect the recording.
 
 
 ## Animation Track
-The **Animation Track** allows you to animate the properties of _Actors_ over time, to create either dynamic gameplay or cinematic in-game sequences. Through the use of this _Track Type_, you can place keyframes to control when and how they play specified animations. 
+The **Animation Track** allows you to animate the properties of _Actors_ over time, to create either dynamic gameplay or cinematic in-game sequences. 
+Through the use of this _Track Type_, you can place keyframes to control when and how they play specified animations. 
 
-To add an animation to the _Track_, scrub through the _Track_ and press the **Enter** key or **Add Key** Button. When a key has been added, you will be prompted to select an animation to add to the track through a pop-up window as shown below. 
+To add an animation to the _Track_, scrub through the _Track_ and **press** the **Enter** key or **Add Key** Button. When a key has been added, you will be prompted to select 
+an animation to add to the track through a pop-up window as shown below. 
 
 [REGION:fullwidth]
 ![animtrack_animation_select.png](animtrack_animation_select.png)
@@ -383,7 +436,8 @@ To add an animation to the _Track_, scrub through the _Track_ and press the **En
 
 Once an animation has been added to the _Track_, it will be displayed in Blue. You can manipulate the animation by **Left-clicking** the key to select it. 
 
-With the animation selected, you can adjust the length by moving the cursor to the start or end point until the double-arrows are shown as seen below and dragging in either direction. 
+With the animation selected, you can adjust the length by moving the cursor to the start or end point until the double-arrows are shown as seen below and dragging in either 
+direction. 
 
 [REGION:fullwidth]
 ![animtrack_adjust_time.png](animtrack_adjust_time.png)
@@ -395,9 +449,11 @@ You can also move the position of the animation on the track by **Control+Left-C
 ![animtrack_moving.png](animtrack_moving.png)
 [/REGION] 
 
-The numbers shown on the animation illustrate the (Start->End) times in which the animation will play. The number following the Start/End times (x0.89 for example), indicate the playback rate of the animation where 1.0 is the default and higher numbers increase the speed of the playback while lower numbers decrease the speed. To adjust playback rate, hold **Control** when adjusting the length of the animation. 
+The numbers shown on the animation illustrate the (Start->End) times in which the animation will play. The number following the Start/End times (x0.89 for example), indicate the 
+playback rate of the animation where 1.0 is the default and higher numbers increase the speed of the playback while lower numbers decrease the speed. To adjust playback rate, 
+hold **Control** when adjusting the length of the animation. 
 
-Below are the _Properties_ that can be adjusted for this _Track_ from the **Details** panel. 
+Below are the _Properties_ that can be adjusted for this _Track_ from the Details panel. 
 
 |Property| Description |
 | ------- | ------------ |
@@ -408,7 +464,10 @@ Below are the _Properties_ that can be adjusted for this _Track_ from the **Deta
 
 
 ## Event Track
-The **Event Track** gives _Matinee_ a mechanism for firing events back into _Blueprints_ at particular points in the sequence. When you place a keyframe on an _Event Track_, you will be asked for an **Event Name**. Each unique _Event Name_ will create a new output connector on the _Matinee Actor_. When the sequence reaches the key, any _Blueprint_ nodes that are connected to the output will be activated. If you have multiple keys with the same name, one output connector will be created, but it will be activated when each of the keys are passed.
+The Event Track gives _Matinee_ a mechanism for firing events back into _Blueprints_ at particular points in the sequence. When you place a keyframe on an _Event Track_, you 
+will be asked for an **Event Name**. Each unique _Event Name_ will create a new output connector on the _Matinee Actor_. When the sequence reaches the key, any _Blueprint_ nodes 
+that are connected to the output will be activated. If you have multiple keys with the same name, one output connector will be created, but it will be activated when each of the 
+keys are passed.
 
 Here is an _Event Track_ with one key called "ToggleLight".
 
@@ -418,11 +477,15 @@ And this is the _Matinee Actor_ in a _Blueprint_, showing the new output connect
 
 ![event_connector.png](event_connector.png)
 
-The ![eventarrows.jpg](eventarrows.jpg) buttons on the _Event Track_ header allow you to control if the output should be fired when the sequence is playing forwards, backwards, or both. If the right arrow button is yellow, events on that track will fire as they are passed when the sequence is playing forwards. If the left arrow button is yellow, events are fired when the sequence is played backwards.
+The ![eventarrows.jpg](eventarrows.jpg) buttons on the _Event Track_ header allow you to control if the output should be fired when the sequence is playing forwards, backwards, 
+or both. If the right arrow button is yellow, events on that track will fire as they are passed when the sequence is playing forwards. If the left arrow button is yellow, events 
+are fired when the sequence is played backwards.
 
-Regardless of which group the _Event Track_ is in and whether or not an _Actor_ is attached to the group; the _Matinee Actor_ will create an output connector for each key placed on its _Track_. You can also **Right-click** on a key to rename it and doing so will rename any other keys with the same name, as well as the output connector on the _Matinee Actor_.
+Regardless of which group the _Event Track_ is in and whether or not an _Actor_ is attached to the group; the _Matinee Actor_ will create an output connector for each key placed 
+on its _Track_. You can also **Right-click** on a key to rename it and doing so will rename any other keys with the same name, as well as the output connector on 
+the _Matinee Actor_.
 
-Below are the _Properties_ that can be adjusted for this _Track_ from the **Details** panel. 
+Below are the _Properties_ that can be adjusted for this _Track_ from the Details panel. 
 
 |Property| Description |
 | ------- | ------------ |
@@ -433,9 +496,9 @@ Below are the _Properties_ that can be adjusted for this _Track_ from the **Deta
 
 
 ## Sound Track
-The **Sound Track** is used to play audio from _SoundCues_ during a _Matinee_ sequence. When placing a new key, the _SoundCue_ currently selected in the **Content Browser** will be used. Keys for the _Sound Track_ appear similar to those of the _AnimControl Track_. You will see a bar with the name of the _SoundCue_ that stretches for the duration of the _SoundCue_.
+The **Sound Track** is used to play audio from _SoundCues_ during a _Matinee_ sequence. When placing a new key, the _SoundCue_ currently selected in the Content Browser will be used. Keys for the _Sound Track_ appear similar to those of the _AnimControl Track_. You will see a bar with the name of the _SoundCue_ that stretches for the duration of the _SoundCue_.
 
-Below are the _Properties_ that can be adjusted for this _Track_ from the **Details** panel. 
+Below are the _Properties_ that can be adjusted for this _Track_ from the Details panel. 
 
 |Property| Description |
 | ------- | ------------ |
@@ -447,7 +510,8 @@ Below are the _Properties_ that can be adjusted for this _Track_ from the **Deta
 
 
 ## Toggle Track
-The **Toggle Track** gives _Matinee_ a mechanism for toggling _Emitter Actors_ on or off at particular points in the sequence. When you place a keyframe on a _Toggle Track_, you will be presented with the following dialog:
+The **Toggle Track** gives _Matinee_ a mechanism for toggling _Emitter Actors_ on or off at particular points in the sequence. When you place a keyframe on a _Toggle Track_, 
+you will be presented with the following dialog:
 
 ![PSysToggle_Key.png](PSysToggle_Key.png)
 
@@ -456,16 +520,17 @@ The options are:
 * **On** - Turns the particle system on.
 * **Off** - Turns the particle system off.
 
-The _Track_ display will show a green bar when an emitter is toggled on to make it easy to identify the state of the _Actor_. This can be seen below for a _Toggle Track_ that turns the emitter on and back off again.
+The _Track_ display will show a green bar when an emitter is toggled on to make it easy to identify the state of the _Actor_. This can be seen below for a _Toggle Track_ that 
+turns the emitter on and back off again.
 
 ![PSysToggle_Track.png](PSysToggle_Track.png)
 
-Below are the _Properties_ that can be adjusted for this _Track_ from the **Details** panel. 
+Below are the _Properties_ that can be adjusted for this _Track_ from the Details panel. 
 
 |Property| Description |
 | ------- | ------------ |
 | Activate System Each Update | If _true_, the track will call _ActivateSystem_ on the emitter each update (old 'incorrect' behavior). If _false_, the _ParticleSystem_ will only be activated if it was previously inactive. |
-| Activate with Just Attached Flag | If _true_, the track will activate the system with the **Just Attached** flag.  |
+| Activate with Just Attached Flag | If _true_, the track will activate the system with the Just Attached flag.  |
 | Fire Events When Backwards | If _true_, events for this track will be fired when the _Matinee_ is playing in reverse. |
 | Fire Events When Forwards | If _true_, events for this track will be fired when the _Matinee_ is playing forwards. |
 | Fire Events When Jumping Forward | If _true_, allows all of the events in the track to be fired even if the _Matinee_ is skipped. |
@@ -473,11 +538,13 @@ Below are the _Properties_ that can be adjusted for this _Track_ from the **Deta
 
 
 ## Visibility Track
-The **Visibility Track** is used to control whether an _Actor_ is visible or not (by way of the **Actor Hidden In Game** property). When adding a new key for this _Track_, you are presented with a dialog with options to **Show**, **Hide**, and **Toggle** the associated _Actor_.
+The **Visibility Track** is used to control whether an _Actor_ is visible or not (by way of the **Actor Hidden In Game** property). When adding a new key for this _Track_, you 
+are presented with a dialog with options to **Show**, **Hide**, and **Toggle** the associated _Actor_.
 
 ![visibilitymenu.png](visibilitymenu.png)
 
-The _Track_ display will show a green bar when the _Actor_ is set to visible and later set not visible again to make it easy to identify the state of the _Actor_. This can be seen below for a _Visibility Track_ that shows the _Actor_ for 1 second then hides it.
+The _Track_ display will show a green bar when the _Actor_ is set to visible and later set not visible again to make it easy to identify the state of the _Actor_. This can 
+be seen below for a _Visibility Track_ that shows the _Actor_ for 1 second then hides it.
 
 ![visibilitytrack.png](visibilitytrack.png)
 
@@ -485,7 +552,7 @@ Keys set to toggle the visibility of the _Actor_ are displayed with a small red
 
 ![visibilitytoggle.png](visibilitytoggle.png)
 
-Below are the _Properties_ that can be adjusted for this _Track_ from the **Details** panel. 
+Below are the _Properties_ that can be adjusted for this _Track_ from the Details panel. 
 
 |Property| Description |
 | ------- | ------------ |
