#UE4doc #docs: Removed extraneous bolding as according to the instructional guide. Changed title of How To to match instructional guide. Altered images to match the instructional guide: green -> yellow.

diff --git a/Engine/Documentation/Source/Gameplay/HowTo/SpawnAndDestroyActors/Blueprints/GHT_1B.INT.udn b/Engine/Documentation/Source/Gameplay/HowTo/SpawnAndDestroyActors/Blueprints/GHT_1B.INT.udn
--- a/Engine/Documentation/Source/Gameplay/HowTo/SpawnAndDestroyActors/Blueprints/GHT_1B.INT.udn
+++ b/Engine/Documentation/Source/Gameplay/HowTo/SpawnAndDestroyActors/Blueprints/GHT_1B.INT.udn
@@ -1,13 +1,11 @@
 Availability: Public
-Title:How to Spawn/Destroy an Actor in Blueprints
+Title:Spawning/Destroying an Actor in Blueprints
 Crumbs: %ROOT%, Gameplay, Gameplay/HowTo/SpawnAndDestroyActors
 Description:A How To Guide for Spawning/Destroying Actors through Blueprints in Unreal Engine 4. 
 
 [TOC (start:2)]
 
-## Overview
-
- Here you will learn how to spawn an Actor in **Blueprints** and define its spawn location. You will also learn how to destroy spawned Actors. 
+ Here you will learn how to spawn an Actor in Blueprints and define its spawn location. You will also learn how to destroy spawned Actors. 
 
 ## Spawning an Actor
 
@@ -27,7 +25,7 @@ Below you will find steps to spawn an Actor in the level by pressing a key.
 
 	_Click image for full view._
 
-	In order to spawn an Actor, we will need to use the **Spawn Actor from Class** node. In this example, we are using the **Level Blueprint**, however you can use the **Spawn Actor from Class** node in most **Blueprints** where it can be found in the **Palette** window as shown below.
+	In order to spawn an Actor, we will need to use the Spawn Actor from Class node. In this example, we are using the Level Blueprint, however you can use the Spawn Actor from Class node in most Blueprints where it can be found in the Palette window as shown below.
 
 1.  **Right-click** anywhere in the graph and in the context menu, search for and add the **Spawn Actor from Class** node. 
 
@@ -37,16 +35,16 @@ Below you will find steps to spawn an Actor in the level by pressing a key.
 
 	![](GHT1_4.png)
 
-	The **Spawn Actor from Class** node takes an **Actor Class** (specified via the **Class** input) and attempts to spawn an instance of that class in the world. The **Spawn Transform** input is used to define the Actor's spawn location in the world (and its initial orientation). If that location is blocked by a collider, then the Actor will fail to spawn unless the **Spawn Even if Colliding** input is set to **True**. 
+	The Spawn Actor from Class node takes an Actor Class (specified via the Class input) and attempts to spawn an instance of that class in the world. The Spawn Transform input is used to define the Actor's spawn location in the world (and its initial orientation). If that location is blocked by a collider, then the Actor will fail to spawn unless the Spawn Even if Colliding input is set to True. 
 
-	Each of the pins on the **Spawn Actor** node are outlined below.
+	Each of the pins on the Spawn Actor node are outlined below.
 
 	![](GHT1_6.png)
 
 	| Number | Description |
 	| ------- | ------------ |
 	| **Execution Input (1)** | This is an execution input and is what triggers the node to spawn an Actor. |
-	| **Class (2)** | This is the **Actor Class** you want to spawn an instance of in the world (this is required). |
+	| **Class (2)** | This is the Actor Class you want to spawn an instance of in the world (this is required). |
 	| **Spawn Transform (3)** | This takes in a Transform used to position and orientate the Actor in the world. |
 	| **Spawn Even If Colliding (4)** | Determines if the Actor should be spawned when something is blocking the spawn location. |
 	| **Execution Output (5)** | This is an execution output and leads to script that should occur following the Spawn Actor. |
@@ -60,7 +58,7 @@ Below you will find steps to spawn an Actor in the level by pressing a key.
 
 	![](GHT1_8.png)
 
-	The **Instigator** is optional and is the Pawn responsible for damage done by the spawned Actor.
+	The Instigator is optional and is the Pawn responsible for damage done by the spawned Actor.
 
 1.  **Left-click** and drag off the **Spawn Transform** pin to the left, then in the context menu search for and add **Get Actor Transform**. 
 
@@ -104,7 +102,7 @@ Below you will find steps to spawn an Actor in the level by pressing a key.
 
 	Whenever you press **F**, the Actor you selected should spawn where the player character is in the world.
 
-This example illustrates spawning an Actor through a key press, however you could use any number of **Events** to execute the **SpawnActor** node. 
+This example illustrates spawning an Actor through a key press, however you could use any number of Events to execute the SpawnActor node. 
 
 Refer to the [](Engine/Blueprints/UserGuide) for more on working with **Blueprints**. 
 
@@ -154,7 +152,7 @@ The examples below will show you how to manually set a Transform as well as refe
 
 1.	**Compile** then close the **Blueprint** and play in the editor, you should see your Actor spawn at the Target Point. 
 
-In addition to specifying a spawn point by referencing an Actor, you can manually enter the **Spawn Transform** for the **SpawnActor** as well, to do this: 
+In addition to specifying a spawn point by referencing an Actor, you can manually enter the Spawn Transform for the SpawnActor as well, to do this: 
 
 1.	Drag off the **Spawn Transform** pin of a **SpawnActor** node, then select **Make Transform**. 
 
@@ -168,7 +166,7 @@ In addition to specifying a spawn point by referencing an Actor, you can manuall
 
 Once you have spawned an Actor, suppose you wanted to destroy it (for example, maybe the player killed the Actor or picked the Actor up as a pick-up, etc.). 
 
-Below are the steps you can take to destroy a Spawned Actor in the **Level Blueprint**. 
+Below are the steps you can take to destroy a Spawned Actor in the Level Blueprint. 
 
 1.  Open up any project, once inside the editor, click the **Blueprints** button and **Open Level Blueprint**.  
 
@@ -206,9 +204,9 @@ Below are the steps you can take to destroy a Spawned Actor in the **Level Bluep
 
 1.	**Compile** then close the **Blueprint** and play in the editor, you should see the Actor spawn at the player location then despawn after 3 seconds. 
 
-As with executing the **SpawnActor** node to spawn an Actor, destroying the Actor can be triggered by several different types of **Events** based on your needs. 
+As with executing the SpawnActor node to spawn an Actor, destroying the Actor can be triggered by several different types of Events based on your needs. 
 
-Refer to the [](Engine/Blueprints/UserGuide) for more on working with **Blueprints**. 
+Refer to the [](Engine/Blueprints/UserGuide) for more on working with Blueprints. 
 
 
 
