#UEDoc #docs - fixed bold lettering

diff --git a/Engine/Documentation/Source/Engine/Matinee/UI/MatineeInterface.INT.udn b/Engine/Documentation/Source/Engine/Matinee/UI/MatineeInterface.INT.udn
--- a/Engine/Documentation/Source/Engine/Matinee/UI/MatineeInterface.INT.udn
+++ b/Engine/Documentation/Source/Engine/Matinee/UI/MatineeInterface.INT.udn
@@ -6,7 +6,6 @@ Description:Description of the user interface for Matinee.
 [TOC(start:2 end:2)]
 
 
-## Overview
 
 The following screenshot shows the Matinee Editor:
 
@@ -192,7 +191,7 @@ The timeline is a graphical representation of all the keyframes for all the trac
 
 ![Properties Pane](Matinee_PropertiesPane.png)
 
-The **Details** panel is a standard Unreal Details window which displays the details of the currently selected folder, group, or track in the Timeline Pane's track list.
+The Details panel is a standard Unreal Details window which displays the details of the currently selected folder, group, or track in the Timeline Pane's track list.
 
 
 ## Controls
