#UE4doc #docs: Adjusted bold text in accordance with the Instructional Guide

diff --git a/Engine/Documentation/Source/Engine/UI/ProjectSettings/ProjectSettings.INT.udn b/Engine/Documentation/Source/Engine/UI/ProjectSettings/ProjectSettings.INT.udn
--- a/Engine/Documentation/Source/Engine/UI/ProjectSettings/ProjectSettings.INT.udn
+++ b/Engine/Documentation/Source/Engine/UI/ProjectSettings/ProjectSettings.INT.udn
@@ -54,9 +54,7 @@ The **Project Settings** editor is divided up into various categories and sectio
 are displayed as headings, while each of these sections is displayed as a hyperlink that opens the options for that 
 section in the editor.
 
-While viewing a particular section, you can navigate back to the main view by using the **< Back** link at the top of the left sidebar:
 
-![](back_nav.png)
 
 ### Controls
 
diff --git a/Engine/Documentation/Source/Engine/UI/PropertyMatrix/PropertyMatrix.INT.udn b/Engine/Documentation/Source/Engine/UI/PropertyMatrix/PropertyMatrix.INT.udn
--- a/Engine/Documentation/Source/Engine/UI/PropertyMatrix/PropertyMatrix.INT.udn
+++ b/Engine/Documentation/Source/Engine/UI/PropertyMatrix/PropertyMatrix.INT.udn
@@ -49,7 +49,7 @@ Description:Tool for viewing and editing multiple properties of multiple actors
 [EXCERPT:Intro]
 The **Property Matrix** allows for easy bulk editing and value comparison for large numbers of Objects 
 or Actors. It displays a configurable set of properties for a collection of objects as columns in a 
-table view that can be sorted on any column. The **Property Matrix** also provides a standard 
+table view that can be sorted on any column. The Property Matrix also provides a standard 
 property editor that displays all properties for the current selection set in the table view.
 [/EXCERPT:Intro]
 
@@ -64,23 +64,23 @@ property editor that displays all properties for the current selection set in th
 <!--
 [EXCERPT:BulkEditing]
 
-* Easier workflow **for setting a series of varying values to a bunch of objects**, without sacrificing 
+* Easier workflow for setting a series of varying values to a bunch of objects, without sacrificing 
 the ability to set the property on a bunch of objects to the same value.
-* **Works with THOUSANDS of objects** at the same time.
+* Works with THOUSANDS of objects at the same time.
 * Can handle editing a wide variety of object types at the same time.
 
 [/EXCERPT:BulkEditing]
 [EXCERPT:BulkComparison]
 
-* **Sort the values on thousands of objects** at a time.
-* **Quickly find assets and Actors with incorrect settings**.
+* Sort the values on thousands of objects at a time.
+* Quickly find assets and Actors with incorrect settings.
 
 [/EXCERPT:BulkComparison]
 [EXCERPT:PropertySupport]
 
-* Do all of the above with the properties of **Arrays and Struct types**.
-* Can expose columns for **ANY PROPERTY**.
-* Even handles **ARRAY INDICES**.
+* Do all of the above with the properties of Arrays and Struct types.
+* Can expose columns for ANY PROPERTY.
+* Even handles ARRAY INDICES.
 
 [/EXCERPT:PropertySupport]
 -->
diff --git a/Engine/Documentation/Source/Engine/UI/SourceControl/SourceControl.INT.udn b/Engine/Documentation/Source/Engine/UI/SourceControl/SourceControl.INT.udn
--- a/Engine/Documentation/Source/Engine/UI/SourceControl/SourceControl.INT.udn
+++ b/Engine/Documentation/Source/Engine/UI/SourceControl/SourceControl.INT.udn
@@ -19,7 +19,7 @@ Description:This document overviews the Source Control features.
 
 [TOC(start:2)]
 
-## Overview
+
 
 Unreal Editor supports source control, allowing teams to coordinate their game development and design efforts, as well as handling version control.
 
