#UEDoc #docs - fixed bold lettering

diff --git a/Engine/Documentation/Source/Engine/Matinee/UserGuide/UserGuide.INT.udn b/Engine/Documentation/Source/Engine/Matinee/UserGuide/UserGuide.INT.udn
--- a/Engine/Documentation/Source/Engine/Matinee/UserGuide/UserGuide.INT.udn
+++ b/Engine/Documentation/Source/Engine/Matinee/UserGuide/UserGuide.INT.udn
@@ -37,13 +37,16 @@ To create a Matinee sequence, you create a Matinee Actor and place it in your sc
 
 	![](MatineeToolbar.png)
 
-* In the **Place Mode** panel, find **Matinee** under **All Classes** and drag it into the viewport.
+* In the **Place Mode** panel, find Matinee under **All Classes** and drag it into the viewport.
 
 
-After you have created your Matinee Actor, to view the Matinee interface and modify your cinematic, select the Matinee Actor either in the **World Outliner** or in the viewport, and then click **Open Matinee** in the **Details** panel. Alternatively, you can click the **Matinee** button in the toolbar at the top of the Editor, and then select your Matinee Actor from the list.
+After you have created your Matinee Actor, to view the Matinee interface and modify your cinematic, **select** the Matinee Actor either in the **World Outliner** 
+or in the viewport, and then click **Open Matinee** in the **Details** panel. Alternatively, you can **click** the Matinee button in the toolbar at the top of the Editor, 
+and then **select** your Matinee Actor from the list.
 
 [REGION:note]
-If you want to set key bindings for keyboard shortcuts in Matinee, you can do so by clicking **Editor Preferences** on the **Edit** menu, and then clicking **Keyboard Shortcuts.** Scroll down to and expand the **Matinee** section to view and set key bindings.
+If you want to set key bindings for keyboard shortcuts in Matinee, you can do so by clicking **Editor Preferences** on the **Edit** menu, and then clicking **Keyboard Shortcuts.** 
+Scroll down to and expand the Matinee section to view and set key bindings.
 [/REGION]
 
 
@@ -51,24 +54,33 @@ If you want to set key bindings for keyboard shortcuts in Matinee, you can do so
 
 Matinee is structured around the idea of a set of Groups, each containing a number of different Tracks. A Group is associated with a particular Actor in the level.
 
-For more information on the different types to Tracks and Track Usage, refer to the [](Engine\Matinee\TrackReference) documentation. 
+For more information on the different types to Tracks and Track Usage, refer to the **[](Engine\Matinee\TrackReference)** documentation. 
 
 ### Creating A New Group
 
-To create a new Group in your Matinee sequence, first select the Actor in the level that you wish to modify. Then **right-click** in the gray bar down the left of Matinee and select 'Add New Group'. 
+To create a new Group in your Matinee sequence, first **select** the Actor in the level that you wish to modify. Then **right-click** in the gray bar down the left of Matinee 
+and **select** 'Add New Group'. 
 
 ![New Group menu](Matinee_NewGroup.png)
 
-You will be prompted for a name for the new Group - it should be unique within that Matinee sequence, and should not contain spaces. When you choose OK, you should then see a new Group header. If you now look at your Matinee Action in the Blueprint, you can see it has created a new variable connector with the name of your Group, and a new Object variable containing a reference to the Actor you wanted that Group to work on. It is possible to connect multiple Actors to the same Group connector, which is very useful if you want to control the brightness of a number of lights in the same way over time for example.
+You will be prompted for a name for the new Group - it should be unique within that Matinee sequence, and should not contain spaces. When you choose OK, you should then see a new Group 
+header. If you now look at your Matinee Action in the Blueprint, you can see it has created a new variable connector with the name of your Group, and a new Object variable containing a 
+reference to the Actor you wanted that Group to work on. It is possible to connect multiple Actors to the same Group connector, which is very useful if you want to control the brightness 
+of a number of lights in the same way over time for example.
 
-The little colored bar in the Group header is the editor 'Group color'. This is used for various things in Matinee to help you identify which objects in the scene are controlled by which Group in Matinee. It can be changed by selecting the Group and adjusting the GroupColor properties at the bottom of the window. **Right-clicking** on the Group title allows you to Rename or Delete the entire Group. Once you have a Group selected, pressing the Up and Down arrow keys allows you to move it around in the list, which can be useful for keeping similar Groups together.
+The little colored bar in the Group header is the editor 'Group color'. This is used for various things in Matinee to help you identify which objects in the scene are controlled by 
+which Group in Matinee. It can be changed by selecting the Group and adjusting the GroupColor properties at the bottom of the window. **Right-clicking** on the Group title allows you to 
+Rename or Delete the entire Group. Once you have a Group selected, pressing the Up and Down arrow keys allows you to move it around in the list, which can be useful for keeping similar 
+Groups together.
 
 Note that whenever you select a Group (or Track in a Group), the Actor in the level will be selected for you, and vice versa.
 
 
 ### Creating a New Track
 
-A Group on its own does nothing to the Actor it is working on. To actually modify it in some way, we must add a Track to that Group. **Right-click** on the Group header and it will bring up the New Track menu. This lists all the different types of Matinee Tracks that can be added. Select the track you want and it will appear within that Group. Depending on the Track class, you may be prompted for more information before it is added.
+A Group on its own does nothing to the Actor it is working on. To actually modify it in some way, we must add a Track to that Group. **Right-click** on the Group header and it 
+will bring up the New Track menu. This lists all the different types of Matinee Tracks that can be added. Select the track you want and it will appear within that Group. Depending on 
+the Track class, you may be prompted for more information before it is added.
 
 ![New Track menu](Matinee_NewTrack.png)
 
@@ -76,9 +88,11 @@ Some Track types (such as the Movement Track) only allow you to have one per Gro
 
 When you select an Actor in the viewport that has a track in Matinee, and the Matinee Editor is open, the Actor's track becomes highlighted in the Matinee Editor.
 
-If you want to hide all the Tracks in a Group, you can press the arrow on the Group header to 'fold' it up. You can change the order of Tracks within a Group in the same way as you move Groups - just select it and use the Up and Down arrow keys.
+If you want to hide all the Tracks in a Group, you can press the arrow on the Group header to 'fold' it up. You can change the order of Tracks within a Group in the same way as you 
+move Groups - just select it and use the Up and Down arrow keys.
 
-The effects of tracks can be toggled on and off by using the checkmark box next to the name of the track.  For example, in the following image, the movement track is disabled and the effects of the track will not be seen when the matinee sequence is run.
+The effects of tracks can be toggled on and off by using the checkmark box next to the name of the track.  For example, in the following image, the movement track is disabled and the 
+effects of the track will not be seen when the matinee sequence is run.
 
 ![Track view with Movement track disabled](Matinee_TrackDisabled.png)
 
@@ -87,7 +101,8 @@ The effects of tracks can be toggled on and off by using the checkmark box next
 
 The selection and manipulation controls for Matinee are similar to tools in other parts of UnrealEngine 4.
 
-Clicking on a keyframe will move the playback position to that keyframe and put Matinee into 'record' mode. Modifying the property of the Actor controlled by that track will then modify the keyframe value. When editing a keyframe, a little red circle will appear in the Matinee info box:
+Clicking on a keyframe will move the playback position to that keyframe and put Matinee into 'record' mode. Modifying the property of the Actor controlled by that track will then modify 
+the keyframe value. When editing a keyframe, a little red circle will appear in the Matinee info box:
 
 ![Key recording indicator](Matinee_KeyMovement.png)
 
@@ -96,25 +111,35 @@ Scrubbing to a different position or beginning playback will take you out of 're
 
 ### Adjusting Sequence Length
 
-Changing the length of your sequence is very straightforward - simply grab the red handle that marks the end of your sequence and drag it to the new time you with the sequence to end. It is completely safe to have keyframes outside the range of the sequence.
+Changing the length of your sequence is very straightforward - simply grab the red handle that marks the end of your sequence and drag it to the new time you with the sequence to end. 
+It is completely safe to have keyframes outside the range of the sequence.
 
-You can also add time at a particular point in the sequence by moving the playback position to the point where you want to insert the time and choosing 'Insert Space At Current' from the 'Edit' menu.
+You can also add time at a particular point in the sequence by moving the playback position to the point where you want to insert the time and choosing 'Insert Space At Current' from the 
+'Edit' menu.
 
 
 ### The Loop Section
 
-The green area highlighted in Matinee is the 'Loop Section'. This is used for several utilities within Matinee, and can be adjusted by dragging the green handles on the scrub bar that mark the beginning and end of the section.
+The green area highlighted in Matinee is the 'Loop Section'. This is used for several utilities within Matinee, and can be adjusted by dragging the green handles on the scrub bar that mark 
+the beginning and end of the section.
 
-If you press the 'Loop preview playback' button, the playback position will jump to the beginning of the loop section, play through it and then jump back to the start again. This allows you to view a small section of the sequence repeatedly to ensure its doing exactly what you want. 
+If you press the 'Loop preview playback' button, the playback position will jump to the beginning of the loop section, play through it and then jump back to the start again. This allows you 
+to view a small section of the sequence repeatedly to ensure its doing exactly what you want. 
 
-You can also 'stretch' the current loop section to a new length. To do this, set the loop section over the area you wish to stretch and choose 'Stretch Section' from the 'Edit' menu. A text-entry dialog will show you the current length of the section, and allow you to enter the new length you desire. Keyframes will be re-spaced over the new length of the section. You can also delete the loop section altogether by selecting 'Delete Section' from the Edit menu.
+You can also 'stretch' the current loop section to a new length. To do this, set the loop section over the area you wish to stretch and choose 'Stretch Section' from the 'Edit' menu. 
+A text-entry dialog will show you the current length of the section, and allow you to enter the new length you desire. Keyframes will be re-spaced over the new length of the section. 
+You can also delete the loop section altogether by selecting 'Delete Section' from the Edit menu.
 
 
 ## Locking The View To A Group
 
-In Matinee, particularly when creating cinematics, it is often helpful to lock the Editor view to one of the Actors you are controlling in the scene. When you press the ![Locking button](lockview_button.jpg) button on the Group title bar, it locks any perspective viewports to be at the position of the Actor controlled by that group. As you scrub through the sequence, the camera will move with that object. Also, when you move the Editor camera, it will move the Actor around to follow it. This can be useful when adjusting a keyframe for an object to look in a particular direction.
+In Matinee, particularly when creating cinematics, it is often helpful to lock the Editor view to one of the Actors you are controlling in the scene. When you press the 
+![Locking button](lockview_button.jpg) button on the Group title bar, it locks any perspective viewports to be at the position of the Actor controlled by that group. As you scrub 
+through the sequence, the camera will move with that object. Also, when you move the Editor camera, it will move the Actor around to follow it. This can be useful when adjusting a 
+keyframe for an object to look in a particular direction.
 
-When you lock the camera to the Director Group, it will view through the currently 'active' camera. This allows you preview what the scene will actually look like in-game. Whenever you open a Matinee sequence which contains a Director Group, it will automatically turn on 'realtime' mode on any perspective viewports, and lock the camera to the Director Group.
+When you lock the camera to the Director Group, it will view through the currently 'active' camera. This allows you preview what the scene will actually look like in-game. Whenever 
+you open a Matinee sequence which contains a Director Group, it will automatically turn on 'realtime' mode on any perspective viewports, and lock the camera to the Director Group.
 
 When you add a Director Group, you can find it at the top of the Matinee Editor. The lower part of the Director Group contains a timeline. 
 
@@ -124,38 +149,52 @@ Each Matinee sequence can have only one Director Group.
 
 ## Camera Actors
 
-When making a cinematic sequence, it can be very helpful to use Camera Actors placed in the level. Although you can view through any Actor you have a Group for, CameraActors give you additional properties to control using a Float Property Track such as the camera FOV and aspect ratio, and allow you to graphically see the camera view frustum in the Editor.
+When making a cinematic sequence, it can be very helpful to use Camera Actors placed in the level. Although you can view through any Actor you have a Group for, CameraActors give you 
+additional properties to control using a Float Property Track such as the camera FOV and aspect ratio, and allow you to graphically see the camera view frustum in the Editor.
 
-You can add a Camera Actor to the level by opening the [](Engine/UI/ClassViewer), selecting CameraActor from the Actor class hierarchy, **right-clicking** in the level and choosing 'Add CameraActor Here'.
+You can add a Camera Actor to the level by opening the **[](Engine/UI/ClassViewer)**, **selecting** CameraActor from the Actor class hierarchy, **right-clicking** in the level and choosing 
+'Add CameraActor Here'.
 
 Here is a CameraActor in a level:
 
 ![Camera Actor](camera_actor.jpg)
 
-The camera mesh and view frustum will only appear in the Editor, not in the game. As you scrub around in Matinee, the active Camera frustum will highlight yellow. Each camera frustum is otherwise drawn using the Group editor color, making it easy to tell which Group is controlling which camera.
+The camera mesh and view frustum will only appear in the Editor, not in the game. As you scrub around in Matinee, the active Camera frustum will highlight yellow. Each camera frustum is 
+otherwise drawn using the Group editor color, making it easy to tell which Group is controlling which camera.
 
-When viewing through a CameraActor with _Constrain Aspect Ratio_ set to _true_, black bars will be added to the viewport to enforce the screen shape. In the Editor, these will only be added when the view is currently locked to a CameraActor.
+When viewing through a CameraActor with _Constrain Aspect Ratio_ set to _true_, black bars will be added to the viewport to enforce the screen shape. In the Editor, these will only be 
+added when the view is currently locked to a CameraActor.
 
 
 ## Matinee Replication
 
-Matinee objects that have bClientSideOnly set to _false_ are automatically replicated from server to client. Whenever the action is activated on the server, the server sends that information to the client and keeps them in sync.
+Matinee objects that have bClientSideOnly set to _false_ are automatically replicated from server to client. Whenever the action is activated on the server, the server sends that information 
+to the client and keeps them in sync.
 
-Note that this automatic replication only handles Actors with bStatic or bNoDelete set on them (that is, Actors that must be placed in the level in UnrealEd) and the Object variables connected to the matinee must be set to the desired objects in UnrealEd (not at runtime via another action). Actors for which this is not the case are responsible for replicating changes caused by matinee themselves using their own Blueprint or C++ replication code.
+Note that this automatic replication only handles Actors with bStatic or bNoDelete set on them (that is, Actors that must be placed in the level in UnrealEd) and the Object variables 
+connected to the matinee must be set to the desired objects in UnrealEd (not at runtime via another action). Actors for which this is not the case are responsible for replicating changes 
+caused by matinee themselves using their own Blueprint or C++ replication code.
 
 
 ## Cinematics
 
-In your game, you may have many in-game cinematic sequences, but lighting on the characters using the light environments set up may look a bit basic and it may be difficult to really show off the models nicely. In this case, you can use custom lighting channels (sometimes a separate lighting channel for separate characters to control the amount of lights per character), and custom placed dynamic lights to light the characters during close-up shots.
+In your game, you may have many in-game cinematic sequences, but lighting on the characters using the light environments set up may look a bit basic and it may be difficult to really show 
+off the models nicely. In this case, you can use custom lighting channels (sometimes a separate lighting channel for separate characters to control the amount of lights per character), and 
+custom placed dynamic lights to light the characters during close-up shots.
 
-For these sequences, you can use real dynamic lights because they look a little better and give you more control. You should still use very few lights, and often only one of them for casting shadows.
+For these sequences, you can use real dynamic lights because they look a little better and give you more control. You should still use very few lights, and often only one of them for 
+casting shadows.
 
-There is a small performance difference, but since cinematics have no gameplay to further slow down things you should be able to get away with it. You may have to have LD's go into cinematic maps and optimize them a few times to get things running smoothly on consoles. In most cases, all static geometry in cut scenes (except for the ground the characters are standing on, which needs real lighting for shadows) can still be light mapped.
+There is a small performance difference, but since cinematics have no gameplay to further slow down things you should be able to get away with it. You may have to have LD's go into 
+cinematic maps and optimize them a few times to get things running smoothly on consoles. In most cases, all static geometry in cut scenes (except for the ground the characters are standing 
+on, which needs real lighting for shadows) can still be light mapped.
 
 
 ## Exporting and Importing Matinee Data
 
-Matinee has _Export_ and _Import_ menu options that allow you to export the entire Matinee sequence out to a FBX file that can be loaded into a 3D content creation package.  This allows you to use an external program to animate the movement of your Actors and cameras as well as other properties such as motion blur amount and camera FOV.  Many aspects of the scene can be exported, such that you will be able to visualize your world geometry while working with an animated sequences.
+Matinee has _Export_ and _Import_ menu options that allow you to export the entire Matinee sequence out to a FBX file that can be loaded into a 3D content creation package.  
+This allows you to use an external program to animate the movement of your Actors and cameras as well as other properties such as motion blur amount and camera FOV.  Many aspects of the 
+scene can be exported, such that you will be able to visualize your world geometry while working with an animated sequences.
 
 When exporting a Matinee sequence, UnrealEd will attempt to export most of your scene geometry along with the animation curve data.
 
@@ -172,16 +211,20 @@ The following data will be exported:
 * All Emitters in current Level (placement cue only)
 
 
-Note that terrain Actors, source brush models and certain other object types are not currently supported for export.
+Note that landscape Actors, source brush models and certain other object types are not currently supported for export.
 
-Exported objects will be named similarly to their names in UnrealEd.  After making changes to the data in a 3D model package, you will want to export the changes back to a FBX file where they can be imported into Matinee.  It is important that the names of the objects in the scene stay the same during this process!
+Exported objects will be named similarly to their names in UnrealEd.  After making changes to the data in a 3D model package, you will want to export the changes back to a FBX file where 
+they can be imported into Matinee.  It is important that the names of the objects in the scene stay the same during this process!
 
-When importing the new FBX file back into the Editor, Matinee will attempt to import all of the animation curves for Actors whose imported name matches an Actor bound to the active Matinee sequence.  In some cases (such as for cameras), you can also add new objects by creating them in the modeling package and exporting them to FBX format.
+When importing the new FBX file back into the Editor, Matinee will attempt to import all of the animation curves for Actors whose imported name matches an Actor bound to the active 
+Matinee sequence.  In some cases (such as for cameras), you can also add new objects by creating them in the modeling package and exporting them to FBX format.
 
 
 ## Key Reduction
 
-The key reduction tool, which can be found under the edit tab, makes it easy to remove unneeded keys in the animation curve for a single track. You may have a track with several keys in the curve that all fall along a fairly straight line. Only the first and last keys are really needed in such cases. The key reduction tool can remove the intermediate keys to lessen the impact on memory and keep from performing useless calculations.
+The key reduction tool, which can be found under the edit tab, makes it easy to remove unneeded keys in the animation curve for a single track. You may have a track with several keys in 
+the curve that all fall along a fairly straight line. Only the first and last keys are really needed in such cases. The key reduction tool can remove the intermediate keys to lessen the 
+impact on memory and keep from performing useless calculations.
 
 ![Key Reduction tool](KeyReduction.png)
 
@@ -200,15 +243,18 @@ You play a Matinee sequence via function calls in a Blueprint graph. Most common
 
 **To play a Matinee sequence:**
 
-1. Select the Matinee Actor that contains the sequence you want to play from the **World Outliner**.
+1. **Select** the Matinee Actor that contains the sequence you want to play from the World Outliner.
 
-1. Open the Level Blueprint.  Make sure the Blueprint is in Graph Mode.
+1. **Open** the Level Blueprint.  Make sure the Blueprint is in Graph Mode.
 
-1. In the Level Blueprint, with the Matinee Actor still selected, **right-click** in the graph. In the context menu that appears, expand **Call Function on YourMatineeName**, expand the **Cinematic** section, and then click **Play**. (Alternatively, in the context menu, start typing "play" to filter the results, and choose **Play** from the list.) A **Play** node appears, with the **Matinee Actor** node already attached to it.
+1. In the Level Blueprint, with the Matinee Actor still selected, **right-click** in the graph. In the context menu that appears, expand 'Call Function on YourMatineeName', expand the 
+Cinematic section, and then **click** Play. (Alternatively, in the context menu, start typing "play" to filter the results, and **choose** Play from the list.) A **Play** node appears, 
+with the 'Matinee Actor' node already attached to it.
 
 	![](Matinee_Play.png)
 
-1. Determine what kind of event you want to use to trigger the Play function. For example, you could use a Begin Play event.  Add the event, and hook up its output execution pin to the Play node's input execution pin.
+1. Determine what kind of event you want to use to trigger the Play function. For example, you could use a Begin Play event.  Add the event, and hook up its output execution pin to the 
+Play node's input execution pin.
 
 	![](Matinee_EventBeginPlay.png)
 
@@ -223,7 +269,7 @@ You can also use individual function nodes to do any of the following with Matin
 * **Set Position**
 * **Stop**
 
-You can find these function nodes in the **Palette**, under **Matinee**.
+You can find these function nodes in the **Palette**, under Matinee.
 
 <!----For more information about Matinee and Blueprints, see *link*.----->
 
