#UE4doc #docs: Removed extraneous bolding as according to the instructional guide. ##Overview removed as according to the instructional guide. Changed title of How To to match instructional guide. Altered images to match the instructional guide: green -> yellow.

diff --git a/Engine/Documentation/Source/Gameplay/HowTo/UseOnHit/Blueprints/UseOnHit_BP.INT.udn b/Engine/Documentation/Source/Gameplay/HowTo/UseOnHit/Blueprints/UseOnHit_BP.INT.udn
--- a/Engine/Documentation/Source/Gameplay/HowTo/UseOnHit/Blueprints/UseOnHit_BP.INT.udn
+++ b/Engine/Documentation/Source/Gameplay/HowTo/UseOnHit/Blueprints/UseOnHit_BP.INT.udn
@@ -1,5 +1,5 @@
 Availability: Public
-Title:How to use the OnHit Event in Blueprints
+Title:Using the OnHit Event in Blueprints
 Crumbs: %ROOT%, Gameplay, Gameplay\HowTo\UseOnHit
 Description:In this HowTo, the OnHit Event is used to affect Actors in different ways. 
 Related: Engine/Blueprints
@@ -9,9 +9,7 @@ Related: Engine/Blueprints/UserGuide/CastNodes
 
 [TOC (start:2 end:2)]
 
-## Overview
-
-On this page, we will use the **Event Hit** Event to shoot out a light in a light fixture, which causes an emitter to spawn. We will also apply an impulse to a physics object, causing it to be launched in the direction of the impact point. Finally, we will apply damage to an object and destroy it after its "Health" has been depleted. 
+On this page, we will use the Event Hit Event to shoot out a light in a light fixture, which causes an emitter to spawn. We will also apply an impulse to a physics object, causing it to be launched in the direction of the impact point. Finally, we will apply damage to an object and destroy it after its "Health" has been depleted. 
 
 ## Project Setup
 
@@ -142,7 +140,7 @@ Here we will create a simple box that reacts to physics so that our character ca
 
 The final setup section will have us modify the player character Blueprint so they can press a button to shoot our Bomb. 
 
-1.	Inside the **Content Browser** under **Game/Blueprints**, open the **MyCharacter** Blueprint. 
+1.	Inside the **Content Browser** under **Game/Blueprints**, open the **FirstPersonCharacter** Blueprint. 
 
 1.	On the **Graph** tab, **Right-click** in the graph and add the **Left Mouse Button** event. 
 
@@ -226,12 +224,12 @@ With our project setup, we can now use the **Event Hit** node to shoot out our C
 
 	You should now see that when you hit the light with the bomb, the light goes off and the spark appears then disappears. 
 
-By using the **Cast to BP_Bomb** node, we are ensuring that only the Bomb Blueprint can register a Hit on the Light. If we wanted to have other objects hit the light and register hits, we could remove that aspect of the script and use just the **Event Hit** node alone (or possibly add a **Branch** following the Event Hit to determine what happens when different objects hit it). 
+By using the Cast to BP_Bomb node, we are ensuring that only the Bomb Blueprint can register a Hit on the Light. If we wanted to have other objects hit the light and register hits, we could remove that aspect of the script and use just the Event Hit node alone (or possibly add a Branch following the Event Hit to determine what happens when different objects hit it). 
 
 
 ## Event Hit Example (Applying Force and Damage)
 
-In this example, we are going to use the **Event Hit** node and apply force to the object we hit at the location we hit it. We will also apply damage to the object we hit, where once we hit it enough times, the object is destroyed. 
+In this example, we are going to use the Event Hit node and apply force to the object we hit at the location we hit it. We will also apply damage to the object we hit, where once we hit it enough times, the object is destroyed. 
 
 1.	In the **Content Browser**, locate your **Cube** Blueprint and drag-and-drop it into the level. 
 
@@ -307,7 +305,7 @@ In this example, we are going to use the **Event Hit** node and apply force to t
 
 	When playing in the editor, you should now see that the cube responds to the impact from the bomb projectile and apply force at the location of impact. Each time you hit the cube, health is taken away from the cube's overall health value and printed to the upper left hand corner of the screen. After shooting it 5 times, the cube is destroyed. 
 
-These are two examples of how the **Event Hit** can be used for gameplay purposes, depending upon your needs, you may perform additional checks when a Hit occurs and based on the Hit, perform different scripts that do different things. 
+These are two examples of how the Event Hit can be used for gameplay purposes, depending upon your needs, you may perform additional checks when a Hit occurs and based on the Hit, perform different scripts that do different things. 
 	
 	
 
