#UEDoc #docs - updated images and bold lettering

diff --git a/Engine/Documentation/Source/Engine/Matinee/UserGuide/BlueprintComponents/MatineeAndBlueprints.INT.udn b/Engine/Documentation/Source/Engine/Matinee/UserGuide/BlueprintComponents/MatineeAndBlueprints.INT.udn
--- a/Engine/Documentation/Source/Engine/Matinee/UserGuide/BlueprintComponents/MatineeAndBlueprints.INT.udn
+++ b/Engine/Documentation/Source/Engine/Matinee/UserGuide/BlueprintComponents/MatineeAndBlueprints.INT.udn
@@ -6,9 +6,10 @@ Description:An overview of how Matinee can control Blueprint components.
 [TOC(start:2 end:2)]
 
 
-## Overview
 
-Just as you can animate Static Meshes, Skeletal Meshes, Lights, and other Actors in Matinee, you can also animate those objects when they are components of Blueprints. You can control Blueprint components and other properties by exposing specific kinds of Blueprint variables to Matinee, using Matinee's tracks to modify those variables over time, and using those variables in combination with events and functions in the Blueprint.
+Just as you can animate Static Meshes, Skeletal Meshes, Lights, and other Actors in Matinee, you can also animate those objects when they are components of Blueprints. 
+You can control Blueprint components and other properties by exposing specific kinds of Blueprint variables to Matinee using Matinee's tracks to modify those variables over time 
+and using those variables in combination with events and functions in the Blueprint.
 
 You can also create events in Matinee, and use them to trigger Blueprint functions at specific times in the Matinee sequence's playback.
 
@@ -32,7 +33,7 @@ The types of variables you can create in a Blueprint and expose to Matinee inclu
 
 	![](BP_ExposeToMatinee.png)
 
-For more information about Blueprint variables, see [](Engine/Blueprints/UserGuide/Variables).
+For more information about Blueprint variables, see **[](Engine/Blueprints/UserGuide/Variables)**.
 
 
 
@@ -42,17 +43,17 @@ Typically, when you want to use a Blueprint variable in a Matinee sequence, it i
 
 **To modify the value of a variable in Matinee**:
 
-1. Open the Matinee Actor in Matinee.
+1. **Open** the Matinee Actor in Matinee.
 
-1. In the **World Outliner** or the Viewport, select the Blueprint Actor containing the variable.
+1. In the **World Outliner** or the Viewport, **select** the Blueprint Actor containing the variable.
 
-1. In Matinee, create an Empty Group. Give it a relevant name, such as "Blueprint Actor."
+1. In Matinee, **create** an Empty Group. Give it a relevant name, such as "Blueprint Actor."
 
 	![](Matinee_AddEmptyGroup.png)
 
 	For more information about creating Groups in Matinee, see [Working with Groups and Tracks](Engine/Matinee/UserGuide#WorkingwithGroupsAndTracks).
 
-1. **Right-click** on the new group, and select the type of track you want from the context menu:
+1. **Right-click** on the new group, and **select** the type of track you want from the context menu:
 
 	* **Bool Property Track**
 
@@ -64,7 +65,7 @@ Typically, when you want to use a Blueprint variable in a Matinee sequence, it i
 
 	![](Matinee_AddVectorTrack.png)
 
-1. In the drop-down menu that appears, select the name of the variable you want to modify, and then click **OK**.
+1. In the drop-down menu that appears, **select** the name of the variable you want to modify, and then **click ** 'OK.'
 
 	![](Matinee_ChooseVectorVariable.png)
 
@@ -74,7 +75,8 @@ Typically, when you want to use a Blueprint variable in a Matinee sequence, it i
 
 1. Modify the variable's value(s):
 
-	* For Floats, Vectors, and LinearColors, click the toggle icon (![](Engine/Matinee/UI/curve_send_button.jpg)) next to the name of the track to display its animation curves in the Curve Editor. Modify the curves in the Curve Editor.
+	* For Floats, Vectors, and LinearColors, **click** the toggle icon (![](Engine/Matinee/UI/curve_send_button.jpg)) next to the name of the track to display its animation curves in 
+	the Curve Editor. Modify the curves in the Curve Editor.
 
 		![](Matinee_CurveEditorIcon.png)
 
@@ -82,7 +84,7 @@ Typically, when you want to use a Blueprint variable in a Matinee sequence, it i
 
 		For more information about the Curve Editor, see [](Engine/UI/CurveEditor).
 
-	* For Bools, **right-click** the keyframes, and click **Set to True** or **Set to False**, depending on the Bool's current value.
+	* For Bools, **right-click** the keyframes, and **click** **Set to True** or **Set to False**, depending on the Bool's current value.
 
 		![](Matinee_BoolTrack.png)
 
@@ -94,15 +96,19 @@ When you animate a component of a Blueprint in Matinee, you cannot preview it by
 
 ## Using the Modified Variables in Blueprints
 
-In your Blueprint, you can use the variables modified in Matinee to animate specific changes over time. To do so, in the **Graph Panel**, you create a **get** node for the variable. Then, you connect that variable node to an appropriate function, and choose an event to trigger it.
+In your Blueprint, you can use the variables modified in Matinee to animate specific changes over time. To do so, in the **Graph Panel**, you create a **get** node for the variable. 
+Then, you connect that variable node to an appropriate function, and choose an event to trigger it.
 
 **To use a Matinee-modified variable in Blueprints**:
 
-1. Open the Blueprint.
+1. **Open** the Blueprint.
 
-1. In the Graph Panel, create a "get" node for the variable, either by holding down **Ctrl** and dragging the variable into the graph, or by dragging the variable into the graph and choosing **Get** from the context menu that appears.
+1. In the Graph Panel, **create** a Get node for the variable, either by holding down **Ctrl** and dragging the variable into the graph, or by dragging the variable into the graph 
+and choosing Get from the context menu that appears.
 
-1. Add and connect an appropriate event and function. For example, to change the scale of a Static Mesh component over time with a vector variable, you could use an **Event Tick** event node hooked up to a **Set Relative Scale 3D** function node, with the Static Mesh component as the **Target** and the variable node hooked up to the **New Scale 3D** pin. The Tick event is called every frame, so in general, it is a good place to update your Blueprint when using Matinee to control variable values.
+1. **Add** and connect an appropriate event and function. For example, to change the scale of a Static Mesh component over time with a vector variable, you could use an **Event Tick** 
+event node hooked up to a **Set Relative Scale 3D** function node, with the Static Mesh component as the **Target** and the variable node hooked up to the **New Scale 3D** pin. The Tick 
+event is called every frame, so in general, it is a good place to update your Blueprint when using Matinee to control variable values.
 
 	![](BP_MatineeVariable.png)
 
@@ -110,7 +116,8 @@ In your Blueprint, you can use the variables modified in Matinee to animate spec
 
 ## Triggering Events Over Time
 
-In Matinee, you can create an Event Track for a Blueprint Actor to time the triggering of specific events. You can then expose those events in the Blueprint and hook them up to functions that you want to have triggered at specific times.
+In Matinee, you can create an Event Track for a Blueprint Actor to time the triggering of specific events. You can then expose those events in the Blueprint and hook them up to 
+functions that you want to have triggered at specific times.
 
 **To use Matinee events to trigger actions at specific times in Blueprints**:
 
@@ -124,7 +131,8 @@ In Matinee, you can create an Event Track for a Blueprint Actor to time the trig
 
 	![](Matinee_Events.png)
 
-1. In the Level Blueprint, with the Matinee Actor selected, **right-click** in the graph. One of the options available is to add a MatineeController for your Matinee Actor. Select it to create a MatineeController node, which lists the events for which you created keyframes in Matinee.
+1. In the Level Blueprint, with the Matinee Actor selected, **right-click** in the graph. One of the options available is to add a MatineeController for your Matinee Actor. 
+**Select** it to create a MatineeController node, which lists the events for which you created keyframes in Matinee.
 
 	![](Matinee_AddController.png)
 
