#UEDoc #docs - updated images and doc bold lettering

diff --git a/Engine/Documentation/Source/Engine/MediaFramework/MediaFramework.INT.udn b/Engine/Documentation/Source/Engine/MediaFramework/MediaFramework.INT.udn
--- a/Engine/Documentation/Source/Engine/MediaFramework/MediaFramework.INT.udn
+++ b/Engine/Documentation/Source/Engine/MediaFramework/MediaFramework.INT.udn
@@ -51,28 +51,33 @@ Related: Engine/Rendering/Materials
 
 [TOC(start:2 end:2)]
 
-## Overview
 
-While there is currently a **MoviePlayer** system in UE4, it can only be used for playing startup movies while the Engine loads up. It cannot be used to play movies in-game, say for instance as part of a UI element or perhaps on a Static Mesh in your level like a movie playing on a TV. This is where the **Media Framework** comes in, as not only will it allow you to do both of the examples presented above, it will also provide much more general media playback capabilities (a rundown of which are outlined below).  
+While there is currently a **MoviePlayer** system in UE4, it can only be used for playing startup movies while the Engine loads up. It cannot be used to play movies in-game, 
+say for instance as part of a UI element or perhaps on a Static Mesh in your level like a movie playing on a TV. This is where the **Media Framework** comes in, as not only 
+will it allow you to do both of the examples presented above, it will also provide much more general media playback capabilities (a rundown of which are outlined below).  
 
-The **Media Framework** in Unreal Engine 4 is...
+The Media Framework in Unreal Engine 4 is...
 * Both Engine & Slate Agnostic
 * Able to support Localized Audio & Video Tracks
-* Accessible in the **Content Browser**, Material Editor, & Sound System
+* Accessible in the Content Browser, Material Editor, & Sound System
 * Available for use with Blueprint & the UMG UI Designer
 * Able to support Streaming Media
 * Able to Fast Forward, Reverse, Play, Pause, & Scrub Media
 * Able to support Pluggable Players
 
-As stated above, the Media Framework itself is both Engine and Slate agnostic, which means it can be used in pretty much any application not just the game engine or the Editor. There are layers on top of the framework that provide media playback capabilities to other sub-systems as well, such as **Engine**, **Blueprints**, **Slate**, and the **UMG UI Designer**.  This will cover most expected use cases, such as in-game textures & UI, in-Editor video tutorials, and Marketplace videos. 
+As stated above, the Media Framework itself is both Engine and Slate agnostic, which means it can be used in pretty much any application not just the game engine or the 
+Editor. There are layers on top of the framework that provide media playback capabilities to other sub-systems as well, such as **Engine**, **Blueprints**, **Slate**, and 
+the **UMG UI Designer**.  This will cover most expected use cases, such as in-game textures & UI, in-Editor video tutorials, and Marketplace videos. 
 
 [REGION:warning]
-The Media Framework is not quite production ready yet and is still in the experimental stages. Some features may not work as expected, however we are currently working hard to resolve those issues as well as add addition features and support. 
+The Media Framework is not quite production ready yet and is still in the experimental stages. Some features may not work as expected, however we are currently working hard 
+to resolve those issues as well as add addition features and support. 
 [/REGION]
 
 ## What is Media Framework
 
-The Media Framework is largely a collection of C++ interfaces, a couple of helper classes for common use cases and a media player factory that can be extended with so called Media Player plug-ins. All the interesting work of playing movies and other media files happens in the plug-ins. 
+The Media Framework is largely a collection of C++ interfaces, a couple of helper classes for common use cases and a media player factory that can be extended with so called 
+Media Player plug-ins. All the interesting work of playing movies and other media files happens in the plug-ins. 
 
 [REGION:note]
 Programmers: 
@@ -80,7 +85,8 @@ Programmers:
 * The existing player plug-ins can be found in the following location: _/Engine/Plugins/Media/_
 [/REGION:note]
 
-Currently a player plug-in for Windows is being used which uses the Windows Media Foundation API under the hood. There is also a plug-in for MacOS, which uses Apple's AV Foundation, however it is currently quite limited but should work on iOS. Work has also begun on a player plug-in for Android. 
+Currently a player plug-in for Windows is being used which uses the Windows Media Foundation API under the hood. There is also a plug-in for MacOS, which uses Apple's AV 
+Foundation, however it is currently quite limited but should work on iOS. Work has also begun on a player plug-in for Android. 
 
 [REGION:warning]
 While the Windows plug-in is the most complete, it has issues with certain movie formats, such as H.264 encoded .mp4 files. 
@@ -89,17 +95,18 @@ While the Windows plug-in is the most complete, it has issues with certain movie
 
 ## Media Asset Types & Importing Media
 
-With the implementation of the Media Framework comes three new asset types can be created from the **Content Browser**. 
+With the implementation of the Media Framework comes three new asset types can be created from the Content Browser. 
 
 ### MediaPlayer Asset
 
-The **MediaPlayer** asset represents a player for a media source such as a movie file on disk or a streaming media URL on the internet. Unlike the `UTextureMovie` in Unreal Engine 3, the MediaPlayer asset does not actually contain any data, it only stores the path or URL to the media source. 
+The MediaPlayer asset represents a player for a media source such as a movie file on disk or a streaming media URL on the internet. Unlike the `UTextureMovie` in 
+Unreal Engine 3, the MediaPlayer asset does not actually contain any data, it only stores the path or URL to the media source. 
 
 There are a number of different ways you can create a MediaPlayer asset as illustrated below. 
 
 **Create MediaPlayer Asset**
 
-1. Click the **Create** button, then select **Miscellaneous** and **Media Player**. 
+1. **Click** the **Create** button, then select **Miscellaneous** and **Media Player**. 
 
 	![](CreateMediaPlayerAsset.png)
 
@@ -111,7 +118,7 @@ Instead of clicking the **Create** button, you can also **Right-click** in the *
 
 	![](MediaPlayerAsset.png)(w:720)
 
-1. Inside the MediaPlayer asset, in the **Details** panel, click the **Browse** button. 
+1. Inside the MediaPlayer asset, in the Details panel, click the **Browse** button. 
 
 	![](AddMediaButton.png)
 
@@ -119,15 +126,15 @@ Instead of clicking the **Create** button, you can also **Right-click** in the *
 
 	![](BrowseToMovie.png)
 
-1. Once added, you can click the **Play** button inside the MediaPlayer to play your media. 
+1. Once added, you can **click** the **Play** button inside the MediaPlayer to play your media. 
 
 	![](MediaPlay.png)
 
-1. Inside the **Details** panel, you may notice that a warning icon is present. 
+1. Inside the Details panel, you may notice that a warning icon is present. 
 
 	![](WarningIcon.png)
 
-1. In order to package correctly, local media files **MUST** be placed in the **Content/Movies/** directory of your project. 
+1. In order to package correctly, local media files MUST be placed in the **Content/Movies/** directory of your project. 
 
 	![](NoWarning.png)
 
@@ -137,7 +144,7 @@ Instead of clicking the **Create** button, you can also **Right-click** in the *
 
 The easiest method of importing media is through drag-and-drop. 
 
-1. With your media inside your project's **Content/Movies/** folder, **drag-and-drop** it into the **Content Browser**. 
+1. With your media inside your project's **Content/Movies/** folder, **drag-and-drop** it into the Content Browser. 
 
 	![](DragAndDropMedia.png)
 
@@ -145,23 +152,27 @@ The easiest method of importing media is through drag-and-drop.
 
 #### MediaPlayer Asset Details
 
-The MediaPlayer asset can output various content streams, including: audio, video, and caption texts. You can set these from the drop-down menus located along the top of the media window below the playback controls. 
+The MediaPlayer asset can output various content streams, including: audio, video, and caption texts. You can set these from the drop-down menus located along the top of 
+the media window below the playback controls. 
 
 ![](OutputStreams.png)
 
-Inside the **Details** panel, you can inspect the properties of the media source. 
+Inside the Details panel, you can inspect the properties of the media source. 
 
 ![](MediaPlayerDetails.png)
 
-Under the **Playback** options, you can set the media to **Auto Play** (which will automatically play the media when loaded) and **Auto Play Rate** (to define the playback rate when Auto Play is enabled), as well as the ability to the media to looping when playback reaches the end. 
+Under the **Playback** options, you can set the media to **Auto Play** (which will automatically play the media when loaded) and **Auto Play Rate** (to define the playback 
+rate when Auto Play is enabled), as well as the ability to the media to looping when playback reaches the end. 
 
-For the **Source** section, you can specify the path to the media locally or through the Internet via URL as well as set the Streaming Mode. The last section, displays **Information** about the media itself. 
+For the **Source** section, you can specify the path to the media locally or through the Internet via URL as well as set the Streaming Mode. The last section, displays 
+information about the media itself. 
 
 ### MediaTexture Asset
 
-Once you have created a **MediaPlayer** asset, from it you can then create a **MediaTexture** asset which will allow you to extract a video stream. The MediaTexture allows you to pick a video stream (if more than one is available) and you can use it just like any other texture in the Engine.
+Once you have created a MediaPlayer asset, from it you can then create a **MediaTexture** asset which will allow you to extract a video stream. The MediaTexture allows 
+you to pick a video stream (if more than one is available) and you can use it just like any other texture in the Engine.
 
-To create a MediaTexture from a MediaPlayer, **Right-click** on a MediaPlayer asset then select the **Create Media Texture** option. 
+To create a MediaTexture from a MediaPlayer, **Right-click** on a MediaPlayer asset then **select** the **Create Media Texture** option. 
 
 ![](MediaTextureAsset.png)
 
@@ -174,14 +185,18 @@ From here, you can **Right-click** on the MediaTexture and create a **Material**
 ![](CreateMaterial.png)
 
 [REGION:note]
-If the video does not show on the texture but plays in the Media Asset Editor, then most likely the wrong video track is selected. You can change this by **Double-clicking** the MediaTexture to open the texture editor and selecting a track via the **Details** panel under MediaPlayer. 
+If the video does not show on the texture but plays in the Media Asset Editor, then most likely the wrong video track is selected. You can change this by **Double-clicking** 
+the MediaTexture to open the texture editor and selecting a track via the Details panel under MediaPlayer. 
 
-If this does not solve the issue, inside your MediaPlayer asset, check the **Auto Play** option to see if the video is then streamed inside your texture. If that addresses the issue, you can either continue to use **Auto Play** or through **Blueprints** set up some script to play the asset (see the [Media Framework With Blueprints](#mediaframeworkwithblueprints) section on this page). 
+If this does not solve the issue, inside your MediaPlayer asset, check the **Auto Play** option to see if the video is then streamed inside your texture. If that addresses 
+the issue, you can either continue to use **Auto Play** or through **Blueprints** set up some script to play the asset (see the 
+[Media Framework With Blueprints](#mediaframeworkwithblueprints) section on this page). 
 [/REGION]
 
 ### MediaSoundWaves
 
-**MediaSoundWaves** are audio components that can be placed in your level as you would any other Sound Actor. Since it relies on the MediaPlayer Asset however, in order to play the sound, you would need to tell the MediaPlayer Asset to play (if it is not set to Auto Play). 
+**MediaSoundWaves** are audio components that can be placed in your level as you would any other Sound Actor. Since it relies on the MediaPlayer Asset however, in order to
+play the sound, you would need to tell the MediaPlayer Asset to play (if it is not set to Auto Play). 
 
 You can create a MediaSoundWave by **Right-clicking** on a MediaPlayer asset and selecting **Create Media Sound Wave**. 
 
@@ -203,35 +218,36 @@ Here is an example of a MediaTexture playing on a TV (click to play).
 
 _For this example, we **Migrated** the **SM_TV** asset from the **Realistic Rendering** project which can be found on the Learn Tab._
 
-1. Creating a **MediaPlayer** and **MediaTexture** asset with a desired video using the steps outlined on this page. 
+1. Creating a MediaPlayer and MediaTexture asset with a desired video using the steps outlined on this page. 
 
-1. Open the **MediaPlayer** asset and set it to **Auto Play**. 
+1. Open the MediaPlayer asset and **set** it to **Auto Play**. 
 
-1. Open up the **M_TV** material located in the **Game/ArchVis/Materials** folder. 
+1. Open up the 'M_TV' material located in the **Game/ArchVis/Materials** folder. 
 
-1. Drag in your MediaTexture into the **M_TV** material and place it next to the **ScreenColor** node. 
+1. **Drag in** your MediaTexture into the 'M_TV' material and **place it** next to the **ScreenColor** node. 
 
 	![](Texture1.png)(w:800)
 
-1. Delete the **ScreenColor** node and plug in your texture to the **Lerp** node. 
+1. **Delete** the ScreenColor node and **plug in** your texture to the Lerp node. 
 
-1. Add a **TextureCoordinate** node connected to the Texture Sample and adjust your **UTiling** and **VTiling** as needed. 
+1. **Add** a **TextureCoordinate** node connected to the Texture Sample and adjust your 'UTiling' and 'VTiling' as needed. 
 
 	![](Texture2.png)
 
 
 ### Materials 
 
-For Materials, all you need to do is create a **Material** from your **MediaTexture** then apply the material to an object. 
+For Materials, all you need to do is create a Material from your MediaTexture then apply the material to an object. 
 
 ![](MaterialsExample.png)
 
-Above, we have applied our material to a basic cube inside our level. Our MediaPlayer asset is also set to **Auto Play**. 
+Above, we have applied our material to a basic cube inside our level. Our MediaPlayer asset is also set to Auto Play. 
 
 
 ## Media Framework with Blueprints
 
-You can also directly access and control the MediaPlayer asset through Blueprints. In this example, we have set a Trigger in front of our TV which has a MediaTexture applied to it. Our MediaPlayer asset is also not set to **Auto Play**. 
+You can also directly access and control the MediaPlayer asset through Blueprints. In this example, we have set a Trigger in front of our TV which has a MediaTexture applied
+to it. Our MediaPlayer asset is also not set to Auto Play. 
 
 ![](Blueprint1.png)
 
@@ -247,22 +263,27 @@ Inside our **Level Blueprint** we have created the following script.
 
 _Click for a full view._
 
-We have created a new variable of the **MediaPlayer** type and assigned our MediaPlayer asset as the Default Value. We have then stated that when the player character (MyCharacter) overlaps the trigger box that we want to **Play** our MediaPlayer asset (which is assigned to the variable). 
+We have created a new variable of the MediaPlayer type and assigned our MediaPlayer asset as the Default Value. We have then stated that when the player character 
+(MyCharacter) overlaps the trigger box that we want to Play our MediaPlayer asset (which is assigned to the variable). 
 
-We have done the same thing for when the player character leaves the trigger volume, except we have used the **Pause** node. The **Event End Play** node is also used to pause the MediaPlayer after exiting a Play in Editor session. These are just two of the nodes you can use with respect to MediaPlayers, if you drag off a MediaPlayer variable, you will notice that there are several nodes that can be used with it. 
+We have done the same thing for when the player character leaves the trigger volume, except we have used the Pause node. The **Event End Play** node is also used to 
+pause the MediaPlayer after exiting a Play in Editor session. These are just two of the nodes you can use with respect to MediaPlayers, if you drag off a MediaPlayer 
+variable, you will notice that there are several nodes that can be used with it. 
 
 ![](BlueprintActions.png)
 
 
 ## Media Framework with UMG
 
-If you want to have media playback as part of your UI or as an in-game movie player for example, you can do so with UMG by creating a **Widget Blueprint** and assigning a material that is created from a MediaTexture (which in turn is created from a MediaPlayer asset). 
+If you want to have media playback as part of your UI or as an in-game movie player for example, you can do so with UMG by creating a **Widget Blueprint** and assigning 
+a material that is created from a MediaTexture (which in turn is created from a MediaPlayer asset). 
 
-Below we have created a setup so that when the player press the M key in-game, a movie player is shown on screen with controls to Play, Pause, or Rewind the video that is being played (click to play). 
+Below we have created a setup so that when the player press the M key in-game, a movie player is shown on screen with controls to Play, Pause, or Rewind the video that 
+is being played (click to play). 
 
 [![](video3_image.png)](video3.mp4)
 
-We did this by first creating a [Widget Blueprint](Engine\UMG\UserGuide\WidgetBlueprints) as indicated below. 
+We did this by first creating a **[Widget Blueprint](Engine\UMG\UserGuide\WidgetBlueprints)** as indicated below. 
 
 ![](WidgetBlueprint.png)
 
@@ -276,11 +297,11 @@ We have also added a MediaPlayer variable and assigned the MediaPlayer asset we
 
 ![](WidgetBlueprint2.png)(w:840)
 
-In order for the Material to be applied to our brush, we have to make one small change and set it to **Emissive**. 
+In order for the Material to be applied to our brush, we have to make one small change and set it to Emissive. 
 
 ![](Material.png)
 
-In our **Level Blueprint**, we set up script for the player to enable/disable the in-game movie player we have created. 
+In our Level Blueprint, we set up script for the player to enable/disable the in-game movie player we have created. 
 
 [REGION:lightbox]
 [![](LevelBlueprint.png)(w:840)](LevelBlueprint.png)
