#UE4doc #docs:
Changed the "Installing Dependencies" section and added the "Compiling Source" sections to reflect the most recent versions of the engine.

diff --git a/Engine/Documentation/Source/GettingStarted/DownloadingUnrealEngine/GitHub/GitHub.INT.udn b/Engine/Documentation/Source/GettingStarted/DownloadingUnrealEngine/GitHub/GitHub.INT.udn
--- a/Engine/Documentation/Source/GettingStarted/DownloadingUnrealEngine/GitHub/GitHub.INT.udn
+++ b/Engine/Documentation/Source/GettingStarted/DownloadingUnrealEngine/GitHub/GitHub.INT.udn
@@ -84,17 +84,16 @@ Unpack the downloaded archive to a location of your choosing on your computer.
 ## Installing Dependencies
 
 In addition to the source code, there are dependencies that need to be installed in order to compile and run the 
-engine. These are mostly binary files, content assets, and external libraries. You can download these in a series 
-of .zip files from the [releases](https://github.com/EpicGames/UnrealEngine/releases) page.
+engine. These are mostly binary files, content assets, and external libraries. 
 
-Download the required dependencies files for the latest release (Unreal Engine %Globals:current_release_number%): 
-
-%Globals:current_release_dependencies%
-
-Unzip the dependencies into the UnrealEngine folder alongside the source. Be careful to make sure the folders are 
-merged together correctly. 
+After unziping the source code, you should have a file named `Setup.bat`. Running this file will install the dependencies needed.
 
 [REGION:tip]
-On Mac, we recommend **Option** + dragging the unzipped files into the UnrealEngine folder, then selecting Keep Newer 
-if prompted.
+In version 4.5.1 and prior, the `Setup.bat` does not exist and the depencies must be downloaded separately. They can be found on the [releases](https://github.com/EpicGames/UnrealEngine/releases) page listed as "Required_1of2" and "Required_2of2". These must be extracted into the Unreal Engine folder.
 [/REGION]
+
+## Compiling Source
+
+After the `Setup.bat` has completed, run `GenerateProjectFiles.bat` to generate the Visual Studio / XCode project files.
+
+Once this has completed, open the project file entitled `UE4.sln` or `UE4.xcodeproj` and **Ctrl+Shift+B** once it is done loading to begin compiling the engine. This will take some time.
