#UE4doc #docs: Removed extraneous bolding as according to the instructional guide. ##Overview removed as according to the instructional guide. Changed title of How To to match instructional guide.

diff --git a/Engine/Documentation/Source/Gameplay/HowTo/UseRaycasts/GHT_2_Overview.INT.udn b/Engine/Documentation/Source/Gameplay/HowTo/UseRaycasts/GHT_2_Overview.INT.udn
--- a/Engine/Documentation/Source/Gameplay/HowTo/UseRaycasts/GHT_2_Overview.INT.udn
+++ b/Engine/Documentation/Source/Gameplay/HowTo/UseRaycasts/GHT_2_Overview.INT.udn
@@ -1,20 +1,17 @@
 Availability: Public
-Title:How to Use Raycasts (Tracing)
+Title:Using Raycasts (Tracing)
 Crumbs: %ROOT%, Gameplay
 Description:A How To Guide for using Raycasts (Tracing) for gameplay purposes. 
 
+There may be instances in your game where you want to determine if the player character is looking at something, and if so, alter the game state in some way (for example, highlight something when a player looks at it). Or maybe you want to determine if an enemy can see the player character, and if so, start shooting or engaging them in some way. You can accomplish both scenarios by using Raycasts (or Tracing) to "shoot" out an invisible ray which will detect geometry between two points and if geometry is hit, return what was hit so that you may then do something with it. 
 
-##Overview
+There are several different options available when running a Trace. You can run a Trace to check for collision with any Objects where any Objects hit are returned, or you can run a Trace by Trace Channel where any Objects hit will return hit information only if the Object is set to specifically respond to a specified Trace Channel (which can be set via Collision Settings). 
 
-There may be instances in your game where you want to determine if the player character is looking at something, and if so, alter the game state in some way (for example, highlight something when a player looks at it). Or maybe you want to determine if an enemy can see the player character, and if so, start shooting or engaging them in some way. You can accomplish both scenarios by using **Raycasts** (or **Tracing**) to "shoot" out an invisible ray which will detect geometry between two points and if geometry is hit, return what was hit so that you may then do something with it. 
-
-There are several different options available when running a Trace. You can run a Trace to check for collision with any **Objects** where any Objects hit are returned, or you can run a Trace by **Trace Channel** where any Objects hit will return hit information only if the Object is set to specifically respond to a specified **Trace Channel** (which can be set via **Collision Settings**). 
-
-In addition to running traces by **Objects** or **Trace Channel**, you can run your Trace to detect **Single** hits or **Multi** hits, where a **Single Trace** returns only one a singular hit result and a **Multi Trace** returns multiple hits resulting from the Trace. With Traces, you can also specify the type of ray that is used: either a straight line, a box, a capsule, or sphere. 
+In addition to running traces by Objects or Trace Channel, you can run your Trace to detect Single hits or Multi hits, where a Single Trace returns only one a singular hit result and a Multi Trace returns multiple hits resulting from the Trace. With Traces, you can also specify the type of ray that is used: a straight line, a box, a capsule, or sphere. 
 
 ## Setup
 
-To learn more about **Raycasts** using Blueprints , follow the link below:
+To learn more about Raycasts using Blueprints , follow the link below:
 
 [REGION:buttonlist]
 * [](Gameplay\HowTo\UseRaycasts\Blueprints)
