#UE4doc #docs:
Unbolded multiple terms

diff --git a/Engine/Documentation/Source/Engine/Performance/PerformanceProfiling.INT.udn b/Engine/Documentation/Source/Engine/Performance/PerformanceProfiling.INT.udn
--- a/Engine/Documentation/Source/Engine/Performance/PerformanceProfiling.INT.udn
+++ b/Engine/Documentation/Source/Engine/Performance/PerformanceProfiling.INT.udn
@@ -117,27 +117,27 @@ Here we can see the **GPU** is the limiting factor (the largest of the three). I
 
 The engine show flags can be used to toggle a lot of rendering features. The editor lists all show flags in a convenient 2D UI. You can click on the checkbox to toggle multiple checkboxes without closing the menu.
 
-In-game, you can use the **show** command. Use **show** to get a list of all show flags and their state. Use **show showflagname** to toggle a feature. Note that this only works in game viewports, so in editor
+In-game, you can use the show command. Use show to get a list of all show flags and their state. Use `show showflagname` to toggle a feature. Note that this only works in game viewports, so in editor
 viewports, the editor UI needs to be used. You can use the console variables (e.g. showflag.Bloom) to override show flag values in-game or in the editor, but that also disables the UI.
 
 [REGION:note]
-Some features still consume performance even if the feature is not rendering any more, e.g. **show particles** hides the particles but simulation time is still
-required to support re-enabling them at a later time. The console variable **fx.freezeparticlesimulation** allows to you disable the update part as well (editor UI also exists).
-Another one: **show tessellation** disables triangle amplification but still uses the tessellation shaders.
+Some features still consume performance even if the feature is not rendering any more, e.g. show particles hides the particles but simulation time is still
+required to support re-enabling them at a later time. The console variable fx.freezeparticlesimulation allows to you disable the update part as well (editor UI also exists).
+Another one: show tessellation disables triangle amplification but still uses the tessellation shaders.
 [/REGION:note]
 
-A good profiling starting point is with high-level features, such as **show StaticMeshes** or **show tessellation**.
+A good profiling starting point is with high-level features, such as show StaticMeshes or show tessellation.
 
-All show flags are also exposed as console variables, e.g. console **show Bloom**, **showflag.Bloom 0** or in configuration files: **showflag.Bloom = 0**
+All show flags are also exposed as console variables, e.g. console show Bloom, showflag.Bloom 0 or in configuration files: showflag.Bloom = 0
 Console variables require more typing but they also override the editor UI showflags and they can be put into configuration files like other console variables.
 
 The most useful ones for profiling:
 
 | Show Flag | Description |
 | -------- | ----------- |
-| **ScreenSpaceReflections** | Toggles screen space reflections, can cost a lot of performance, only affects pixels up to a certain roughness (adjusted with **r.SSR.MaxRoughness** or in postprocess settings).
+| **ScreenSpaceReflections** | Toggles screen space reflections, can cost a lot of performance, only affects pixels up to a certain roughness (adjusted with r.SSR.MaxRoughness or in postprocess settings).
 | **AmbientOcclusion** | Screen Space Ambient Occlusion (Some scenes only benfit very little, for static objects you bake AO in lightmass).
-| **AntiAliasing** | Toggle any antialiasing (TemporalAA and FXAA), use **TemporalAA** to switch to FXAA (faster, lower quality).
+| **AntiAliasing** | Toggle any antialiasing (TemporalAA and FXAA), use TemporalAA to switch to FXAA (faster, lower quality).
 | **Bloom** | Affects image based lens flares and the bloom feature.
 | **DeferredLighting** | Toggle all deferred lighting passes.
 | **DirectionalLights** **PointLights** **SpotLights** | Toggles the different light types (useful to see what kind of light type is having what effect and performance impact.
@@ -166,8 +166,8 @@ The most useful ones for profiling:
 
 ### View Modes
 
-View modes are just a combination of show flags. The editor UI is separate from the show flags and you can use the **ViewMode** console command to switch between them.
-The most useful for performance are: **Wireframe**, **LightComplexity**, **ShaderComplexity**, and **Lit**.
+View modes are just a combination of show flags. The editor UI is separate from the show flags and you can use the ViewMode console command to switch between them.
+The most useful for performance are: Wireframe, LightComplexity, ShaderComplexity, and Lit.
 
 [REGION:imagetable]
 | ![](ViewModes.png)(w:794) |
@@ -181,11 +181,11 @@ Unreal Engine has [scalability](Engine/Performance/Scalability) built into many
 
 [REGION:note]
 Integrated graphics chips often have a slower memory subsystem but they also have less ALU and texture units. It is important to test on a wide range of hardware or identify the actual performance characteristics.
-The **SynthBenchmark** tool can help to identify groups of hardware that are worth testing. A number around 100 is normal for a modern GPU but if you see a disproportional change in some area, you know this GPU has a special
+The SynthBenchmark tool can help to identify groups of hardware that are worth testing. A number around 100 is normal for a modern GPU but if you see a disproportional change in some area, you know this GPU has a special
 characteristic and might need special attention.
 [/REGION]
 
-You can use the tool by typing **SynthBenchmark** into the console.
+You can use the tool by typing SynthBenchmark into the console.
 
 
     FSynthBenchmark (V0.92):
@@ -228,9 +228,9 @@ You can use the tool by typing **SynthBenchmark** into the console.
 
 It can be very useful to get the stat unit times over a longer period of time (e.g. in-game cutscene or a camera path set up to test many cases).
 
-The following chart was generated on an Android device (**Streaming**) from a cutscene. Before and after the cutscene, the console commands **StartFPSChart** and **StopFPSChart** were entered.
+The following chart was generated on an Android device (Streaming) from a cutscene. Before and after the cutscene, the console commands StartFPSChart and StopFPSChart were entered.
 The resulting .csv file (in [ProjectFolder]\Saved\Cooked\Android_ES31\SubwayPatrol\Saved\Profiling\FPSChartStats) was opened in Microsoft Excel.
-For this example, we removed the first 4 header lines, selected all, and inserted a **Scatter With Straight Lines** plot.
+For this example, we removed the first 4 header lines, selected all, and inserted a Scatter With Straight Lines plot.
 
 [REGION:imagetable]
 | ![](fpschart.png)(w:837) |
