#UE4doc #docs:
Removed Overview header

diff --git a/Engine/Documentation/Source/GettingStarted/DownloadingUnrealEngine/DownloadingUnrealEngine4.INT.udn b/Engine/Documentation/Source/GettingStarted/DownloadingUnrealEngine/DownloadingUnrealEngine4.INT.udn
--- a/Engine/Documentation/Source/GettingStarted/DownloadingUnrealEngine/DownloadingUnrealEngine4.INT.udn
+++ b/Engine/Documentation/Source/GettingStarted/DownloadingUnrealEngine/DownloadingUnrealEngine4.INT.udn
@@ -5,7 +5,6 @@ Description:Steps for connecting to the source code repository and downloading t
 
 [TOC (start:2 end:3)]
 
-## Overview
 
 [EXCERPT:Main]
 [PUBLISH:Licensee]
