#UE4doc #docs: Removed extraneous bolding as according to the instructional guide. ##Overview removed as according to the instructional guide. Changed title of How To to match instructional guide.Altered images to match the instructional guide: green -> yellow.

diff --git a/Engine/Documentation/Source/Gameplay/HowTo/UseRaycasts/Blueprints/GHT_2B.INT.udn b/Engine/Documentation/Source/Gameplay/HowTo/UseRaycasts/Blueprints/GHT_2B.INT.udn
--- a/Engine/Documentation/Source/Gameplay/HowTo/UseRaycasts/Blueprints/GHT_2B.INT.udn
+++ b/Engine/Documentation/Source/Gameplay/HowTo/UseRaycasts/Blueprints/GHT_2B.INT.udn
@@ -1,12 +1,10 @@
 Availability: Public
-Title:How to Use Raycasts (Tracing) in Blueprints
+Title:Using Raycasts (Tracing) in Blueprints
 Crumbs: %ROOT%, Gameplay, Gameplay/HowTo/UseRaycasts
 Description:A How To Guide for using Raycasts (Tracing) in Blueprints. 
 
 [TOC (start:2 end:2)]
 
-## Overview
-
 On this page, you will learn the difference between **Single** and **Multi** Raycasts as well as how to set up each which could be used to retrieve collision information. 
 
 ## Single Line Trace
@@ -15,7 +13,7 @@ On this page, you will learn the difference between **Single** and **Multi** Ray
 
 The **Single Line Trace by Channel** will perform a collision trace along a given line and return the first Object (that responds to the Trace Channel) the trace hits. 
 
-Below you will find steps for setting up a **Single Line Trace by Channel**. 
+Below you will find steps for setting up a Single Line Trace by Channel. 
 
 1.  Create a new project using the **Blueprint First Person** template with **Include Starter Content** and open the project. 
 
@@ -43,7 +41,7 @@ Below you will find steps for setting up a **Single Line Trace by Channel**.
 
 	![](GHT2B_4.png)(w:800)
 
-	Here we are starting the trace from the location of the **FirstPersonCamera**, then we are getting the rotation of the **FirstPersonCamera**. 
+	Here we are starting the trace from the location of the FirstPersonCamera, then we are getting the rotation of the FirstPersonCamera. 
 
 1.	Drag off the **Get World Rotation** and add a **Get Forward Vector**, then drag off that and add a **Vector * Float** node set to **1500**. 
 
@@ -55,7 +53,7 @@ Below you will find steps for setting up a **Single Line Trace by Channel**.
 
 	![](GHT2B_6.png)(w:800)
 
-	Here we are taking the location of the **FirstPersonCamera** and extending out from it, 1500 units based on its rotation and forward vector.
+	Here we are taking the location of the FirstPersonCamera and extending out from it, 1500 units based on its rotation and forward vector.
 
 1.	On the Trace node, set the **Draw Debug Type** to **For One Frame**. 
 
@@ -89,15 +87,15 @@ Below you will find steps for setting up a **Single Line Trace by Channel**.
 	
 	You should see that when the trace hits a cube, it prints the cube to the screen. 
 
-The example above returns all Objects that are set to respond to the provided Trace Channel, however there may be instances when you want to only return certain objects. In the example above, you can use the **Actors to Ignore** pin which takes in an **Array** of Actors which should be ignored by the trace, but this would mean you have to specify each individual Actor to ignore. 
+The example above returns all Objects that are set to respond to the provided Trace Channel, however there may be instances when you want to only return certain objects. In the example above, you can use the Actors to Ignore pin which takes in an Array of Actors which should be ignored by the trace, but this would mean you have to specify each individual Actor to ignore. 
 
-You can also perform a **Single Line Trace by Object** where only the specified **ObjectTypes** are returned. This will allow you to target a specific set of Objects only to be included in the trace. 
+You can also perform a Single Line Trace by Object where only the specified ObjectTypes are returned. This will allow you to target a specific set of Objects only to be included in the trace. 
 
 ### Single Line Trace by Object
 
 The **Single Line Trace by Object** will perform a collision trace along a given line and return the first Object (that responds to the ObjectType) the trace hits. 
 
-To set up a **Single Line Trace by Object** trace, do the following: 
+To set up a Single Line Trace by Object trace, do the following: 
 
 1.  Follow the steps used for the **Single Line Trace by Channel** example to set up a trace. 
 
@@ -209,11 +207,11 @@ Below you will find steps for setting up a **Multi Line Trace by Object**.
 ![](GHT2B_24.png)
 [/REGION]
 
-When a Line Trace is not enough, you may get the results you are after by using a **Shape Trace** instead. For example, maybe you are creating a "vision cone" for an enemy and you want to detect players that enter it. A Line Trace may not be enough as players may be able to avoid detection by ducking under the Line Trace. 
+When a Line Trace is not enough, you may get the results you are after by using a Shape Trace instead. For example, maybe you are creating a "vision cone" for an enemy and you want to detect players that enter it. A Line Trace may not be enough as players may be able to avoid detection by ducking under the Line Trace. 
 
-In this situation, you could possibly use a **Box Trace**, **Capsule Trace**, or **Sphere Trace** instead. Each Shape Trace functions like a Line Trace where you are sweeping and checking for collision from a start point to an end point, however the Shape Traces have an added layer of checking as you are using a shape as a volume of sorts in your Raycast. 
+In this situation, you could possibly use a Box Trace, Capsule Trace, or Sphere Trace instead. Each Shape Trace functions like a Line Trace where you are sweeping and checking for collision from a start point to an end point, however the Shape Traces have an added layer of checking as you are using a shape as a volume of sorts in your Raycast. 
 
-You can use a Shape Trace as a **Single Trace** or a **Multi Trace** and each are set up in the same manner as a Line Trace however you will need to provide additional details pertaining to the size or orientation of the shape you use. 
+You can use a Shape Trace as a Single Trace or a Multi Trace and each are set up in the same manner as a Line Trace however you will need to provide additional details pertaining to the size or orientation of the shape you use. 
 
 
 
