#UE4doc #docs: Removed extraneous bolding as according to the instructional guide. ##Overview removed as according to the instructional guide. Changed title of How To to match instructional guide.

diff --git a/Engine/Documentation/Source/Gameplay/HowTo/UseTimers/UseTimersOverview.INT.udn b/Engine/Documentation/Source/Gameplay/HowTo/UseTimers/UseTimersOverview.INT.udn
--- a/Engine/Documentation/Source/Gameplay/HowTo/UseTimers/UseTimersOverview.INT.udn
+++ b/Engine/Documentation/Source/Gameplay/HowTo/UseTimers/UseTimersOverview.INT.udn
@@ -1,18 +1,15 @@
 Availability: Public
-Title:How to Use Timers
+Title:Using Timers
 Crumbs: %ROOT%, Gameplay
 Description:A How To Guide for Using Timers for Gameplay in Unreal Engine 4. 
 
+**Timers** in Unreal Engine 4 are a useful tool that can be used to perform an action once (or repeatedly) at specified intervals. Timers can be used to set actions to happen on a delay or over a period of time. For example, you may want to make the player invincible after obtaining a certain item and then revert after 10 seconds. Or, you may want to apply damage every second to a player that is standing in fire. Each of these scenarios can be accomplished with the use of Timers. 
 
-##Overview
-
-**Timers** in Unreal Engine 4 are a useful tool that can be used to perform an action once (or repeatedly) at specified intervals. **Timers** can be used to set actions to happen on a delay or over a period of time. For example, you may want to make the player invincible after obtaining a certain item and then revert after 10 seconds. Or, you may want to apply damage every second to a player that is standing in fire. Each of these scenarios can be accomplished with the use of **Timers**. 
-
-The link below in the **Setup** section will give you examples on **How to Use Timers** using Blueprints.
+The link below in the Setup section will give you examples on How to Use Timers using Blueprints.
 
 ## Setup
 
-To learn more about **Using Timers** using Blueprints, follow the link below:
+To learn more about Using Timers using Blueprints, follow the link below:
 
 [REGION:buttonlist]
 * [](Gameplay\HowTo\UseTimers\Blueprints)
