#UE4doc #docs:
Unbolded a term
Changed some tensing / grammar

diff --git a/Engine/Documentation/Source/Engine/Performance/CPU/CPU.INT.udn b/Engine/Documentation/Source/Engine/Performance/CPU/CPU.INT.udn
--- a/Engine/Documentation/Source/Engine/Performance/CPU/CPU.INT.udn
+++ b/Engine/Documentation/Source/Engine/Performance/CPU/CPU.INT.udn
@@ -14,14 +14,14 @@ have to combine draw calls to reduce the cost for that (e.g. combine multiple wa
 * DirectX validates some data and passes the information to the graphics card driver.
 * The driver (e.g. NVIDIA, AMD, Intel, ...) validates further and creates a command buffer for the hardware. Sometimes this part is split in another thread.
 
-**Mesh Draw Calls** displayed when using the stats commands show the draw calls caused by 3D meshes - this is the number artists can reduce:
-* Reduce object count (static/dynamic meshes, mesh particles)
-* Reduce view distance (e.g. on the Scene Capture Actor or per object)
-* Adjust the view (more zoomed out view, moving objects to not be in the same view)
-* Avoid SceneCaptureActor (needs to re-render the scene, set fps to be low or update only when needed)
-* Avoid split screen (Always more CPU bound than single view, needs custom scalability settings or content to be more aggressive)
-* Elements per draw calls (combine materials accepting more complex pixel shaders or simply have less materials, combine textures to fewer larger textures - only if that reduces material count, use LOD models with fewer elements)
-* Disable features on the mesh like custom depth or shadow casting
+**Mesh Draw Calls** displayed when using the stats commands show the draw calls caused by 3D meshes - this is the number artists can reduce by:
+* Reducing object count (static/dynamic meshes, mesh particles)
+* Reducing view distance (e.g. on the Scene Capture Actor or per object)
+* Adjusting the view (more zoomed out view, moving objects to not be in the same view)
+* Avoiding SceneCaptureActor (needs to re-render the scene, set fps to be low or update only when needed)
+* Avoiding split screen (Always more CPU bound than single view, needs custom scalability settings or content to be more aggressive)
+* Reducing Elements per draw calls (combine materials accepting more complex pixel shaders or simply have less materials, combine textures to fewer larger textures - only if that reduces material count, use LOD models with fewer elements)
+* Disabling features on the mesh like custom depth or shadow casting
 * Changing light sources to not shadow cast or having a tighter bounding volume (view cone, attenuation radius)
 
 In some cases, hardware instancing might be an option (same 3d model, same shader, only few parameters change, hardware needs to support it). Hardware instancing reduces the driver overhead per draw call
@@ -54,7 +54,7 @@ The build in CPU profiler can help you to find the functions causing the issue:
 [/REGION]
 
 Here we used a threshold of 0.1 milliseconds to customize the output. After running the command, you can find the result in the log and in the console.
-The hierarchy shows the the time in milliseconds and the call count. If needed, you can add **QUICK_SCOPE_CYCLE_COUNTER** to the code to refine the hierarchy further like the following example:
+The hierarchy shows the the time in milliseconds and the call count. If needed, you can add QUICK_SCOPE_CYCLE_COUNTER to the code to refine the hierarchy further like the following example:
 
 	virtual void DrawDynamicElements(FPrimitiveDrawInterface* PDI,const FSceneView* View) override
 	{
