#UE4doc #docs:
Unbolded multiple terms.

diff --git a/Engine/Documentation/Source/Engine/Performance/GPU/GPU.INT.udn b/Engine/Documentation/Source/Engine/Performance/GPU/GPU.INT.udn
--- a/Engine/Documentation/Source/Engine/Performance/GPU/GPU.INT.udn
+++ b/Engine/Documentation/Source/Engine/Performance/GPU/GPU.INT.udn
@@ -5,12 +5,12 @@ Description:How to optimize the GPU demands of your game.
 Navigation:topic
 
 The GPU has many units working in parallel and it is common to be bound by different units for different parts of the frame.
-Because of this, it makes sense to look at the parts (**Where**) when looking for the bottleneck (**What**).
+Because of this, it makes sense to look at the parts (Where) when looking for the bottleneck (What).
 
 ## Where is the GPU Cost Going?
 
-The **ProfileGPU** command allows you to quickly identify the GPU cost of the various passes, sometimes down to the draw calls.
-You can either use the mouse based UI or the text version. You can suppress the UI with **r.ProfileGPU.ShowUI**. The data is based on GPU timestamps
+The ProfileGPU command allows you to quickly identify the GPU cost of the various passes, sometimes down to the draw calls.
+You can either use the mouse based UI or the text version. You can suppress the UI with r.ProfileGPU.ShowUI. The data is based on GPU timestamps
 and is usually quite accurate. Certain optimizations can make the numbers less reliable and it is good to be critical about any number. We found that some drivers
 tend to optimize shader cost a few seconds after using the shader. This can be noticeable and it might be useful to wait a bit or measure another time to get more confidence. 
 
@@ -38,14 +38,14 @@ ProfileGPU shows the light name which makes it easier for the artist to optimize
 
 It makes sense to look at the high level cost in each frame and get a feel for what is reasonable (e.g. draw call heavy, complex materials, dense triangle meshes, far view distance):
 
-* EarlyZPass: By default we use a partial z pass. DBuffer decals require a full Z Pass. This can be customized with **r.EarlyZPass** and **r.EarlyZPassMovable**.
+* EarlyZPass: By default we use a partial z pass. DBuffer decals require a full Z Pass. This can be customized with r.EarlyZPass and r.EarlyZPassMovable.
 * Base Pass: When using deferred, simple materials can be bandwidth bound. Actual vertex and pixel shader is defined in the material graph. There is an additional cost for indirect lighting on dynamic objects.
 * Shadow map rendering: Actual vertex and pixel shader is defined in the material graph. The pixel shader is only used for masked or translucent materials.
-* Shadow projection/filtering: Adjust the shader cost with **r.ShadowQuality**. Disable shadow casting on most lights. Consider static or stationary lights.
-* Occlusion culling: HZB occlusion has a high constant cost but a smaller per object cost. Toggle **r.HZBOcclusion** to see if you do better without it on.
+* Shadow projection/filtering: Adjust the shader cost with r.ShadowQuality. Disable shadow casting on most lights. Consider static or stationary lights.
+* Occlusion culling: HZB occlusion has a high constant cost but a smaller per object cost. Toggle r.HZBOcclusion to see if you do better without it on.
 * Deferred lighting: This scales with the pixels touched, and is more expensive with light functions, IES profiles, shadow receiving, area lights, and complex shading models.
-* Tiled deferred lighting: Toggle **r.TiledDeferredShading** to disable GPU lights, or use **r.TiledDeferredShading.MinimumCount** to define when to use the tiled method or the non-deferred method.
-* Environment reflections: Toggle **r.NoTiledReflections** to use the non-tiled method which is usually slower unless you have very few probes.
+* Tiled deferred lighting: Toggle r.TiledDeferredShading to disable GPU lights, or use r.TiledDeferredShading.MinimumCount to define when to use the tiled method or the non-deferred method.
+* Environment reflections: Toggle r.NoTiledReflections to use the non-tiled method which is usually slower unless you have very few probes.
 * Ambient occlusion: Quality can be adjusted, and you can use multiple passes for efficient large effects.
 * Post processing: Some passes are shared, so toggle show flags to see if the effect is worth the performance.
 
@@ -56,14 +56,14 @@ Some passes can have an effect on the passes following after them. Some examples
 
 ## What is the GPU Bottleneck?
 
-Often, the performance cost scales with the amount of pixels. To test that, you can vary the render resolution using **r.SetRes** or scale the viewport in the editor.
-Using **r.ScreenPercentage** is even more convenient but keep in mind there is some extra upsampling cost added once the feature gets used.
+Often, the performance cost scales with the amount of pixels. To test that, you can vary the render resolution using r.SetRes or scale the viewport in the editor.
+Using r.ScreenPercentage is even more convenient but keep in mind there is some extra upsampling cost added once the feature gets used.
 
 If you see a measurable performance change, you are bound by something pixel related. Usually, it is either memory bandwidth (reading and writing) or math bound (ALU), but in rare cases, some specific 
 units are saturated (e.g. MRT export). If you can lower memory (or math) on the relevant passes and see a performance difference, you know it was bound by the memory bandwith (or the ALU units). 
 The change does not have to look the same - this is just a test. Now you know you have to reduce the cost there to improve the performance.
 
-The shadow map resolution does not scale with the screen resolution (use **r.Shadow.MaxResolution**), but unless you have very large areas of shadow casting masked or translucent materials, 
+The shadow map resolution does not scale with the screen resolution (use r.Shadow.MaxResolution), but unless you have very large areas of shadow casting masked or translucent materials, 
 you are not pixel shader bound. Often the shadow map rendering is bound by vertex processing or triangle processing (causes: dense meshes, no LOD, tessellation, WorldPositionOffset use). 
 Shadow map rendering cost scales with the number of lights, number of cascades/cubemaps sides and the number of shadow casting objects in the light frustum. This is a very common bottleneck 
 and only larger content changes can reduce the cost.
@@ -75,7 +75,7 @@ we get very good quad utilization on the lighting. The problem will remain durin
 The solution to this is using less dense meshes. With level of detail meshes this can be done only where it is a problem (in the distance).
 
 [REGION:note]
-You can adjust **r.EarlyZPass** to see if your scene would benefit from a full early Z pass (more draw calls, less overdraw during base pass).
+You can adjust r.EarlyZPass to see if your scene would benefit from a full early Z pass (more draw calls, less overdraw during base pass).
 [/REGION]
 
 If changing the resolution does not matter much, you are likely bound by vertex processing (vertex shader, or tessellation) cost.
@@ -83,7 +83,7 @@ Often, you will have to change content to verify that. Typical causes include:
 
 * Too many vertices (Use Level of Detail meshes)
 * Complex WorldPositionOffset / Displacement material using textures with poor mip mapping (adjust the material)
-* Tessellation (Avoid if possible, adjust tessellation factor - fastest way: **show Tessellation**, some hardware scales badly with larger tessellation levels)
+* Tessellation (Avoid if possible, adjust tessellation factor - fastest way: show Tessellation, some hardware scales badly with larger tessellation levels)
 * Many UV or normal seams result in more vertices (look at unwrapped UV - many island are bad, flat shaded meshes have a 3 vertices per triangle)
 * Too many vertex attributes (extra UV channels)
 * Verify the vertex count is reasonable, some importer code might not have welded the vertices (combine vertices that have same position, UV and normal)
