#UE4doc #docs:
Unbolded multiple terms

diff --git a/Engine/Documentation/Source/Engine/Performance/Options/Options.INT.udn b/Engine/Documentation/Source/Engine/Performance/Options/Options.INT.udn
--- a/Engine/Documentation/Source/Engine/Performance/Options/Options.INT.udn
+++ b/Engine/Documentation/Source/Engine/Performance/Options/Options.INT.udn
@@ -9,7 +9,7 @@ Navigation:topic
 
 Certain features can be toggled at run-time in the game console or the editor output log / console.
 A lot of console variables exist to scale the render quality for a wide range of platforms
-(look at **BaseScalability.ini**, use the editor UI, the **sg.** console variables or the **Scalability** console command).
+(look at BaseScalability.ini, use the editor UI, the sg. console variables or the Scalability console command).
 
 The most useful ones for profiling:
 
@@ -21,7 +21,6 @@ The most useful ones for profiling:
 | **r.AllowOcclusionQueries** | To disable occlusion (can make scenes slower).|
 | **r.TiledDeferredShading** | Allows to turn off the tile based deferred lighting technique (GPU particle lights do not have a fallback).|
 | **r.TiledDeferredShading.MinimumCount** | Adjust how many lights we need to activate the tile based deferred lighting technique (no visual difference but can be a performance difference).|
-| **r.NoTiledReflections** | Toggles tiled reflections. |
 | **Pause** | Stop the game or matinee (much more reliable profiling but disables updates/ticks). |
 | **Slomo** | Allows to slow down or accelerate time. |
 | **r.VisualizeOccludedPrimitives** | Displays the bounding boxes of occluded primitives.|
@@ -31,7 +30,7 @@ The most useful ones for profiling:
 | **r.SceneColorFormat** | Allows to use a different scenecolor format (default is 64bit for best quality and to support screenspace subsurface scattering).|
 | **FX.AllowGPUSorting** | To disable particle sorting (pays off when used for a lot of particles).|
 | **FX.FreezeParticleSimulation** | To disable the particle updates.|
-| **r.SSR.MaxRoughness** | Adjust the max roughness of Screen Space Reflections, to override the postprocess settings, see Show Flag **VisualizeSSR**.|
+| **r.SSR.MaxRoughness** | Adjust the max roughness of Screen Space Reflections, to override the postprocess settings, see Show Flag VisualizeSSR.|
 
 
 
@@ -43,15 +42,15 @@ The most useful ones for profiling:
 | Commandline Option | Description |
 | -------- | ----------- |
 | **-NoSound** | To disable the sound and music system. |
-| **-NoTextureStreaming** | To disable texture streaming (good to isolate where hitches are coming from).|
+| **-NoTextureStreaming**  | To disable texture streaming (good to isolate where hitches are coming from).|
 | **-NoVerifyGC** | Otherwise expect massive hitching in Release builds at least every 30 seconds or so.|
 | **-NoVSync** | Gives faster rendering but can result in image tearing, especially important with high FPS.|
 | **-Streaming** | Useful with StartFPSChart/StopFPSChart to get the data from a non windows device to the cooking PC for further investigations (assuming we use cook on the fly).|
 
 ## Compile Options
 
-You should never profile or measure performance in **Debug** builds. For convenience, we suggest profiling and working on performance in **Development** build.
-The actual performance in **Shipping** is the real target but as it has fewer development features it should only be faster. **Test** acts almost like **Shipping** but has a few
+You should never profile or measure performance in Debug builds. For convenience, we suggest profiling and working on performance in Development build.
+The actual performance in Shipping is the real target but as it has fewer development features it should only be faster. Test acts almost like Shipping but has a few
 code features still enabled (e.g. console, stat unit).
 
 
