#UE4 Docs: fixed wrapping on command line item in table

 - updated CSS: remove bottom border for image table headers, add nowrap TD rule

diff --git a/Engine/Documentation/Source/Engine/Performance/Options/Options.INT.udn b/Engine/Documentation/Source/Engine/Performance/Options/Options.INT.udn
--- a/Engine/Documentation/Source/Engine/Performance/Options/Options.INT.udn
+++ b/Engine/Documentation/Source/Engine/Performance/Options/Options.INT.udn
@@ -42,7 +42,7 @@ The most useful ones for profiling:
 | Commandline Option | Description |
 | -------- | ----------- |
 | **-NoSound** | To disable the sound and music system. |
-| **-NoTextureStreaming**  | To disable texture streaming (good to isolate where hitches are coming from).|
+| [REGION:nowrap]**-NoTextureStreaming**[/REGION] | To disable texture streaming (good to isolate where hitches are coming from).|
 | **-NoVerifyGC** | Otherwise expect massive hitching in Release builds at least every 30 seconds or so.|
 | **-NoVSync** | Gives faster rendering but can result in image tearing, especially important with high FPS.|
 | **-Streaming** | Useful with StartFPSChart/StopFPSChart to get the data from a non windows device to the cooking PC for further investigations (assuming we use cook on the fly).|
