#UE4doc #docs: Removed extraneous bolding as according to the instructional guide. ##Overview removed as according to the instructional guide. Changed title of How To to match instructional guide.

diff --git a/Engine/Documentation/Source/Gameplay/HowTo/UseTimers/Blueprints/UseTimers_BP.INT.udn b/Engine/Documentation/Source/Gameplay/HowTo/UseTimers/Blueprints/UseTimers_BP.INT.udn
--- a/Engine/Documentation/Source/Gameplay/HowTo/UseTimers/Blueprints/UseTimers_BP.INT.udn
+++ b/Engine/Documentation/Source/Gameplay/HowTo/UseTimers/Blueprints/UseTimers_BP.INT.udn
@@ -1,12 +1,10 @@
 Availability: Public
-Title:How to use Timers in Blueprints
+Title:Using Timers in Blueprints
 Crumbs: %ROOT%, Gameplay, Gameplay\HowTo\UseTimers
 Description:A How To Guide for using timers in Blueprints. 
 
 [TOC (start:2 end:2)]
 
-## Overview
-
 This page covers how to use **Timers** for gameplay purposes in Blueprints and the nodes associated with Timers. 
  
 ## Creating and Clearing Timers
@@ -96,7 +94,7 @@ While the **Set Timer** and **Clear Timer** nodes are used to start/stop a Timer
 
 ![](GHT6B_10.png)
 
-The **Get Timer Elapsed** and **Get Timer Remaining** nodes will get the amount of time that has elapsed since starting or get the amount of time remaining for the specified Timer (indicated by the **FunctionName**). It returns the value as a **Float** and can be used for several different purposes, for example, changing an effect over time or providing a message when an effect is about to expire. 
+The **Get Timer Elapsed** and **Get Timer Remaining** nodes will get the amount of time that has elapsed since starting or get the amount of time remaining for the specified Timer (indicated by the FunctionName). It returns the value as a Float and can be used for several different purposes, for example, changing an effect over time or providing a message when an effect is about to expire. 
 
 ### Pause Timer and Un Pause Timer
 
@@ -114,4 +112,4 @@ The **Timer Exists**, **Is Timer Active**, and **Is Timer Paused** nodes all ret
 
 ![](GHT6B_13.png)
 
-The **Set Timer Delegate** node functions similar to the **Set Timer** node except instead of specifying a **Function Name**, you can connect a **Custom Event** which will be fired. You can set this to looping as well, just as you could with the **Set Timer** node. To clear it, you can use the **Clear Timer** node and in the **Function Name**, enter the name of your **Custom Event**.
+The **Set Timer Delegate** node functions similar to the Set Timer node except instead of specifying a Function Name, you can connect a **Custom Event** which will be fired. You can set this to looping as well, just as you could with the Set Timer node. To clear it, you can use the Clear Timer node and in the Function Name, enter the name of your Custom Event.
