#UE4doc #docs:
Removed Overview header

diff --git a/Engine/Documentation/Source/Engine/Performance/Scalability/Scalability.INT.udn b/Engine/Documentation/Source/Engine/Performance/Scalability/Scalability.INT.udn
--- a/Engine/Documentation/Source/Engine/Performance/Scalability/Scalability.INT.udn
+++ b/Engine/Documentation/Source/Engine/Performance/Scalability/Scalability.INT.udn
@@ -46,7 +46,6 @@ Description:Options for adjusting quality and performance.
 ![](SC_Banner.png)
 [/REGION]
 
-## Overview
 
 Modern games need to run on a wide range of hardware and need to adapt the game performance for various needs, such as split screen or some special gameplay mode. Tuning the game's graphics is generally the easiest to scale without affecting gameplay much.
 
