#UE4doc #docs:
Changes didn't commit the first time.

diff --git a/Engine/Documentation/Source/Engine/Performance/Scalability/ScalabilityAndYou/ScalabilityAndYou.INT.udn b/Engine/Documentation/Source/Engine/Performance/Scalability/ScalabilityAndYou/ScalabilityAndYou.INT.udn
--- a/Engine/Documentation/Source/Engine/Performance/Scalability/ScalabilityAndYou/ScalabilityAndYou.INT.udn
+++ b/Engine/Documentation/Source/Engine/Performance/Scalability/ScalabilityAndYou/ScalabilityAndYou.INT.udn
@@ -48,7 +48,6 @@ Description:Desc
 
 [TOC(start:2)]
 
-## Overview
 
 Scalability is more than just graphics settings, it is about choosing a target platform and making decisions about art, gameplay, sound, AI, and any number of other systems, fit that platform. For the consoles, this is easy enough, the platform is fixed. But, for the PC, the number of different equipment combinations makes it impossible to focus on a single hardware specification. 
 
