#UE4doc #docs: Removed extraneous bolding as according to the instructional guide. ##Overview removed as according to the instructional guide. Changed title of How To to match instructional guide. Altered images to match the instructional guide: green -> yellow.

diff --git a/Engine/Documentation/Source/Gameplay/HowTo/UsingCameras/UsingCamerasOverview.INT.udn b/Engine/Documentation/Source/Gameplay/HowTo/UsingCameras/UsingCamerasOverview.INT.udn
--- a/Engine/Documentation/Source/Gameplay/HowTo/UsingCameras/UsingCamerasOverview.INT.udn
+++ b/Engine/Documentation/Source/Gameplay/HowTo/UsingCameras/UsingCamerasOverview.INT.udn
@@ -1,14 +1,11 @@
 Availability: Public
-Title:How to Use Cameras
+Title:Using Cameras
 Crumbs: %ROOT%, Gameplay
 Description:A How To Guide for Finding Actors in your Scenes in Unreal Engine 4. 
 
+This "How To" series covers working with and placing **Cameras** in Unreal Engine 4. A Camera can be used by itself and placed directly into a level or it can be a part of a Blueprint (for instance to provide a perspective for the player when flying a plane, driving a car, or controlling a character). 
 
-##Overview
-
-This "How To" series covers working with and placing **Cameras** in Unreal Engine 4. A **Camera** can be used by itself and placed directly into a level or it can be a part of a **Blueprint** (for instance to provide a perspective for the player when flying a plane, driving a car, or controlling a character). 
-
-Placing a **Camera** Actor can be done by dragging a **Camera** Actor into your level from the **Basic** tab of the **Modes** Menu. 
+Placing a Camera Actor can be done by dragging a **Camera** Actor into your level from the **Basic** tab of the **Modes** Menu. 
 
 ![](camera1.png)(w:820)
 
@@ -25,13 +22,13 @@ With a Camera selected, you will notice that the **Details** panel becomes popul
 
 ![](camera3.png)
 
-Below is an overview of each section in the **Details** panel for a **Camera Actor**: 
+Below is an overview of each section in the **Details** panel for a Camera Actor: 
 
 | Property | Description |
 | ----- | ---------- |
 | **Transform** | This section represents the location of the Camera in the world. |
 | **Camera Settings** | This section allows you to modify the type of projection used by the camera, field of view, aspect ratio, and post process blending.|
-| **Auto Player Activation** | Specifies which **Player Controller**, if any, should automatically use this Camera when the controller is active.|
+| **Auto Player Activation** | Specifies which Player Controller, if any, should automatically use this Camera when the controller is active.|
 | **Film** | In this section, you can apply film effects such as Tint, Saturation, or Contrast. See also [Film](Resources/ContentExamples/PostProcessing/1_2).|
 | **Scene Color** | Used to apply effects to the camera. See also [Scene Color](Resources/ContentExamples/PostProcessing/1_3).|
 | **Bloom** | This simulates the effects the eye perceives when viewing bright objects. See also [Bloom](Resources/ContentExamples/PostProcessing/1_4).|
@@ -40,26 +37,26 @@ Below is an overview of each section in the **Details** panel for a **Camera Act
 | **Auto Exposure** | This simulates the human eye adjusting to bright/dark areas. See also [Auto Exposure](Resources/ContentExamples/PostProcessing/1_6).|
 | **Lens Flares** | This simulates the scattering of light when viewing bright objects through imperfections found in camera lenses. See also [Lens Flares](Resources/ContentExamples/PostProcessing/1_7).|
 | **Ambient Occlusion** | This is an effect that approximates the attenuation of light due to occlusion. See also [Ambient Occlusion](Resources/ContentExamples/PostProcessing/1_8).|
-| **Global Illumination** | Used to affect the indirect lighting contribution coming from **Lightmass** in order to alter a scene's brightness, color, or tint. See also [Global Illumination](Resources/ContentExamples/PostProcessing/1_10).|
+| **Global Illumination** | Used to affect the indirect lighting contribution coming from Lightmass in order to alter a scene's brightness, color, or tint. See also [Global Illumination](Resources/ContentExamples/PostProcessing/1_10).|
 | **Depth Of Field** | This is used to apply a blur effect to a scene based on distance in front or behind a focal point. See also [Depth Of Field](Resources/ContentExamples/PostProcessing/1_11).|
-| **Motion Blur** | Used to generate a **Motion Blur** effect that blurs objects based on its motion. See also [Motion Blur](Resources/ContentExamples/PostProcessing/1_12).|
-| **Misc** | This is used to apply **Blendables** (screen overlays), set **Anti-Aliasing** methods or **Screen Percentage** options for the Camera. See also [Blendables](Resources/ContentExamples/PostProcessing/1_16), [Screen Percentage](Resources/ContentExamples/PostProcessing/1_13), or [AA Method](Resources/ContentExamples/PostProcessing/1_14).|
-| **Screen Space Reflections** | This is an effect which is enabled by default that alters the appearance of objects on the surface of **Materials**. See also [Screen Space Reflection](Engine/Rendering/PostProcessEffects/ScreenSpaceReflection).|
+| **Motion Blur** | Used to generate a Motion Blur effect that blurs objects based on its motion. See also [Motion Blur](Resources/ContentExamples/PostProcessing/1_12).|
+| **Misc** | This is used to apply Blendables (screen overlays), set Anti-Aliasing methods or Screen Percentage options for the Camera. See also [Blendables](Resources/ContentExamples/PostProcessing/1_16), [Screen Percentage](Resources/ContentExamples/PostProcessing/1_13), or [AA Method](Resources/ContentExamples/PostProcessing/1_14).|
+| **Screen Space Reflections** | This is an effect which is enabled by default that alters the appearance of objects on the surface of Materials. See also [Screen Space Reflection](Engine/Rendering/PostProcessEffects/ScreenSpaceReflection).|
 | **Activation** | This determines whether or not the Camera is automatically enabled or not.|
-| **Tags** | This allows you to place **Tags** on the Actor. |
+| **Tags** | This allows you to place Tags on the Actor. |
 | **Actor** | This is information pertaining to the Camera Actor itself. |
-| **Blueprint** | This allows you to turn the Actor into a Blueprint or add Events to the **Level Blueprint** for the Actor. |
+| **Blueprint** | This allows you to turn the Actor into a Blueprint or add Events to the Level Blueprint for the Actor. |
 
 [REGION:note]
 For information on using Cameras for cinematic purposes, refer to [](Engine/Matinee/HowTo) documentation and the Camera sections. The [](Resources/Showcases/MatineeFightScene) is also an excellent resource for a sample cutscene and how it was constructed.
 [/REGION]
 
-Now that you have an understanding of how to place a **Camera Actor** in your level and some of the settings used to alter its look, head to the **Setup** section below to learn how to work with a Camera including how to use a Camera as a viewpoint for a player, how to use a **Camera Component** as part of a **Blueprint**, how to use a **Spring Arm Component** with a **Camera Component** (typically used for creating Third Person perspectives), as well as how to switch between Cameras during gameplay. 
+Now that you have an understanding of how to place a Camera Actor in your level and some of the settings used to alter its look, head to the Setup section below to learn how to work with a Camera including how to use a Camera as a viewpoint for a player, how to use a Camera Component as part of a Blueprint, how to use a Spring Arm Component with a Camera Component (typically used for creating Third Person perspectives), as well as how to switch between Cameras during gameplay. 
 
 
 ## Setup
 
-To learn more about **Using Cameras** in Blueprints, use the link below:
+To learn more about Using Cameras in Blueprints, use the link below:
 
 [REGION:buttonlist]
 * [](Gameplay\HowTo\UsingCameras\Blueprints)
