#UE4doc #docs: Removed extraneous bolding as according to the instructional guide. ##Overview removed as according to the instructional guide. Changed title of How To to match instructional guide. Altered images to match the instructional guide: green -> yellow.

diff --git a/Engine/Documentation/Source/Gameplay/HowTo/UsingCameras/Blueprints/UsingCameras_BP.INT.udn b/Engine/Documentation/Source/Gameplay/HowTo/UsingCameras/Blueprints/UsingCameras_BP.INT.udn
--- a/Engine/Documentation/Source/Gameplay/HowTo/UsingCameras/Blueprints/UsingCameras_BP.INT.udn
+++ b/Engine/Documentation/Source/Gameplay/HowTo/UsingCameras/Blueprints/UsingCameras_BP.INT.udn
@@ -1,12 +1,10 @@
 Availability: Public
-Title:How to Use Cameras in Blueprints
+Title:Using Cameras in Blueprints
 Crumbs: %ROOT%, Gameplay, Gameplay\HowTo\UsingCameras
 Description:A How To Guide for Finding Actors in Blueprints. 
 
 [TOC (start:2 end:2)]
 
-## Overview
-
 This page covers several examples on how to use **Cameras** within Blueprints for gameplay purposes. Whether you are looking to create a fixed camera angle for your game or cycle between multiple cameras, in order for the camera to be used during gameplay, you will need to know how to set them up first in Blueprints. 
 
 [REGION:note]
