#UE4doc #docs: corrected minor typo: from "one one now" to "open one now"

diff --git a/Engine/Documentation/Source/Gameplay/Networking/Actors/ReplicationFlow/ReplicationFlow.INT.udn b/Engine/Documentation/Source/Gameplay/Networking/Actors/ReplicationFlow/ReplicationFlow.INT.udn
--- a/Engine/Documentation/Source/Gameplay/Networking/Actors/ReplicationFlow/ReplicationFlow.INT.udn
+++ b/Engine/Documentation/Source/Gameplay/Networking/Actors/ReplicationFlow/ReplicationFlow.INT.udn
@@ -40,7 +40,7 @@ The high level flow looks like this:
 		* If the connection hasn't loaded the level this actor is in, close the channel (if any), and continue
 		* Every 1 second, determine if actor is relevant to connection by calling AActor::IsNetRelevantFor
 		* If not relevant for 5 seconds, close channel
-		* If relevant and no channel is open, one one now
+		* If relevant and no channel is open, open one now
 		* If at any point this connection is saturated
 			* For remaining actors
 				* If relevant for less than 1 second, force an update next tick
