#UE4doc #docs:
Removed Overview header
Unbolded and unitalicized multiple terms/lines
Converted some sentences to captions due to referring to pictures

diff --git a/Engine/Documentation/Source/Engine/Performance/Scalability/ScalabilityReference/ScalabilityReference.INT.udn b/Engine/Documentation/Source/Engine/Performance/Scalability/ScalabilityReference/ScalabilityReference.INT.udn
--- a/Engine/Documentation/Source/Engine/Performance/Scalability/ScalabilityReference/ScalabilityReference.INT.udn
+++ b/Engine/Documentation/Source/Engine/Performance/Scalability/ScalabilityReference/ScalabilityReference.INT.udn
@@ -48,7 +48,6 @@ Description:Scalability options, properties, and console variables.
 
 [TOC(start:2)]
 
-## Overview
 
 The Scalability settings allow you to adjust the quality of various features in order to maintain the best performance for your game on different platforms and hardware.
 
@@ -65,15 +64,17 @@ Accessing the Scalability settings in the editor is done through the **Quick Set
 
 Unreal Engine 4 can render a scene in lower resolution and scale the image up to the target resolution. As 2D User interfaces usually costs less performance and suffers more from lower resolution, Unreal Engine 4 does not apply this technique to UI. There is a minor cost for the up-sampling pass, but usually it is worth the effort.
 
-![](SC_00007.png)(w:920)
 
-_ Left: 50% no AA, Middle: 50% with AA, Right: 100% (no resolution scale) with AA_
+![](SC_00007.png)(w:920)
+[REGION:caption]
+Left: 50% no AA, Middle: 50% with AA, Right: 100% (no resolution scale) with AA 
+[/REGION]
 
 [REGION:note]
 Having a softer input image helps the up-sampling step. This means this scalability options benefits from another scalability option: the anti-aliasing quality.
 [/REGION]
 
-This setting can be access with the `r.ScreenPercentage` console variable. It takes in a value from 10-100 with -1 also equaling 100%.
+This setting can be access with the r.ScreenPercentage console variable. It takes in a value from 10-100 with -1 also equaling 100%.
 
 
 ## View Distance
@@ -81,8 +82,9 @@ This setting can be access with the `r.ScreenPercentage` console variable. It ta
 Objects can be culled based on their distance to the viewer. By default, all objects are not distance culled (Desired max draw distance of 0). On top of the designer specified value, there is a global scalability setting working like a multiplier (r.ViewDistanceScale). Here you can see some grass objects (Desired max draw distance of 1000):
 
 ![](SC_00006.png)(w:920)
-
-_left: `r.ViewDistanceScale 0.4`, middle: `r.ViewDistanceScale 0.7`, right: `r.ViewDistanceScale 1.0` (default)_
+[REGION:caption]
+left: r.ViewDistanceScale 0.4, middle: r.ViewDistanceScale 0.7, right: r.ViewDistanceScale 1.0 (default)
+[/REGION]
 
 
 ## Anti-Aliasing
@@ -90,12 +92,12 @@ _left: `r.ViewDistanceScale 0.4`, middle: `r.ViewDistanceScale 0.7`, right: `r.V
 [REGION:lightbox]
 [![](ScaleR_AA_small.png)(w:920)](ScaleR_AA.png)
 [/REGION]
+[REGION:caption]
+From the left: r.PostProcessAAQuality 0 to 6. The first 3 equal the **Low**, **Medium**, **High/Epic** settings in the **View -> Scalability Options -> AA** settings.
+[/REGION]
+Adjusting the Anti-Aliasing quality level using the r.PostProcessAAQuality console command will adjust the quality of which ever Anti-Aliasing method you are using (FXAA or Temporal AA). For either Anti-Aliasing method, a value of 0 used with r.PostProcessAAQuality will disable the effect. For FXAA, the effect of values 2, 4, and 6, can be seen in the above image; the smoothing of jagged edges becomes better and better. Values above 6 have no effect.
 
-_From the left:_ `r.PostProcessAAQuality` _0 to 6. The first 3 equal the **Low**, **Medium**, **High/Epic** settings in the **View -> Scalability Options -> AA** settings._
-
-Adjusting the Anti-Aliasing quality level using the `r.PostProcessAAQuality` console command will adjust the quality of which ever Anti-Aliasing method you are using (FXAA or Temporal AA). For either Anti-Aliasing method, a value of 0 used with `r.PostProcessAAQuality` will disable the effect. For FXAA, the effect of values 2, 4, and 6, can be seen in the above image; the smoothing of jagged edges becomes better and better. Values above 6 have no effect.
-
-For Temporal AA, there is a trade off between fill speed of the effect and quality, the higher the value you use. `r.PostProcessAAQuality 2` with Temporal AA is fast to settle but jitter caused by the effect will be more pronounced. `r.PostProcessAAQuality 4` will settle slower but will not jitter.
+For Temporal AA, there is a trade off between fill speed of the effect and quality, the higher the value you use. r.PostProcessAAQuality 2 with Temporal AA is fast to settle but jitter caused by the effect will be more pronounced. r.PostProcessAAQuality 4 will settle slower but will not jitter.
 
 
 ## Post Processing - sg.PostProcessQuality
@@ -103,10 +105,11 @@ For Temporal AA, there is a trade off between fill speed of the effect and quali
 [REGION:lightbox]
 [![](SC_PostProcessComp.png)(w:900)](SC_PostProcessComp.png)
 [/REGION]
+[REGION:caption]
+sg.PostProcessQuality set from 0 on the left to 3 on the right.
+[/REGION]
 
-`sg.PostProcessQuality` _set from 0 on the left to 3 on the right._
-
-The **View** -> **Engine Scalability Settings** -> **Post Processing** option adjusts the quality of the Post Processing effects in accordance with the settings found in the `BaseScalaiblity.ini` file located in `[UE4_InstallPath]/Engine/Config` folder. With the _Low_ setting equating to `sg.PostProcessQuality 0` and _Epic_ to `sg.PostProcessQuality 3`.
+The **View** -> **Engine Scalability Settings** -> **Post Processing** option adjusts the quality of the Post Processing effects in accordance with the settings found in the BaseScalaiblity.ini file located in [UE4_InstallPath]/Engine/Config folder. With the **Low** setting equating to sg.PostProcessQuality 0 and **Epic** to sg.PostProcessQuality 3.
 
 [REGION:simpletable]
 |sg.PostProcessQuality 0|sg.PostProcessQuality 1|
@@ -193,10 +196,11 @@ The **View** -> **Engine Scalability Settings** -> **Post Processing** option ad
 [REGION:lightbox]
 [![](SC_ShadowQualLevels.png)(w:900)](SC_ShadowQualLevels.png)
 [/REGION]
+[REGION:caption]
+sg.ShadowQuality set from 0 on the left to 3 on the right.
+[/REGION]
 
-`sg.ShadowQuality` _set from 0 on the left to 3 on the right._
-
-The **View** -> **Engine Scalability Settings** -> **Shadows** option adjusts the quality of dynamic shadows in accordance with the settings found in the `BaseScalaiblity.ini` file located in `[UE4_InstallPath]/Engine/Config` folder. With the _Low_ setting equating to `sg.ShadowQuality 0` and _Epic_ to `sg.ShadowQuality 3`.
+The **View** -> **Engine Scalability Settings** -> **Shadows** option adjusts the quality of dynamic shadows in accordance with the settings found in the BaseScalaiblity.ini file located in [UE4_InstallPath]/Engine/Config folder. With the **Low** setting equating to sg.ShadowQuality 0 and **Epic** to sg.ShadowQuality 3.
 
 [REGION:simpletable]
 |sg.ShadowQuality 0|sg.ShadowQuality 1|
@@ -305,7 +309,7 @@ The effect of this feature is heavily dependent on your game and hardware. If yo
 
 ## Effects - sg.EffectsQuality
 
-The **View** -> **Engine Scalability Settings** -> **Effects** option adjusts the quality of many different types of effects in accordance with the settings found in the `BaseScalaiblity.ini` file located in `[UE4_InstallPath]/Engine/Config` folder. With the _Low_ setting equating to `sg.EffectsQuality 0` and _Epic_ to `sg.EffectsQuality 3`.
+The **View** -> **Engine Scalability Settings** -> **Effects** option adjusts the quality of many different types of effects in accordance with the settings found in the BaseScalaiblity.ini file located in [UE4_InstallPath]/Engine/Config folder. With the **Low** setting equating to sg.EffectsQuality 0 and **Epic** to sg.EffectsQuality 3.
 
 [REGION:simpletable]
 |sg.EffectsQuality 0|sg.EffectsQuality 1|
@@ -373,7 +377,7 @@ The scalability mode can be changed from the console:
 
 	r.DetailMode 1
 
-With `r.DetailMode 0` = low, `r.DetailMode 1` = medium and, `r.DetailMode 2` = high detail modes.
+With r.DetailMode 0 = low, r.DetailMode 1 = medium and, r.DetailMode 2 = high detail modes.
 
 ![](SC_00005.png)
 
@@ -428,6 +432,6 @@ Most material editor outputs only affect the pixel shader. World position offset
 
 ![](SC_Skeletal.png)(w:920)
 
-Skeletal meshes can have static Level of Details models. With the console variable `r.SkeletalMeshLODBias`, it is possible to globally bias the LOD level. Depending on the quality of your first or second LOD, it might be a good idea to have this option in the scalability settings. Here you can see the scalability setting on top if the distance dependent LOD.
+Skeletal meshes can have static Level of Details models. With the console variable r.SkeletalMeshLODBias, it is possible to globally bias the LOD level. Depending on the quality of your first or second LOD, it might be a good idea to have this option in the scalability settings. Here you can see the scalability setting on top if the distance dependent LOD.
 
 
