#UE4doc #docs:
Fixed a hotkey reference

diff --git a/Engine/Documentation/Source/Engine/Performance/GPU/GPU.INT.udn b/Engine/Documentation/Source/Engine/Performance/GPU/GPU.INT.udn
--- a/Engine/Documentation/Source/Engine/Performance/GPU/GPU.INT.udn
+++ b/Engine/Documentation/Source/Engine/Performance/GPU/GPU.INT.udn
@@ -18,7 +18,7 @@ tend to optimize shader cost a few seconds after using the shader. This can be n
 | ![](ProfileGPU.png)(w:474) |
 | :----: | 
 | CONSOLE: ProfileGPU |
-| Shortcut: Ctrl Shift , |
+| Shortcut: Ctrl+Shift+, |
 [/REGION]
 
     ...
