#UE4doc #docs:
Changes didn't commit the first time

diff --git a/Engine/Documentation/Source/Engine/Performance/Profiler/Profiler.INT.udn b/Engine/Documentation/Source/Engine/Performance/Profiler/Profiler.INT.udn
--- a/Engine/Documentation/Source/Engine/Performance/Profiler/Profiler.INT.udn
+++ b/Engine/Documentation/Source/Engine/Performance/Profiler/Profiler.INT.udn
@@ -9,9 +9,8 @@ Description:An overview on the Profiler Tool used to monitor game performance.
 
 [TOC (start:2 end:2)]
 
-## Overview
 
-**Profiler** is the replacement for Unreal Engine 3's **StatsViewer** tool and is used to monitor the performance of a game by collecting and tracking game data. This data can then be used to identify possible sources of slowdowns or hitches in your game. 
+**Profiler** is the replacement for Unreal Engine 3's StatsViewer tool and is used to monitor the performance of a game by collecting and tracking game data. This data can then be used to identify possible sources of slowdowns or hitches in your game. 
 
 If you are new to performance tuning, you may want to refer to the [](Engine/Performance) documentation first, as there are several tips in there that will help illustrate how to monitor your game's performance. 
 
@@ -20,8 +19,8 @@ If you are new to performance tuning, you may want to refer to the [](Engine/Per
 There are two methods you can use to enable **Profiler** and collect data: 
 
 1. Live Connection
-	* Run the game with the parameter **-messaging** (for example: **UE4Editor-Win64-Debug.exe -messaging**).
-	* Run the UFE with the parameter **-messaging** (for example: **UnrealFrontend-Win64-Debug.exe -messaging**).
+	* Run the game with the parameter -messaging (for example: **UE4Editor-Win64-Debug.exe -messaging**).
+	* Run the UFE with the parameter -messaging (for example: **UnrealFrontend-Win64-Debug.exe -messaging**).
 	* Select **Session Frontend** from the **Window** menu bar, then select the **Data Preview** tab. 
 1. Load previously captured data
 	* By pressing the **Load** button on the toolbar from the **Profiler** tab and selecting a **.ue4stats** file. 
@@ -58,7 +57,7 @@ Below are descriptions of each action contained on the **Main Toolbar**:
 | Option | Description |
 | ------- | ------------ |
 | ![](LoadIcon.png) | Allows you to load a stats file. |
-| ![](SaveIcon.png) | **This option is not yet implemented.** |
+| ![](SaveIcon.png) | This option is not yet implemented. |
 | ![](DataPreviewIcon.png) | Only for live connection, enables communication with the connected session/client (when active, the icon will be lit). |
 | ![](LivePreviewIcon.png) | Only for live connection, enables live synchronization of the data graph view (when active, the icon will be lit). |
 | ![](DataCaptureIcon.png) | Only for live connection, captures the data on the client. When capture is complete, a pop-up message will ask if you would like to transfer the capture data to the user's machine. Once the data is transferred, a pop-up message will appear and allow you to load the stats file if desired. |
@@ -135,10 +134,10 @@ The values denoted with **for hierarchical data only** are only displayed for hi
 
 An explanation of the **Data Graph** axes are also outlined below: 
 
-* **Bottom** - Display the time from the beginning, starting from 0.0. 
-* **Top** - Display the frame number from the beginning, starting from 0.0. 
-* **Left** - Display adaptive information about time, in ms, based on the zoom level. Using the **Mouse Wheel**, there are three predefined axes for a game running at 30, 60 or 100 FPS each displayed in a different color. 
-* **Right** - Display adaptive information about other stats like float or memory. This axis has different scale than the left one due to the fact that float or memory values are usually larger than time related stats. 
+* Bottom - Display the time from the beginning, starting from 0.0. 
+* Top - Display the frame number from the beginning, starting from 0.0. 
+* Left - Display adaptive information about time, in ms, based on the zoom level. Using the **Mouse Wheel**, there are three predefined axes for a game running at 30, 60 or 100 FPS each displayed in a different color. 
+* Right - Display adaptive information about other stats like float or memory. This axis has different scale than the left one due to the fact that float or memory values are usually larger than time related stats. 
 
 **Right-clicking** in the **Data Graph** window will also present a context menu with additional options: 
 
