#UE4doc #docs: Removed extraneous bolding as according to the instructional guide.

diff --git a/Engine/Documentation/Source/Gameplay/Input/Input.INT.udn b/Engine/Documentation/Source/Gameplay/Input/Input.INT.udn
--- a/Engine/Documentation/Source/Gameplay/Input/Input.INT.udn
+++ b/Engine/Documentation/Source/Gameplay/Input/Input.INT.udn
@@ -8,7 +8,7 @@ The **PlayerInput** object is responsible for converting input from the player i
 can understand and make use of. It is part of an input processing flow that translates hardware input from players into game events and movement with PlayerInput mappings and InputComponents.
 
 [REGION:note]
-For an example of setting up **Input**, refer to the [](Gameplay\HowTo\SetUpInput) documentation.
+For an example of setting up Input, refer to the [](Gameplay\HowTo\SetUpInput) documentation.
 [/REGION]
 
 [TOC (start:2 end:4)]
@@ -89,8 +89,8 @@ If one InputComponent takes the input, it is not available further down the stac
 
 This example is taken from the First Person template provided with Unreal Engine 4.  
 
-1. **Hardware Input from Player:** The player presses **W**.
-1. **PlayerInput Mapping:** The AxisMapping translates **W** to "MoveForward" with a scale of 1.
+1. **Hardware Input from Player:** The player presses W.
+1. **PlayerInput Mapping:** The AxisMapping translates W to "MoveForward" with a scale of 1.
 		
 	![](AxisMappingW.png)
 
@@ -105,7 +105,7 @@ This example is taken from the First Person template provided with Unreal Engine
 				...
 			} 
 
-	This step could also be accomplished in Blueprints by having an InputAxis MoveForward node in the Character's EventGraph. Whatever this node is connected to is what will execute when **W** is pressed.
+	This step could also be accomplished in Blueprints by having an InputAxis MoveForward node in the Character's EventGraph. Whatever this node is connected to is what will execute when W is pressed.
 
 	![](InputMappingMoveForward.png)
 
@@ -147,4 +147,4 @@ Note that you will need to check the **Show Engine Content** checkbox in the obj
 ![](showenginecontent.png)
 
 If you do not want any virtual joysticks, just clear the Default Touch Interface property. Additionally, you can force the touch interface for your game independent of the platform it is running by checking
-**Always Show Touch Interface** (or by running the PC game with -faketouches).
+Always Show Touch Interface (or by running the PC game with -faketouches).
