#UE4doc #docs: ##Overview removed as according to the instructional guide. Changed title from Overview -> Localization System Overview.

diff --git a/Engine/Documentation/Source/Gameplay/Localization/Overview/Overview.INT.udn b/Engine/Documentation/Source/Gameplay/Localization/Overview/Overview.INT.udn
--- a/Engine/Documentation/Source/Gameplay/Localization/Overview/Overview.INT.udn
+++ b/Engine/Documentation/Source/Gameplay/Localization/Overview/Overview.INT.udn
@@ -1,10 +1,8 @@
 Availability:Public
-Title: Overview
+Title:Localization System Overview
 Crumbs:%ROOT%, Gameplay, Gameplay/Localization
 Description:An overview of the elements involved in localizing your project. 
 
-
-## Overview
 Localization of a project involves a few key concepts:
 
 * [Text](#text)
