#UE4doc #docs: ##Overview removed as according to the instructional guide.

diff --git a/Engine/Documentation/Source/Gameplay/Localization/Setup/Setup.INT.udn b/Engine/Documentation/Source/Gameplay/Localization/Setup/Setup.INT.udn
--- a/Engine/Documentation/Source/Gameplay/Localization/Setup/Setup.INT.udn
+++ b/Engine/Documentation/Source/Gameplay/Localization/Setup/Setup.INT.udn
@@ -5,7 +5,6 @@ Description:How to set up a project for localization.
 
 [TOC (start:2)]
 
-## Overview
 Setting a project up for localization involves:
 * writing localization configuration scripts for the commandlet pipeline.
 * configuring where to load localization data from.
