#UE4doc #docs:
Removed Overview header
Changed font settings for multiple terms

diff --git a/Engine/Documentation/Source/Engine/Performance/SystemSettings/SystemSettings.INT.udn b/Engine/Documentation/Source/Engine/Performance/SystemSettings/SystemSettings.INT.udn
--- a/Engine/Documentation/Source/Engine/Performance/SystemSettings/SystemSettings.INT.udn
+++ b/Engine/Documentation/Source/Engine/Performance/SystemSettings/SystemSettings.INT.udn
@@ -10,11 +10,10 @@ Description:Platform-specific collection of settings that control certain aspect
 
 
 
-## Overview
 
-Unreal Engine 4 allows one to control various performance-related features through the `FSystemSettings` class, which is initialized from the `[SystemSettings]` section of `Engine.ini`.
+Unreal Engine 4 allows one to control various performance-related features through the **FSystemSettings** class, which is initialized from the `[SystemSettings]` section of Engine.ini.
 
-The following sections detail the various settings that can be controlled through SystemSettings. The INI setting is literally the string name of the variable, and the type of the variable determines the values that are supported. See `EngineIni[SystemSettings]` for an example.
+The following sections detail the various settings that can be controlled through `[SystemSettings]`. The .ini setting is literally the string name of the variable, and the type of the variable determines the values that are supported. See EngineIni[SystemSettings] for an example.
 
 
 
