#UE4doc #docs:
Removed Overview header
Abbreviated sequential uses of 'Unreal Engine 4' to 'UE4'

diff --git a/Engine/Documentation/Source/Engine/Performance/Visualizers/RenderingVisualizers.INT.udn b/Engine/Documentation/Source/Engine/Performance/Visualizers/RenderingVisualizers.INT.udn
--- a/Engine/Documentation/Source/Engine/Performance/Visualizers/RenderingVisualizers.INT.udn
+++ b/Engine/Documentation/Source/Engine/Performance/Visualizers/RenderingVisualizers.INT.udn
@@ -47,9 +47,8 @@ Description:Display options and view modes to check on the results of various re
 [/REGION]
 
 
-## Overview
 
-Unreal Engine 4 has a number of visualizers to enable you to debug or simply view what Unreal Engine 4 is doing behind the scenes. These can be rendering steps, such as viewing the depth pass, or specialized views, such as the motion blur vector visualizer.
+Unreal Engine 4 (UE4) has a number of visualizers to enable you to debug or simply view what UE4 is doing behind the scenes. These can be rendering steps, such as viewing the depth pass, or specialized views, such as the motion blur vector visualizer.
 
 [REGION:topics third]
 %Engine/Performance/Visualizers/OutOfBoundsPixels:topiccompact%
