#UE4doc #docs: ##Overview removed as according to the instructional guide.

diff --git a/Engine/Documentation/Source/Gameplay/Networking/Actors/Components/Components.INT.udn b/Engine/Documentation/Source/Gameplay/Networking/Actors/Components/Components.INT.udn
--- a/Engine/Documentation/Source/Gameplay/Networking/Actors/Components/Components.INT.udn
+++ b/Engine/Documentation/Source/Gameplay/Networking/Actors/Components/Components.INT.udn
@@ -8,7 +8,6 @@ Related: Gameplay/Networking/Actors/ReplicationPerformance
 
 [TOC (start:2 end:3)]
 
-## Overview
 Unreal Engine 4 supports component replication. While it is straightforward to use, it is not that common: most components do not replicate. Most gameplay logic is done in the Actor class and components 
 usually just represent smaller pieces that make up the Actor. That gameplay logic in the Actor is what replicates, and its results sometimes end up in calls/changes to the components. However, cases exist
  where properties or events on components themselves need to directly replicate. This guide explains how to use this.
