#UE4doc #docs: Made minor grammar changes: singular vs plural, comma changes.

diff --git a/Engine/Documentation/Source/Gameplay/Networking/Actors/Properties/ObjectReferences/ObjectReferences.INT.udn b/Engine/Documentation/Source/Gameplay/Networking/Actors/Properties/ObjectReferences/ObjectReferences.INT.udn
--- a/Engine/Documentation/Source/Gameplay/Networking/Actors/Properties/ObjectReferences/ObjectReferences.INT.udn
+++ b/Engine/Documentation/Source/Gameplay/Networking/Actors/Properties/ObjectReferences/ObjectReferences.INT.udn
@@ -15,7 +15,7 @@ To replicate object references, you simply mark a UObject property as replicated
 		AActor * Owner;
 	};
 
-In this case, the "Owner" property, will be a replicated reference to the actor that this property refers to.
+In this case, the "Owner" property will be a replicated reference to the actor that this property refers to.
 
 For an object to be legally referenced over the network, it must be supported for networking. To check for this, you can call `UObject::IsSupportedForNetworking()`. This is generally considered a low level function, so it's not common to need to check this in game code.
 
@@ -33,7 +33,7 @@ Stably named objects are simply objects that will exist on both the server and t
 
 Actors are stably named if they were loaded directly from packages (not spawned during gameplay).
 
-Component are stably named if they were:
+Components are stably named if they were:
 * Loaded directly from a package
 * Added via simple construction scripts
 * Manually marked (via `UActorComponent::SetNetAddressable`)
