#UE4doc #docs: ##Overview removed as according to the instructional guide.

diff --git a/Engine/Documentation/Source/Gameplay/Networking/Actors/RPCs/RPCs.INT.udn b/Engine/Documentation/Source/Gameplay/Networking/Actors/RPCs/RPCs.INT.udn
--- a/Engine/Documentation/Source/Gameplay/Networking/Actors/RPCs/RPCs.INT.udn
+++ b/Engine/Documentation/Source/Gameplay/Networking/Actors/RPCs/RPCs.INT.udn
@@ -7,7 +7,6 @@ Related: Gameplay/Networking/Actors/Components
 
 [TOC (start:2 end:3)]
 
-## Overview
 RPCs (Remote Procedure Calls) are functions that are called locally, but executed remotely on another machine (separate from the calling machine). 
 
 RPC functions can be very useful and allow either the client or the server to send messages to each other over a network connection.
