#UE4doc #docs: ##Overview removed as according to the instructional guide. Changed singular vs plural error.

diff --git a/Engine/Documentation/Source/Gameplay/Networking/Blueprints/Blueprints.INT.udn b/Engine/Documentation/Source/Gameplay/Networking/Blueprints/Blueprints.INT.udn
--- a/Engine/Documentation/Source/Gameplay/Networking/Blueprints/Blueprints.INT.udn
+++ b/Engine/Documentation/Source/Gameplay/Networking/Blueprints/Blueprints.INT.udn
@@ -7,8 +7,6 @@ Related: Gameplay/Networking/Actors
 
 [TOC (start:2 end:3)]
 
-## Overview
-
 Unreal Engine 4 provides a lot of multiplayer functionality out of the box, and it's easy set up a basic Blueprint game that works over a network. [It's easy to dive in](Gameplay/HowTo/Networking/TestMultiplayer) and start playing mutliplayer. Most of the logic to make basic multiplayer work is thanks to the built-in networking support in the `Character` class, and its `CharacterMovementComponent`, which the Third Person template project uses.
 
 ## Gameplay framework review
@@ -139,4 +137,4 @@ For any replicated event, you can choose whether it is **reliable** or **unrelia
 
 **Reliable** events are guaranteed to reach their destination (assuming the ownership rules above are followed), but they introduce more bandwidth, and potentially latency, to meet this guarantee. Try to avoid sending reliable events too often, such as on every tick, since the engine's internal buffer of reliable events may overflow - when this happens, the associated player will be disconnected!
 
-**Unreliable** events work as their name implies - they may not reach their destination, in case of packet loss on the network, or if the engine determines there is a lot of higher-priority traffic it needs to send, for example. As a result, unreliable events use less bandwidth than reliables, and they can be called safely more often.
+**Unreliable** events work as their name implies - they may not reach their destination, in case of packet loss on the network, or if the engine determines there is a lot of higher-priority traffic it needs to send, for example. As a result, unreliable events use less bandwidth than reliable events, and they can be called safely more often.
