#UEDoc #docs - Updated doc and removed non-applicable image

diff --git a/Engine/Documentation/Source/Engine/Performance/Visualizers/OutOfBoundsPixels/OutOfBoundsPixels.INT.udn b/Engine/Documentation/Source/Engine/Performance/Visualizers/OutOfBoundsPixels/OutOfBoundsPixels.INT.udn
--- a/Engine/Documentation/Source/Engine/Performance/Visualizers/OutOfBoundsPixels/OutOfBoundsPixels.INT.udn
+++ b/Engine/Documentation/Source/Engine/Performance/Visualizers/OutOfBoundsPixels/OutOfBoundsPixels.INT.udn
@@ -48,11 +48,13 @@ Description:Visualizing out of bounds pixels in the Unreal Editor viewport.
 
 [TOC(start:2)]
 
-## Overview
 
-If you run into a situation where an object is flickering or perhaps it is disappearing even though a portion of it is still on screen, it may be because parts of it are outside of its bounds. Common reasons for this are: use of world position offset or displacement in a shader, Skeletal Mesh without a physics asset, or a component setup to inherit owner bounds.
+If you run into a situation where an object is flickering or perhaps it is disappearing even though a portion of it is still on screen, it may be because parts of it are 
+outside of its bounds. Common reasons for this are: use of world position offset or displacement in a shader, Skeletal Mesh without a physics asset, or a component setup to 
+inherit owner bounds.
 
-This causes frustum culling, occlusion culling, shadow casting, indirect lighting cache and various other rendering functions to operate incorrectly.  The most noticeable artifact will be flickering as you rotate the camera, since it will be incorrectly culled as soon as the bounds are off screen.
+This causes frustum culling, occlusion culling, shadow casting, indirect lighting cache and various other rendering functions to operate incorrectly.  The most noticeable 
+artifact will be flickering as you rotate the camera, since it will be incorrectly culled as soon as the bounds are off screen.
 
 **Visualize -> Out of Bounds Pixels** will show you any pixel that is outside of its object's bounds.
 
@@ -68,21 +70,25 @@ To enable this mode on any viewport, **Left-click** the **Show Button** ![](OOBP
 
 Blueprint components, often times, can be set to inherit the bounds of their parent, which can cause issues if the parent has smaller bounds than the child.
 
-The cube object below is a Blueprint composed of a **Box Component** (the red box in these images) and a **Static Mesh Component**. The Static Mesh is set to inherit its bounds from the box. If the box extents are smaller than the Static Mesh, then **Visualize Out of Bounds Pixels** will show those portions in various colors (blue, yellow, and white).
+The cube object below is a Blueprint composed of a **Box Component** (the red box in these images) and a **Static Mesh Component**. The Static Mesh is set to inherit its 
+bounds from the box. If the box extents are smaller than the Static Mesh, then Visualize Out of Bounds Pixels will show those portions in various colors (blue, yellow, 
+and white).
 
 [REGION:lightbox]
 [![](OOBP_smallBounds_boxExtents.png)(w:300)](OOBP_smallBounds_boxExtents.png) [![](OOBP_smallBounds.png)(w:300)](OOBP_smallBounds.png) [![](OOBP_smallBoundsRotated.png)(w:300)](OOBP_smallBoundsRotated.png)
 
-Since the bounding box (the blue box in these images) is always world aligned, rotating the box can bring portions of the mesh back into the bounding box. However, if any out of bounds portion of the mesh remains in the camera view, it will disappear if the bounding box moves completely off camera.
+Since the bounding box (the blue box in these images) is always world aligned, rotating the box can bring portions of the mesh back into the bounding box. However, if any out 
+of bounds portion of the mesh remains in the camera view, it will disappear if the bounding box moves completely off camera.
 
-Below is a corrected Blueprint, with the **Box Component's** extents set correctly.
+Below is a corrected Blueprint, with the Box Component's extents set correctly.
 
 [![](OOBP_correctBounds_boxExtents.png)(w:300)](OOBP_correctBounds_boxExtents.png) [![](OOBP_correctBounds.png)(w:300)](OOBP_correctBounds.png) [![](OOBP_correctBoundsRotated.png)(w:300)](OOBP_correctBoundsRotated.png)
 [/REGION]
 
 ## World Position Offset or Displacement
 
-The other common time this can occur is when a shader is modifying the world position of vertices in real time, pushing them well outside the bounds of the object's original bounding box:
+The other common time this can occur is when a shader is modifying the world position of vertices in real time, pushing them well outside the bounds of the object's original 
+bounding box:
 
 [REGION:lightbox]
 [![](OOBP_worldPostionOffsetIssue_02.png)(h:300)](OOBP_worldPostionOffsetIssue_02.png) [![](OOBP_worldPostionOffsetIssue_01.png)(h:300)](OOBP_worldPostionOffsetIssue_01.png)
@@ -90,8 +96,13 @@ The other common time this can occur is when a shader is modifying the world pos
 
 Because portions of this object will be visible after the bounding box is completely hidden or off camera, it will flicker or disappear at inappropriate times.
 
-In this case, BoundsScale under **Component -> Rendering -> Advanced** can be used to increase the bounding box of the mesh until it contains all the pixels.  Be careful doing this on shadow casters though, large shadow casters have a large performance cost.  On Skeletal Meshes, the bounds are derived from the physics asset, and individual bones can be marked as contributing to the bounds or not.
+In this case, BoundsScale under **Component -> Rendering -> Advanced** can be used to increase the bounding box of the mesh until it contains all the pixels.  Be careful 
+doing this on shadow casters though, large shadow casters have a large performance cost.  On Skeletal Meshes, the bounds are derived from the physics asset, and individual 
+bones can be marked as contributing to the bounds or not.
 
 ## Skeletal Meshes
 
-The primary reason a Skeletal Mesh becomes displaced from its bounds is due to the animation causing the mesh to move away from its Actor root. However, a Skeletal Mesh with a Physics Asset will use the Physics Asset for calculating its bounds. To enable this, **right-click** the Skeletal Mesh in the **Content Browser** and select **Create Physics Asset**. Using the default settings will get you past any errors being caused by the Skeletal Mesh bounds, but you will probably want to edit the results. See the [**Physics Asset Tool**](Engine/Physics\PhAT) documentation for more information on editing a **Physics Asset**.
+The primary reason a Skeletal Mesh becomes displaced from its bounds is due to the animation causing the mesh to move away from its Actor root. However, a Skeletal Mesh with 
+a Physics Asset will use the Physics Asset for calculating its bounds. To enable this, **right-click** the Skeletal Mesh in the **Content Browser** and select 
+**Create Physics Asset**. Using the default settings will get you past any errors being caused by the Skeletal Mesh bounds, but you will probably want to edit the results. 
+See the [**Physics Asset Tool**](Engine/Physics\PhAT) documentation for more information on editing a **Physics Asset**.
diff --git a/Engine/Documentation/Source/Engine/Physics/Apex/Apex.INT.udn b/Engine/Documentation/Source/Engine/Physics/Apex/Apex.INT.udn
--- a/Engine/Documentation/Source/Engine/Physics/Apex/Apex.INT.udn
+++ b/Engine/Documentation/Source/Engine/Physics/Apex/Apex.INT.udn
@@ -46,16 +46,19 @@ Description:Creation and Import of Destructibles using the Nvidia APEX Toolset.
 ![Physics](Engine/Physics/physics_landing_banner.png)
 [/REGION]
 
-## Overview
 
-Nvidia produces a tool called the **APEX PhysX Lab** which can be used to create destructible meshes, cloth assets, and APEX particle systems. Information on its usage can be found in the Nvidia documentation over the APEX PhysX Lab tools. However, it should be noted that currently Unreal Engine 4 only supports **APEX Destructibles** and **APEX Cloth**. This tool is only required for making APEX Cloth assets.
+Nvidia produces a tool called the APEX PhysX Lab which can be used to create destructible meshes. Information on its usage can be found in the 
+Nvidia documentation over the APEX PhysX Lab tools. However, it should be noted that currently Unreal Engine 4 only supports 
+APEX Destructibles and APEX Cloth. To create APEX Cloth assets it is required to use the PhysX: 3Ds Max/Maya DCC plug-in or the standalone tool is included in the APEX SDK. 
 
-Destructible Actors can be generated within Unreal Editor 4 using **Voronoi** with a single depth level.  APEX Destructible assets with more than one depth level or using a function other than **Voronoi** to shatter the mesh, must currently be generated in the APEX PhysX Lab.
+Destructible Actors can be generated within Unreal Editor 4 using Voronoi with a single depth level.  APEX Destructible assets with more than one depth 
+level or using a function other than Voronoi to shatter the mesh, must currently be generated in the APEX PhysX Lab.
 
 [REGION:lightbox]
-[![](APEX_openAPL.png)(w:450)](APEX_openAPL.png) [![](APEX_destructibleProj_APL.png)(w:450)](APEX_destructibleProj_APL.png)
+ [![](APEX_destructibleProj_APL.png)(w:450)](APEX_destructibleProj_APL.png)
 [/REGION]
 
 ## Importing APEX Destructibles and Cloth
 
-Importing an APEX object is the same as importing a FBX file. All APEX assets are exported as .apb or .apx, Unreal Editor will determine the type of APEX object being imported.
+Importing an APEX object is the same as importing a FBX file. All APEX assets are exported as .apb or .apx, Unreal Editor will determine the type of 
+APEX object being imported.
