#UE4doc #docs: Removed extraneous bolding as according to the instructional guide. ##Overview removed as according to the instructional guide.

diff --git a/Engine/Documentation/Source/Gameplay/SaveGame/SaveGame.INT.udn b/Engine/Documentation/Source/Gameplay/SaveGame/SaveGame.INT.udn
--- a/Engine/Documentation/Source/Gameplay/SaveGame/SaveGame.INT.udn
+++ b/Engine/Documentation/Source/Gameplay/SaveGame/SaveGame.INT.udn
@@ -5,8 +5,6 @@ Description:Overview of how to save and load your game
 
 [TOC (start:2 end:2)]
 
-## Overview
-
 A simple way to set up saving and loading in your game is to use a SaveGame object. 
 
 [REGION:raw]
@@ -34,9 +32,9 @@ The SaveGame object will now hold on to these variable values until it is saved
 [/REGION]
 
 
-Finally, you can save out the SaveGame object to a file on your hard drive by using **SaveGameToSlot**.  This works just like saving a document or a new game in other programs, since you
-give **SaveGameToSlot** a file name. By default, save games are stored as `.sav` files in the `Saved\SaveGames` folder. 
-You also must give **SaveGameToSlot** a user ID, which will help in looking up the information to load later on.
+Finally, you can save out the SaveGame object to a file on your hard drive by using SaveGameToSlot.  This works just like saving a document or a new game in other programs, since you
+give SaveGameToSlot a file name. By default, save games are stored as `.sav` files in the `Saved\SaveGames` folder. 
+You also must give SaveGameToSlot a user ID, which will help in looking up the information to load later on.
 
 Loading information from a SaveGame object is also a simple process, which is almost the reverse of saving information. You use the **LoadGameFromSlot** function, which takes in the file name
 and the user ID as input.  Loading information from the hard drive with **LoadGameFromSlot** creates a new copy of the SaveGame object.
