#UE4doc #docs: Acidentally left out some bolding. Taken care of now.

diff --git a/Engine/Documentation/Source/Gameplay/SaveGame/SaveGame.INT.udn b/Engine/Documentation/Source/Gameplay/SaveGame/SaveGame.INT.udn
--- a/Engine/Documentation/Source/Gameplay/SaveGame/SaveGame.INT.udn
+++ b/Engine/Documentation/Source/Gameplay/SaveGame/SaveGame.INT.udn
@@ -36,8 +36,8 @@ Finally, you can save out the SaveGame object to a file on your hard drive by us
 give SaveGameToSlot a file name. By default, save games are stored as `.sav` files in the `Saved\SaveGames` folder. 
 You also must give SaveGameToSlot a user ID, which will help in looking up the information to load later on.
 
-Loading information from a SaveGame object is also a simple process, which is almost the reverse of saving information. You use the **LoadGameFromSlot** function, which takes in the file name
-and the user ID as input.  Loading information from the hard drive with **LoadGameFromSlot** creates a new copy of the SaveGame object.
+Loading information from a SaveGame object is also a simple process, which is almost the reverse of saving information. You use the LoadGameFromSlot function, which takes in the file name
+and the user ID as input.  Loading information from the hard drive with LoadGameFromSlot creates a new copy of the SaveGame object.
 
 [REGION:raw]
 ![](StoredVariables.png)(h:150)
