#UE4doc #docs: Removed extraneous bolding as according to the instructional guide.

diff --git a/Engine/Documentation/Source/Gameplay/SaveGame/Blueprints/BP_SaveGame.INT.udn b/Engine/Documentation/Source/Gameplay/SaveGame/Blueprints/BP_SaveGame.INT.udn
--- a/Engine/Documentation/Source/Gameplay/SaveGame/Blueprints/BP_SaveGame.INT.udn
+++ b/Engine/Documentation/Source/Gameplay/SaveGame/Blueprints/BP_SaveGame.INT.udn
@@ -6,15 +6,15 @@ Description:How to save and load your game using Blueprints
 
 ### Create a SaveGame Object
 
-To create a new SaveGame object, [create a new Blueprint Class](Engine\Blueprints\UserGuide\Types\ClassBlueprint\Creation). When the **Pick Parent Class** dialog pops up, expand the **Custom Classes** 
-dropdown, then select **SaveGame**. You can use the search box to jump directly to **SaveGame**.
+To create a new SaveGame object, [create a new Blueprint Class](Engine\Blueprints\UserGuide\Types\ClassBlueprint\Creation). When the Pick Parent Class dialog pops up, expand the **Custom Classes** 
+dropdown, then select **SaveGame**. You can use the search box to jump directly to SaveGame.
 
 ![](savegame.png) 
 
 In your new SaveGame object Blueprint, create variables for any information you would like to save.  
 
 [REGION:note]
-In this example, there are also variables declared that will be used to store default values for the **SaveSlotName** and the **UserIndex**, so that each class that saves to this
+In this example, there are also variables declared that will be used to store default values for the SaveSlotName and the UserIndex, so that each class that saves to this
 SaveGame object will not have to independently set those variables. This step is optional, and will cause there to be one save slot that gets overwritten if the default values are not changed.
 [/REGION]
 
@@ -32,8 +32,8 @@ You can set default values for the variables after the Blueprint is compiled.
 [/REGION]
 
 [REGION:fullwidth]
-First, create an instance, or copy, of your SaveGame object, using the **Create Save Game Object** node. Make sure you set the **Save Game Class** dropdown to the name of your new SaveGame object Blueprint, 
-in this case, **MySaveGame**. Because the **Create Save Game Object** node creates a generic SaveGame object copy, you have to drag off of the **Return Value** pin and **Cast to MySaveGame**. Save
+First, create an instance, or copy, of your SaveGame object, using the **Create Save Game Object** node. Make sure you set the Save Game Class dropdown to the name of your new SaveGame object Blueprint, 
+in this case, **MySaveGame**. Because the Create Save Game Object node creates a generic SaveGame object copy, you have to drag off of the **Return Value** pin and **Cast to MySaveGame**. Save
 the result of the casting to a variable so that you can reuse the SaveGame object easily later on.
 
 ![](SaveGameBP_1.png)
