#UEDoc #docs - updated image. Process for APEX Cloth import is not accurate.

diff --git a/Engine/Documentation/Source/Engine/Physics/Apex/ApexTypes/ApexTypes.INT.udn b/Engine/Documentation/Source/Engine/Physics/Apex/ApexTypes/ApexTypes.INT.udn
--- a/Engine/Documentation/Source/Engine/Physics/Apex/ApexTypes/ApexTypes.INT.udn
+++ b/Engine/Documentation/Source/Engine/Physics/Apex/ApexTypes/ApexTypes.INT.udn
@@ -45,7 +45,6 @@ Description:Nvidia APEX Physics Lab APX and APB files.
 [TOC (start:2 end:3)]
 
 
-## Overview
 
 [OBJECT:topiclistnotitle]
 	[PARAM:icon]
@@ -65,13 +64,12 @@ See the [APEX Physics Lab](Engine/Physics/Apex) documentation on creating and ex
 In the **Content Browser**:
 
 1. Click the ![New Asset](CB_button_NewAsset.png) button.
-1. Select _Import_ from the menu.
 1. Select the .APX or .APB file(s) you wish to import.
 
 Or
 
 1. Select your .APX or .APB file(s) file in a Windows File Explorer window.
-1. Click and drag it (them) into the **Content Browser** in Unreal Editor 4.
+1. Click and drag it (them) into the Content Browser in Unreal Editor 4.
 
 ## Related
 
