#UE4doc #docs: Unbolded hyperlink according to the instructional guide.

diff --git a/Engine/Documentation/Source/Engine/UMG/QuickStart/UMGQuickStart.INT.udn b/Engine/Documentation/Source/Engine/UMG/QuickStart/UMGQuickStart.INT.udn
--- a/Engine/Documentation/Source/Engine/UMG/QuickStart/UMGQuickStart.INT.udn
+++ b/Engine/Documentation/Source/Engine/UMG/QuickStart/UMGQuickStart.INT.udn
@@ -44,7 +44,7 @@ Description:Getting started with using Unreal Motion Graphics in Unreal Engine 4
 
 In this **Quick Start Guide**, you will learn how to implement some basic in-game HUD elements and a front end menu using **Unreal Motion Graphics UI Designer (UMG)**. You will learn how to implement Health and Energy Bars as well as a way to track and display the Ammo of a player. You will also learn how to create a Main Menu which can launch into a game, set resolution options or quit the game. From there you will learn how to add an in-game pause menu which can be used to pause the game or return to the Main Menu. 
 
-Over the course of this guide, links to additional documentation will be provided which will go deeper into each section's subject should you have additional questions. While this guide will cover the actions needed to proceed through each step, if you have not worked with the Unreal Editor before, it is recommended that you first look over the **[Level Designer Quick Start](Engine/QuickStart/)** guide to gain a better understanding of general editor usage, terms, and processes. 
+Over the course of this guide, links to additional documentation will be provided which will go deeper into each section's subject should you have additional questions. While this guide will cover the actions needed to proceed through each step, if you have not worked with the Unreal Editor before, it is recommended that you first look over the [Level Designer Quick Start](Engine/QuickStart/) guide to gain a better understanding of general editor usage, terms, and processes. 
 
 
 [OBJECT:TopicButtonList]
