#UEDoc #docs - removed destructible properties that no longer apply. updated doc bold lettering where needed.

diff --git a/Engine/Documentation/Source/Engine/Physics/Collision/Collision.INT.udn b/Engine/Documentation/Source/Engine/Physics/Collision/Collision.INT.udn
--- a/Engine/Documentation/Source/Engine/Physics/Collision/Collision.INT.udn
+++ b/Engine/Documentation/Source/Engine/Physics/Collision/Collision.INT.udn
@@ -25,22 +25,26 @@ Description:All about collision, collision responses, and collision presets.
 
 [TOC(start:2 end:3)]
 
-## Overview
 
-**Collision Responses** and **Trace Responses** form the basis for how Unreal Engine 4 handles collision and ray casting during run time. Every object that can collide gets an **Object Type** and a series of responses that define how it interacts will all other object types. When a collision or overlap event occurs, both (or all) objects involved can be set to affect or be affected by blocking, overlapping, or ignoring each other.
+**Collision Responses** and **Trace Responses** form the basis for how Unreal Engine 4 handles collision and ray casting during run time. Every object that can collide 
+gets an **Object Type** and a series of responses that define how it interacts will all other object types. When a collision or overlap event occurs, both (or all) 
+objects involved can be set to affect or be affected by blocking, overlapping, or ignoring each other.
 
-**Trace Responses** work basically the same way, except the trace (ray cast) itself can be defined as one of the Trace Response types, thusly allowing Actors to either block it or ignore it based on _their_ Trace Responses.
+Trace Responses work basically the same way, except the trace (ray cast) itself can be defined as one of the Trace Response types, thusly allowing Actors to either 
+block it or ignore it based on _their_ Trace Responses.
 
 ## Properties
 
 ![](COL_collisionBasics.png)
 
-In the properties list, there are a large number of physics related key words.  The two that are not really clear are **Simple** and **Complex Collision**. In short, **Simple Collision** shapes are your [boxes, KDOP, spheres, capsules, and convex shapes](Engine/Physics\Collision\Reference), all of which can be generated in the Static Mesh editor. **Complex Collision** is per poly, and it is never used for an Actor that is simulating physics (it will just fall through the world).
+In the properties list, there are a large number of physics related key words.  The two that are not really clear are **Simple** and **Complex Collision**. In short, 
+Simple Collision shapes are your [boxes, KDOP, spheres, capsules, and convex shapes](Engine/Physics\Collision\Reference), all of which can be generated in the Static
+Mesh editor. Complex Collision is per-poly, and it is never used for an Actor that is simulating physics (it will just fall through the world).
 
 
 | **Property** | **Description** |
 | -------- | ----------- |
-| **Simulation Generates Hit Events** | This flag indicates if **Event Hit** and **On Component Hit** should be called for physics collisions involving this PrimitiveComponent. In more simple terms, it lets this object "know" it has actually had a physics collision. However, it is just a notification, and a dynamic object set to collide with the static environment will still do so regardless of this setting; it just will not fire those events. **This is important if you want to create destructible Actors that can shatter upon impact with the ground**.|
+| **Simulation Generates Hit Events** | This flag indicates if **Event Hit** and **On Component Hit** should be called for physics collisions involving this PrimitiveComponent. In more simple terms, it lets this object "know" it has actually had a physics collision. However, it is just a notification, and a dynamic object set to collide with the static environment will still do so regardless of this setting; it just will not fire those events. |
 | **Generate Overlap Events** | If you want an object to generate overlap events, like **Event Actor Begin Overlap** or **Event Actor End Overlap**, then this flag needs to be **true**. Further, for it to be of any use, the object in question will need to be able to overlap another. |
 | [REGION:tablesection] &nbsp; [/REGION]||
 | **Collision Presets** | These are a collection of presets for the **Collision Responses** below. There is an INI file you can edit to change or add profiles. **To edit anything below this point, you need to set this to "Custom..."**|
@@ -54,7 +58,7 @@ In the properties list, there are a large number of physics related key words.
 | **Trace Responses** | These define how an object should behave when interacting with a trace (usually just done with a ray cast). An object can choose to **Block**, **Overlap**, or simply **Ignore** a trace from the given source. There is no difference in the default trace responses except for their names. In code or Blueprint, you can use either for tracing a ray cast from one thing to another. But, for organization of traces, they should be used as their name implies. The default responses for **Trace Responses** are: |
 | ^ | _Visibility_ - Trace from one location to another, if an object that `blocks` this type of trace is in the way, the trace is `blocked`. |
 | ^ | _Camera_ - Exactly the same as _Visibility_, but should be used for setting when tracing from the camera. |
-| **Object Responses** | An object can choose to **Block**, **Overlap**, or simply **Ignore** an object it is interacting with. The default responses for **Object Responses** are: |
+| **Object Responses** | An object can choose to Block, Overlap, or simply Ignore an object it is interacting with. The default responses for **Object Responses** are: |
 | ^ | _WorldStatic_ - Volumes and world geometry should be _WorldStatic_. | 
 | ^ | _WorldDynamic_ - Moving Actors outside of _Pawns_, _PhysicsBodies_, _Vehicles_, and _Destructible_ Actors. Like an elevator. | 
 | ^ | _Pawn_ - Characters. | 
@@ -63,21 +67,17 @@ In the properties list, there are a large number of physics related key words.
 | ^ | _Destructible_ - Destructible Actors. |
 | [REGION:tablesection] &nbsp; [/REGION]||
 | **Use CCD** | Whether or not to use **Continuous Collision Detection** for this component. Increases accuracy of collision detection. |
-| **Always Create Physics State** | Indicates if you would like to create a **physics state** (it's collision properties, mass, if it is sleeping, etc... all go into a **physics state**) all the time (for collision and simulation). If set to _false_, the state will still be generated if collision or simulation are activated on this Actor, say through a Blueprint or code, but there will be some overhead as it has to work the new object into the full physics simulation. Setting this to **true** can improve performance by removing the overhead for calculating the physics state of the object once it actually needs one (on collision or simulation). | 
+| **Always Create Physics State** | Indicates if you would like to create a **physics state** (its collision properties, mass, if it is sleeping, etc... all go into a physics state) all the time (for collision and simulation). If set to _false_, the state will still be generated if collision or simulation are activated on this Actor, say through a Blueprint or code, but there will be some overhead as it has to work the new object into the full physics simulation. Setting this to **true** can improve performance by removing the overhead for calculating the physics state of the object once it actually needs one (on collision or simulation). | 
 | **Multi Body Overlap** | If **true**, this component will generate individual overlaps for each overlapping physics body if it is a multi-body component. Think about skeletal physics assets, with all their individual collision shapes. With this enabled, a character's hand would generate _its own_ overlap event. Thus giving more control over what is reported and what to do in that specific case. |
 | **Check Async Scene on Move** | If this is flagged as **true**, the component will look for collisions in both of the physics scenes (synchronous and asynchronous). The asynchronous scene is primarily used by the broken pieces of destructible Actors. |
-| **Trace Complex on Move** | If this is set to **true**, component sweeps with this component will trace against complex collision during movement. **Complex Collision** is, in short, per face collision, where as **Simple Collision** are your spheres, capsules, boxes, and generated convex shapes. |
+| **Trace Complex on Move** | If this is set to **true**, component sweeps with this component will trace against complex collision during movement. Complex Collision is, in short, per face collision, where as Simple Collision are your spheres, capsules, boxes, and generated convex shapes. |
 | **Return Material on Move ** | Setting this to **true** will return the Physics Material in the **Hit Info**. |
 
-Some components will have additional or fewer options in the collision category, and many times these very specific options will only be shown in the appropriate editor for that component.
+Some components will have additional or fewer options in the collision category, and many times these very specific options will only be shown in the appropriate editor 
+for that component.
 
 |Property | Description |
 | -------- | ----------- |
-| [REGION:tablesection] **Destructible Actor** [/REGION] || 
-| **Collision Complexity** | Collision Trace behavior, by default, it will keep simple (convex) and complex (per-poly) collisions separate. Some objects types such as **Static Meshes** will have this property only listed in their editor and under the **Body Setup** sub category under the **Static Mesh** category. |
-| ^ | _Default_ - Keeps simple collisions as simple, and complex collisions as complex. |
-| ^ | _Use Simple Collision as Complex_ - Will use convex collision for everything. |
-| ^ | _Use Complex Collision as Simple_ - Will use per-poly collision for even simple collision needs. |
 | [REGION:tablesection] **Physics Body in Physics Asset** [/REGION] ||
 | **Collision Response** | Whether or not a specific physics body in a Physics Asset will collide at all. |
 
@@ -85,25 +85,26 @@ Some components will have additional or fewer options in the collision category,
 
 There are a few rules to keep in mind with how collisions are handled:
 
-* **Blocking** will naturally occur between two (or more) Actors set to Block. However, **Simulation Generates Hit Events** needs to be enabled to execute **Event Hit**, which is used in Blueprints, Destructible Actors, Triggers, etc...
+* Blocking will naturally occur between two (or more) Actors set to Block. However, Simulation Generates Hit Events needs to be enabled to execute Event Hit, which is used in Blueprints, Destructible Actors, Triggers, etc...
 * Setting Actors to **Overlap** will often look like they **Ignore** each other, and without **Generate Overlap Events**, they are essentially the same.
 * For two or more simulating objects to block each other, they both need to be set to block their respective object types.
 * For two or more simulating objects: if one is set to overlap an object, and the second object is set to block the other, the overlap will occur but not the block.
-* Overlap events **can** be generated even if an object **Blocks** another, especially if traveling at high speeds.
+* Overlap events can be generated even if an object Blocks another, especially if traveling at high speeds.
 	* It is not recommended for an object to have both collision and overlap events. Though possible, there is much that needs manual handling.
 * If one object is set to ignore and the other is set to overlap, no overlap events will be fired.
-* The **Object Responses** listed are just the defaults and have no actual bearing on what an object _is_. See below for more information on setting up custom [Object Types and Trace Channels](#customobjectchannelsandtracechannels).
+* The Object Responses listed are just the defaults and have no actual bearing on what an object _is_. See below for more information on setting up custom [Object Types and Trace Channels](#customobjectchannelsandtracechannels).
 
 For purposes of testing levels and looking around your worlds:
 
 * The default **Play In Editor** camera is a pawn. Thusly can be blocked by anything set to block pawns.
-* The **Simulate in Editor** camera, before possessing anything, is **not** a pawn. It can freely clip through everything and will not create any collision or overlap events.
+* The **Simulate in Editor** camera, before possessing anything, is not a pawn. It can freely clip through everything and will not create any collision or overlap events.
 
 
 ## Common Collision Interactions
 
 [REGION:note]
-These interactions assume that all the objects have **Collision Enabled** set to **Collision Enabled** so they are set to fully collide with everything. If collision is disabled, it is as if **ignore** has been set for all **Collision Responses**.
+These interactions assume that all the objects have **Collision Enabled** set to **Collision Enabled** so they are set to fully collide with everything. If collision is 
+disabled, it is as if _Ignore_ has been set for all Collision Responses.
 [/REGION]
 
 For the following section, this will be the setup used to explain what is happening:
@@ -139,22 +140,26 @@ In this case, the sphere and the wall will simply collide; no further notificati
 
 ### Collision and Simulation Generates Hit Events
 
-Just collision is useful and in general, the bare minimum for physics interactions, but if you want something to **report** it has collided so a Blueprint or section of code can be triggered:
+Just collision is useful and in general, the bare minimum for physics interactions, but if you want something to "report" it has collided so a Blueprint or section of 
+code can be triggered:
 
 ![](COL_collideEvent.png)(w:940 convert:false)
 
 | Sphere Collision Setup| Wall Collision Setup |
 | --- | ---- |
 | ![](COL_collideEvent_Sphere.png)(w:300) %spacer%|![](COL_collideNoEvent_Box.png)(w:301) %spacer%|
-| As in the example above, the sphere is a **PhysicsBody** and it is set to `block` **WorldDynamic** (which is what the wall is). However, the sphere has also enabled **Simulation Generates Hit Events** so it will trigger an event for itself whenever it collides with something. | The wall is a **WorldDynamic** and is set to `block` **PhysicsBody** Actors (which is what the sphere is). Since the wall is not set to **Simulation Generates Hit Events**, it will not generate an event for itself. |
+| As in the example above, the sphere is a **PhysicsBody** and it is set to `block` **WorldDynamic** (which is what the wall is). However, the sphere has also enabled Simulation Generates Hit Events so it will trigger an event for itself whenever it collides with something. | The wall is a **WorldDynamic** and is set to `block` **PhysicsBody** Actors (which is what the sphere is). Since the wall is not set to Simulation Generates Hit Events, it will not generate an event for itself. |
 
-With the sphere set to **Simulation Generates Hit Events**, the sphere will tell itself that it has had a collision. It will fire off events such as **ReceiveHit** or **OnComponentHit** in the sphere's blueprint. Now if the box had an event for collision, it would not fire because it will never notify itself it has happened.
+With the sphere set to Simulation Generates Hit Events, the sphere will tell itself that it has had a collision. It will fire off events such as **ReceiveHit** or 
+**OnComponentHit** in the sphere's blueprint. Now if the box had an event for collision, it would not fire because it will never notify itself it has happened.
 
-Further, an object that is reporting rigid collisions will report them all and spam reports when it is just sitting on something, so it is best to be careful to filter what it is colliding within its blueprint or in code.
+Further, an object that is reporting rigid collisions will report them all and spam reports when it is just sitting on something, so it is best to be careful to filter 
+what it is colliding within its blueprint or in code.
 
 ### Overlap and Ignore
 
-For all intents and purposes, **Overlap** and **Ignore** work exactly the same **assuming Generate Overlap Events** is disabled. In this case, the sphere is set to overlap or ignore the box:
+For all intents and purposes, **Overlap** and **Ignore** work exactly the same assuming Generate Overlap Events is disabled. In this case, the sphere is set to 
+overlap or ignore the box:
 
 ![](COL_ignore.png)(w:940 convert:false)
 
@@ -164,7 +169,7 @@ For all intents and purposes, **Overlap** and **Ignore** work exactly the same *
 | Sphere Collision Setup| Wall Collision Setup |
 | --- | ---- |
 | ![](COL_OverlapNoEvent_Sphere.png)(w:300) %spacer% |![](COL_collideNoEvent_Box.png)(w:303) %spacer% |
-| Here the sphere is set to `overlap` **WorldDynamic** Actors (like our wall), but it does not have **Generate Overlap Events** enabled. As far as the sphere is concerned, it has not collided or overlapped anything, effectively it has ignored the wall.  | The wall is a **WorldDynamic** and is set to `block` **PhysicsBody** Actors (which is what the sphere is). As stated above, both Actors need to be set to block each other's respective object types. If they do not, they will not collide. |
+| Here the sphere is set to `overlap` **WorldDynamic** Actors (like our wall), but it does not have Generate Overlap Events enabled. As far as the sphere is concerned, it has not collided or overlapped anything, effectively it has ignored the wall.  | The wall is a **WorldDynamic** and is set to `block` **PhysicsBody** Actors (which is what the sphere is). As stated above, both Actors need to be set to block each other's respective object types. If they do not, they will not collide. |
 
 Or:
 
@@ -189,7 +194,8 @@ Unlike collisions that can fire every frame, the overlap events are **ReceiveBeg
  
 ## Custom Object Channels and Trace Channels
 
-There will be times where 6 Object Response Channels and 2 Trace Response Channels simply are not granular enough for what you want to create. This is what the Collision Editor in your Project Settings is for. You can access them from the **Edit Menu** -> **Project Settings** -> **Collision**:
+There will be times where 6 Object Response Channels and 2 Trace Response Channels simply are not granular enough for what you want to create. This is what the Collision 
+Editor in your Project Settings is for. You can access them from the **Edit Menu** -> **Project Settings** -> **Collision**:
 
 ![](COL_projectSettings_1.png)
 
@@ -197,7 +203,8 @@ There will be times where 6 Object Response Channels and 2 Trace Response Channe
 [![](COL_projectSettings_2.png)(w:940)](COL_projectSettings_2.png)
 [/REGION]
 
-From here you can add new Object Response Channels and Trace Response Channels. Click the **New Object Channel...** or **New Trace Channel...** button, provide a name, select a **Default Response**, and click **Accept**.
+From here you can add new Object Response Channels and Trace Response Channels. Click the **New Object Channel...** or **New Trace Channel...** button, provide a name, 
+select a **Default Response**, and click **Accept**.
 
 ![](COL_customChannel.png)
 
@@ -209,7 +216,8 @@ Custom Presets can also be set here by expanding the **Preset** Category and cli
 
 ![](COL_customPresets.png)
 
-From here you can name your preset, enable or disable collision, select your preset's Object Type, and finally define the behavior of each Response Channel for the selected Object Type.
+From here you can name your preset, enable or disable collision, select your preset's Object Type, and finally define the behavior of each Response Channel for the 
+selected Object Type.
 
 
 
