#UEDoc #docs - Need to update images and process.

diff --git a/Engine/Documentation/Source/Engine/Physics/Collision/Reference/Reference.INT.udn b/Engine/Documentation/Source/Engine/Physics/Collision/Reference/Reference.INT.udn
--- a/Engine/Documentation/Source/Engine/Physics/Collision/Reference/Reference.INT.udn
+++ b/Engine/Documentation/Source/Engine/Physics/Collision/Reference/Reference.INT.udn
@@ -27,17 +27,19 @@ Description:Content guide to creating and setting up collision geometry.
 
 ## Overview
 
-As graphical detail has increased, it is important to ensure that collision within the game remains simple. This is because of two reasons; it is much faster to calculate collision against a few polygons and it is a much smoother experience for the player as then the player does not get snagged on a small piece geometry. 
+As graphical detail has increased, it is important to ensure that collision within the game remains simple. This is because of two reasons; it is much faster to calculate collision 
+against a few polygons and it is a much smoother experience for the player as then the player does not get snagged on a small piece geometry. 
 
 [REGION:tip]
-All collision shapes produced in the **Static Mesh Editor** can be moved, rotated, and scaled after they are created. You can create more than one collision shape per Static Mesh so you can get the fidelity you need in your Static Mesh's collision.
+All collision shapes produced in the **Static Mesh Editor** can be moved, rotated, and scaled after they are created. You can create more than one collision shape per Static Mesh so 
+you can get the fidelity you need in your Static Mesh's collision.
 
 This is most useful with the simple collision shapes (box, sphere, and capsule), but can be used with the Convex Hull generator and K-DOP hulls.
 [/REGION]
 
 ## Creating Skeletal Mesh Collision hulls
 
-Please see [PhAT User Guide](Engine/Physics/PhAT) on how to create collision hulls used by Skeletal Meshes.
+Please see **[PhAT User Guide](Engine/Physics/PhAT)** on how to create collision hulls used by Skeletal Meshes.
 
 ## Creating Static Mesh Collision Hulls in the Static Mesh Editor
 
@@ -51,7 +53,9 @@ Please see [PhAT User Guide](Engine/Physics/PhAT) on how to create collision hul
 ![](COLREF_collisionMenu_KDOP.png)
 [/REGION]
 
-There are a series of menu options in the **Static Mesh Editor** under the **Collision** menu entry, named _##DOP_, these are the **K-DOP** simple collision generators. **K-DOP** is a type of bounding volume, which stands for _K discrete oriented polytope_ (where K is the number of axis aligned planes). Basically it takes K axis-aligned planes and pushes them as close to the mesh as it can. The resulting shape is used as a collision hull. In the **Static Mesh Editor** K can be:
+There are a series of menu options in the Static Mesh Editor under the **Collision** menu entry, named _##DOP_, these are the **K-DOP** simple collision generators. 
+**K-DOP** is a type of bounding volume, which stands for _K discrete oriented polytope_ (where K is the number of axis aligned planes). Basically it takes K axis-aligned planes
+and pushes them as close to the mesh as it can. The resulting shape is used as a collision hull. In the Static Mesh Editor K can be:
 
 
 * **10** - Box with 4 edges beveled - you can choose X- Y- or Z-aligned edges.
@@ -69,7 +73,7 @@ See below for an example. This tool is quite handy for packages full of pipes, p
 ![](COLREF_collisionMenu_Auto.png)
 [/REGION]
 
-The **Static Mesh Editor** has the capability to generate collision using a number of convex hulls.
+The Static Mesh Editor has the capability to generate collision using a number of convex hulls.
 
 ![Collision_CreateCollisionHullWithAutoConvex_0.jpg](Collision_CreateCollisionHullWithAutoConvex_0.png)(convert:false)
 
@@ -82,7 +86,7 @@ Selecting **Collision Menu -> Auto Convex Collision** will bring up this panel i
 * **Max Hulls** - Determines the number of hulls to create to best match the shape of the object.
 * **Max Hull Verts** - Maximum number of collision hull vertices. Increasing this value increase how complex the collision hulls are allowed to be.
 * **Apply** - Adds collision to the object based on the settings above.
-* **Defaults** - Resets the values of **Max Hulls** and **Max Hull Verts** to the values seen in the image above.
+* **Defaults** - Resets the values of Max Hulls and Max Hull Verts to the values seen in the image above.
 
 ![Collision_CreateCollisionHullWithAutoConvex_2.jpg](Collision_CreateCollisionHullWithAutoConvex_2.png)(convert:false)
 
@@ -113,16 +117,19 @@ Once a simple collision shape is created, it can be moved, rotated, and scaled i
 
 ### Creating a Collision Hull Using Your 3D Content Creation Package
 
-Please see [FBX Static Mesh Pipeline](Engine/Content/FBX/StaticMeshes#Collision) on how to create Static Mesh collision hulls within a 3D content creation package.
+Please see **[FBX Static Mesh Pipeline](Engine/Content/FBX/StaticMeshes#Collision)** on how to create Static Mesh collision hulls within a 3D content creation package.
 
 
 ## Collision Responses
 
-Unreal Engine 4 is pretty drastically different in terms of how you can have collision objects interact with each other. Now each object is an arbitrary object type (a **Static Mesh** can be in the **WorldDynamic** collision response group for example), that can choose to **ignore**, **block**, or **overlap** an object from another response group.
+Unreal Engine 4 is pretty drastically different in terms of how you can have collision objects interact with each other. Now each object is an arbitrary object type 
+(a Static Mesh can be in the **WorldDynamic** collision response group for example), that can choose to **ignore**, **block**, or **overlap** an object from another response group.
 
-The short version is that instead of an object being "I will collide with all things set to collide" and object can be set to "I will collide with only objects that are _Pawns_ and I will overlap _WorldDynamic_ while ignoring everything else." Further, using that object that will only collide with _Pawns_, and object in the _Pawn_ response group can choose to **ignore** the group our pawn collider exists in.
+The short version is that instead of an object being "I will collide with all things set to collide" and object can be set to "I will collide with only objects that are 
+_Pawns_ and I will overlap _WorldDynamic_ while ignoring everything else." Further, using that object that will only collide with _Pawns_, and object in the _Pawn_ response 
+group can choose to **ignore** the group our pawn collider exists in.
 
-It can get really complicated but the tools at your disposal make it easy to get a handle on. See the [Collision Responses](Engine/Physics\Collision) documentation for more info.
+It can get really complicated but the tools at your disposal make it easy to get a handle on. See the **[Collision Responses](Engine/Physics\Collision)** documentation for more info.
 
 
 ## Reviewing Collision In-Game
