Initial checkin of Mac viewport controls overlay

diff --git a/Engine/Documentation/Source/Shared/CSS/unrealeditor.css b/Engine/Documentation/Source/Shared/CSS/unrealeditor.css
--- a/Engine/Documentation/Source/Shared/CSS/unrealeditor.css
+++ b/Engine/Documentation/Source/Shared/CSS/unrealeditor.css
@@ -135,4 +135,32 @@ BODY
 {
     margin:0;
     padding:0;
+}
+
+.simple-overlay.wide TD.table-cell-content
+{
+    max-width:22em;
+	color:#f0f0f0;
+    /*white-space:nowrap;*/
+}
+
+.simple-overlay TABLE.subtable
+{
+    margin:-4px;
+}
+
+.simple-overlay TABLE.subtable TD:first-child
+{
+    /*text-align:right;*/
+}
+
+.simple-overlay TABLE.subtable TD
+{
+    border-bottom:0px solid #404043;
+    padding:4px;
+}
+
+.simple-overlay TD.table-cell-content TABLE.subtable div
+{
+    margin:3px 0px 0px;
 }
\ No newline at end of file
diff --git a/Engine/Documentation/Source/Shared/Editor/Overlays/ViewportControlsMac.html b/Engine/Documentation/Source/Shared/Editor/Overlays/ViewportControlsMac.html
new file mode 100644
--- /dev/null
+++ b/Engine/Documentation/Source/Shared/Editor/Overlays/ViewportControlsMac.html
@@ -0,0 +1,242 @@
+<?xml version="1.0" encoding="UTF-8" ?>
+<!DOCTYPE HTML PUBLIC "-//W3C//DTD XHTML 1.0 Strict//EN"
+          "http://www.w3.org/TR/xhtml1/DTD/xhtml1-strict.dtd">
+<html xmlns="http://www.w3.org/1999/xhtml" dir="ltr" lang="en-US">
+<head>
+    <title>Viewport Controls</title>
+    <link rel="stylesheet" type="text/css" href="./../../CSS/unrealeditor.css" />
+</head>
+<body>
+    <div class="simple-overlay wide">
+        <table>
+            <tbody>
+                <tr>
+                    <th rowspan="1" colspan="3" align="" class="table-section-heading">
+                        Navigation
+                    </th>
+                </tr>
+                <tr>
+                    <th rowspan="1" colspan="3" align="" class="table-section-subheading">
+                        Perspective Viewport
+                    </th>
+                </tr>
+                <tr>
+                    <td rowspan="1" colspan="1" align="" class="table-cell-heading">
+                        Move forward / backward
+                    </td>
+                    <td rowspan="1" colspan="1" align="" class="table-cell-content">
+                        <table class="subtable">
+                            <tr>
+                                <td>Mouse:</td>
+                                <td>LMB + Drag up / down</td>
+                            </tr>
+                            <tr>
+                                <td>Trackpad:</td>
+                                <td>1 Finger Click + Drag</td>
+                            </tr>
+                        </table>
+                    </td>
+                </tr>
+                <tr>
+                    <td rowspan="1" colspan="1" align="" class="table-cell-heading">
+                        Rotate left / right
+                    </td>
+                    <td rowspan="1" colspan="1" align="" class="table-cell-content">
+                        <table class="subtable">
+                            <tr>
+                                <td>Mouse:</td>
+                                <td>LMB + Drag left / right</td>
+                            </tr>
+                            <tr>
+                                <td>Trackpad:</td>
+                                <td>1 Finger Click + Drag</td>
+                            </tr>
+                        </table>
+                    </td>
+                </tr>
+                <tr>
+                    <td rowspan="1" colspan="1" align="" class="table-cell-heading">
+                        Free Rotate
+                    </td>
+                    <td rowspan="1" colspan="1" align="" class="table-cell-content">
+                        <table class="subtable">
+                            <tr>
+                                <td>Mouse:</td>
+                                <td>
+                                    <div>RMB + Drag</div>
+                                    <div>1 Finger Touch + Drag</div>
+                                </td>
+                            </tr>
+                            <tr>
+                                <td>Trackpad:</td>
+                                <td>
+                                    <div>2 Finger Click + Drag</div>
+                                    <div>2 Finger Touch + Drag</div>
+                                </td>
+                            </tr>
+                        </table>
+                    </td>
+                </tr>
+                <tr>
+                    <td rowspan="1" colspan="1" align="" class="table-cell-heading">
+                        Move up / down
+                    </td>
+                    <td rowspan="1" colspan="1" align="" class="table-cell-content">
+                        LMB + RMB + Drag
+                    </td>
+                </tr>
+                <tr>
+                    <td rowspan="1" colspan="1" align="" class="table-cell-heading">
+                        Zoom in / out
+                    </td>
+                    <td rowspan="1" colspan="1" align="" class="table-cell-content">
+                        2 Finger Pinch
+                    </td>
+                </tr>
+                <tr>
+                    <th rowspan="1" colspan="3" align="" class="table-section-subheading">
+                        Top / Front / Side Viewport
+                    </th>
+                </tr>
+                <tr>
+                    <td rowspan="1" colspan="1" align="" class="table-cell-heading">
+                        Pan
+                    </td>
+                    <td rowspan="1" colspan="1" align="" class="table-cell-content">
+                        <table class="subtable">
+                            <tr>
+                                <td>Mouse:</td>
+                                <td>
+                                    <div>RMB + Drag</div>
+                                    <div>1 Finger Touch + Drag</div>
+                                </td>
+                            </tr>
+                            <tr>
+                                <td>Trackpad:</td>
+                                <td>
+                                    <div>2 Finger Click + Drag</div>
+                                    <div>2 Finger Touch + Drag</div>
+                                </td>
+                            </tr>
+                        </table>
+                    </td>
+                </tr>
+                <tr>
+                    <td rowspan="1" colspan="1" align="" class="table-cell-heading">
+                        Zoom in / out
+                    </td>
+                    <td rowspan="1" colspan="1" align="" class="table-cell-content">
+                        2 Finger Pinch
+                    </td>
+                </tr>
+                <tr>
+                    <th rowspan="1" colspan="3" align="" class="table-section-heading">
+                        Selection
+                    </th>
+                </tr>
+                <tr>
+                    <td rowspan="1" colspan="1" align="" class="table-cell-heading">
+                        Select
+                    </td>
+                    <td rowspan="1" colspan="1" align="" class="table-cell-content">
+                        <table class="subtable">
+                            <tr>
+                                <td>Mouse:</td>
+                                <td>LMB on Actor</td>
+                            </tr>
+                            <tr>
+                                <td>Trackpad:</td>
+                                <td>1 Finger Click on Actor</td>
+                            </tr>
+                        </table>
+                    </td>
+                </tr>
+                <tr>
+                    <td rowspan="1" colspan="1" align="" class="table-cell-heading">
+                        Toggle selection
+                    </td>
+                    <td rowspan="1" colspan="1" align="" class="table-cell-content">
+                        <table class="subtable">
+                            <tr>
+                                <td>Mouse:</td>
+                                <td>Command + LMB on Actor</td>
+                            </tr>
+                            <tr>
+                                <td>Trackpad:</td>
+                                <td>Command + 1 Finger Click on Actor</td>
+                            </tr>
+                        </table>
+                    </td>
+                </tr>
+                <tr>
+                    <td rowspan="1" colspan="1" align="" class="table-cell-heading">
+                        Marquee Selection
+                    </td>
+                    <td rowspan="1" colspan="1" align="" class="table-cell-content">
+                        <table class="subtable">
+                            <tr>
+                                <td rowspan="1" colspan="3" align="" class="table-cell-heading">
+                                    Perspective Viewport
+                                </td>
+                            </tr>
+                            <tr>
+                                <td>Mouse:</td>
+                                <td>Command + Option + LMB + Drag</td>
+                            </tr>
+                            <tr>
+                                <td>Trackpad:</td>
+                                <td>Command + Option + 1 Finger Click + Drag</td>
+                            </tr>
+                            <tr>
+                                <td rowspan="1" colspan="3" align="" class="table-cell-heading">
+                                    Top / Front / Side Viewport
+                                </td>
+                            </tr>
+                            <tr>
+                                <td>Mouse:</td>
+                                <td>LMB + Drag</td>
+                            </tr>
+                            <tr>
+                                <td>Trackpad:</td>
+                                <td>1 Finger Click + Drag</td>
+                            </tr>
+                            </tr>
+                        </table>
+                    </td>
+                </tr>
+                <tr>
+                    <td rowspan="1" colspan="1" align="" class="table-cell-heading">
+                        Clear Selection
+                    </td>
+                    <td rowspan="1" colspan="1" align="" class="table-cell-content">
+                        Esc
+                    </td>
+                </tr>
+                <tr>
+                    <th rowspan="1" colspan="3" align="" class="table-section-heading">
+                        Focusing
+                    </th>
+                </tr>
+                <tr>
+                    <td rowspan="1" colspan="1" align="" class="table-cell-heading">
+                        Focus selected object
+                    </td>
+                    <td rowspan="1" colspan="1" align="" class="table-cell-content">
+                        F
+                    </td>
+                </tr>
+                <tr>
+                    <td rowspan="1" colspan="3" align="" class="table-cell-content">
+                        <div class="width">
+                            <ul>
+                                <li>LMB = Left Mouse Button</li>
+                                <li>RMB = Right Mouse Button</li>
+                            </ul>
+                        </div>
+                    </td>
+                </tr>
+            </tbody>
+        </table>
+    </div>
+</body>
+</html>
\ No newline at end of file
