#UEDoc #docs - updated images and bold lettering where needed

diff --git a/Engine/Documentation/Source/Engine/Physics/Constraints/ConstraintsBlueprints/ConstraintsBlueprints.INT.udn b/Engine/Documentation/Source/Engine/Physics/Constraints/ConstraintsBlueprints/ConstraintsBlueprints.INT.udn
--- a/Engine/Documentation/Source/Engine/Physics/Constraints/ConstraintsBlueprints/ConstraintsBlueprints.INT.udn
+++ b/Engine/Documentation/Source/Engine/Physics/Constraints/ConstraintsBlueprints/ConstraintsBlueprints.INT.udn
@@ -11,12 +11,13 @@ Description:A user guide on using the Physics Constraint Component in Blueprints
 
 ## Overview
 
-Physics Constraint Components operate just like [Physics Constraints Actors](Engine/Physics\Constraints\ConstraintsUserGuide), except they are used in Blueprints or can be created in C++. Given the flexibility of Blueprints and the power of C++, you can constrain just about any Physics Body in your project using Physics Constraint Components.
+Physics Constraint Components operate just like **[Physics Constraints Actors](Engine/Physics\Constraints\ConstraintsUserGuide)**, except they are used in Blueprints or can be 
+created in C++. Given the flexibility of Blueprints and the power of C++, you can constrain just about any Physics Body in your project using Physics Constraint Components.
 
 This document will cover Physics Constraint Component basic creation in a Blueprint.
 
 [REGION:note]
-This document assumes that the user has some knowledge of **Blueprints** and the **Blueprint Editor**. See: [%Engine/Blueprints:title%](Engine/Blueprints) for more information.
+This document assumes that the user has some knowledge of **Blueprints** and the **Blueprint Editor**. See: **[%Engine/Blueprints:title%](Engine/Blueprints)** for more information.
 [/REGION]
 
 ## Usage
@@ -49,7 +50,7 @@ This document assumes that the user has some knowledge of **Blueprints** and the
 	[REGION:raw]
 	![](PhysConBP_ConSettingsVis.png)
 	[/REGION]
-	Want to know what all the properties on a Physics Constraint Component affect? See the [%Engine/Physics\Constraints\ConstraintsReference:title%](Engine/Physics\Constraints\ConstraintsReference) for specific details.
+	Want to know what all the properties on a Physics Constraint Component affect? See the **[%Engine/Physics\Constraints\ConstraintsReference:title%](Engine/Physics\Constraints\ConstraintsReference)** for specific details.
 1. Rotate the Physics Constraint Component as necessary to account for Linear and Angular Limits.
 
 	![](PhysConBP_Rotated.png)
