Removed mouse controls to simplify overlay (ideally need to be able to toggle between mouse/trackpad or detect presence of mouse so we only show relevant controls)

diff --git a/Engine/Documentation/Source/Shared/Editor/Overlays/ViewportControlsMac.html b/Engine/Documentation/Source/Shared/Editor/Overlays/ViewportControlsMac.html
--- a/Engine/Documentation/Source/Shared/Editor/Overlays/ViewportControlsMac.html
+++ b/Engine/Documentation/Source/Shared/Editor/Overlays/ViewportControlsMac.html
@@ -25,16 +25,7 @@
                         Move forward / backward
                     </td>
                     <td rowspan="1" colspan="1" align="" class="table-cell-content">
-                        <table class="subtable">
-                            <tr>
-                                <td>Mouse:</td>
-                                <td>LMB + Drag up / down</td>
-                            </tr>
-                            <tr>
-                                <td>Trackpad:</td>
-                                <td>1 Finger Click + Drag</td>
-                            </tr>
-                        </table>
+                        1 Finger Click + Drag
                     </td>
                 </tr>
                 <tr>
@@ -42,16 +33,7 @@
                         Rotate left / right
                     </td>
                     <td rowspan="1" colspan="1" align="" class="table-cell-content">
-                        <table class="subtable">
-                            <tr>
-                                <td>Mouse:</td>
-                                <td>LMB + Drag left / right</td>
-                            </tr>
-                            <tr>
-                                <td>Trackpad:</td>
-                                <td>1 Finger Click + Drag</td>
-                            </tr>
-                        </table>
+                        1 Finger Click + Drag
                     </td>
                 </tr>
                 <tr>
@@ -59,22 +41,9 @@
                         Free Rotate
                     </td>
                     <td rowspan="1" colspan="1" align="" class="table-cell-content">
-                        <table class="subtable">
-                            <tr>
-                                <td>Mouse:</td>
-                                <td>
-                                    <div>RMB + Drag</div>
-                                    <div>1 Finger Touch + Drag</div>
-                                </td>
-                            </tr>
-                            <tr>
-                                <td>Trackpad:</td>
-                                <td>
-                                    <div>2 Finger Click + Drag</div>
-                                    <div>2 Finger Touch + Drag</div>
-                                </td>
-                            </tr>
-                        </table>
+                        <div>2 Finger Click + Drag</div>
+                        <span style="font-size:9px;">or</span>
+                        <div>2 Finger Touch + Drag</div>
                     </td>
                 </tr>
                 <tr>
@@ -103,22 +72,9 @@
                         Pan
                     </td>
                     <td rowspan="1" colspan="1" align="" class="table-cell-content">
-                        <table class="subtable">
-                            <tr>
-                                <td>Mouse:</td>
-                                <td>
-                                    <div>RMB + Drag</div>
-                                    <div>1 Finger Touch + Drag</div>
-                                </td>
-                            </tr>
-                            <tr>
-                                <td>Trackpad:</td>
-                                <td>
-                                    <div>2 Finger Click + Drag</div>
-                                    <div>2 Finger Touch + Drag</div>
-                                </td>
-                            </tr>
-                        </table>
+                        <div>2 Finger Click + Drag</div>
+                        <span style="font-size:9px;">or</span>
+                        <div>2 Finger Touch + Drag</div>
                     </td>
                 </tr>
                 <tr>
@@ -139,16 +95,7 @@
                         Select
                     </td>
                     <td rowspan="1" colspan="1" align="" class="table-cell-content">
-                        <table class="subtable">
-                            <tr>
-                                <td>Mouse:</td>
-                                <td>LMB on Actor</td>
-                            </tr>
-                            <tr>
-                                <td>Trackpad:</td>
-                                <td>1 Finger Click on Actor</td>
-                            </tr>
-                        </table>
+                        1 Finger Click on Actor
                     </td>
                 </tr>
                 <tr>
@@ -156,16 +103,7 @@
                         Toggle selection
                     </td>
                     <td rowspan="1" colspan="1" align="" class="table-cell-content">
-                        <table class="subtable">
-                            <tr>
-                                <td>Mouse:</td>
-                                <td>Command + LMB on Actor</td>
-                            </tr>
-                            <tr>
-                                <td>Trackpad:</td>
-                                <td>Command + 1 Finger Click on Actor</td>
-                            </tr>
-                        </table>
+                        Cmd + 1 Finger Click on Actor
                     </td>
                 </tr>
                 <tr>
@@ -180,12 +118,7 @@
                                 </td>
                             </tr>
                             <tr>
-                                <td>Mouse:</td>
-                                <td>Command + Option + LMB + Drag</td>
-                            </tr>
-                            <tr>
-                                <td>Trackpad:</td>
-                                <td>Command + Option + 1 Finger Click + Drag</td>
+                                <td>Cmd + Option + 1 Finger Click + Drag</td>
                             </tr>
                             <tr>
                                 <td rowspan="1" colspan="3" align="" class="table-cell-heading">
@@ -193,11 +126,6 @@
                                 </td>
                             </tr>
                             <tr>
-                                <td>Mouse:</td>
-                                <td>LMB + Drag</td>
-                            </tr>
-                            <tr>
-                                <td>Trackpad:</td>
                                 <td>1 Finger Click + Drag</td>
                             </tr>
                             </tr>
