#UEDoc #docs - updated bold lettering where needed.

diff --git a/Engine/Documentation/Source/Engine/Physics/Constraints/ConstraintsReference/Constraints.INT.udn b/Engine/Documentation/Source/Engine/Physics/Constraints/ConstraintsReference/Constraints.INT.udn
--- a/Engine/Documentation/Source/Engine/Physics/Constraints/ConstraintsReference/Constraints.INT.udn
+++ b/Engine/Documentation/Source/Engine/Physics/Constraints/ConstraintsReference/Constraints.INT.udn
@@ -5,10 +5,11 @@ Description:A reference for all the properties available on Physics Constraints
 
 [TOC(start:2)]
 
-## Overview
-
 [EXCERPT:PhysXConstraints_001]
-A constraint is basically a type of joint. It allows you to connect two Actors together (presumably one physically simulating), and also apply limits or forces. Unreal Engine 4 has a very flexible and data-driven constraint system that allows designers to make many different types of joints simply by changing some options. The engine has a few default joint types (ball-and-socket, hinge, prismatic) but they only differ in their settings. You should feel free to start with any joint type and experiment with making your own configurations.
+A constraint is basically a type of joint. It allows you to connect two Actors together (presumably one physically simulating), and also apply limits or forces. Unreal Engine 4 
+has a very flexible and data-driven constraint system that allows designers to make many different types of joints simply by changing some options. The engine has a few default 
+joint types (ball-and-socket, hinge, prismatic) but they only differ in their settings. You should feel free to start with any joint type and experiment with making your own 
+configurations.
 [/EXCERPT:PhysXConstraints_001]
 
 ## Property Reference
@@ -185,7 +186,7 @@ _From left to right (Physics Constraint scaled up 4x): Linear Motions limited to
 
 Physics Constraints can have their 3 rotational axes set to **Locked**, **Free**, or **Limited**. This will form the basis for a joint such as a hinge that can only rotate on one axis for 90 degrees. The location of the Actor defines where the joint is, thusly, it defines the point at which the constrained Actors will rotate from.
 
-Physics Constraints define their Swing Motions as rotations on the Y and Z axes, while Twist is rotation along the X axis **of the Physics Constraint Actor**, not of either of the constrained Actors. So rotating the Physics Constraint Actor can drastically change your results. Swing Limits will be shown as a green cone (or arc if only one is used or the other is limited to 0), and the twist will be shown as an Arc around the X axis.
+Physics Constraints define their Swing Motions as rotations on the Y and Z axes, while Twist is rotation along the X axis of the Physics Constraint Actor, not of either of the constrained Actors. So rotating the Physics Constraint Actor can drastically change your results. Swing Limits will be shown as a green cone (or arc if only one is used or the other is limited to 0), and the twist will be shown as an Arc around the X axis.
 
 [REGION:lightbox]
 [![](PhysCon_Swing1_45.png)(w:300)](PhysCon_Swing1_45.png) [![](PhysCon_Swing2_45.png)(w:300)](PhysCon_Swing2_45.png) [![](PhysCon_Twist_135.png)(w:300)](PhysCon_Twist_135.png) 
