#UEDoc #docs - updated images and bold lettering where needed.

diff --git a/Engine/Documentation/Source/Engine/Physics/Constraints/ConstraintsUserGuide/ConstraintsUserGuide.INT.udn b/Engine/Documentation/Source/Engine/Physics/Constraints/ConstraintsUserGuide/ConstraintsUserGuide.INT.udn
--- a/Engine/Documentation/Source/Engine/Physics/Constraints/ConstraintsUserGuide/ConstraintsUserGuide.INT.udn
+++ b/Engine/Documentation/Source/Engine/Physics/Constraints/ConstraintsUserGuide/ConstraintsUserGuide.INT.udn
@@ -9,17 +9,19 @@ Description:Physical joints used to connect two rigid body simulated objects tog
 
 [TOC(start:2)]
 
-## Overview
 
-If you need to create a swinging chandelier, a water wheel, or just want to keep a Physics Body in a general area, the **Physics Constraint Actor** is what you need. This document will cover the usage of the Physics Constraint Actor.
+If you need to create a swinging chandelier, a water wheel, or just want to keep a Physics Body in a general area, the **Physics Constraint Actor** is what you need. This document 
+will cover the usage of the Physics Constraint Actor.
 
 [EXCERPT:PhysXConstraints_001]
-A constraint is basically a type of joint. It allows you to connect two Actors together (presumably one physically simulating), and also apply limits or forces. Unreal Engine 4 has a very flexible and data-driven constraint system that allows designers to make many different types of joints simply by changing some options. The engine has a few default joint types (ball-and-socket, hinge, prismatic) but they only differ in their settings. You should feel free to start with any joint type and experiment with making your own configurations.
+A constraint is basically a type of joint. It allows you to connect two Actors together (presumably one physically simulating), and also apply limits or forces. Unreal Engine 4 has a 
+very flexible and data-driven constraint system that allows designers to make many different types of joints simply by changing some options. The engine has a few default joint types 
+(ball-and-socket, hinge, prismatic) but they only differ in their settings. You should feel free to start with any joint type and experiment with making your own configurations.
 [/EXCERPT:PhysXConstraints_001]
 
 ## Physics Constraint Actor
 
-1. The **Physics Constraint Actor** can be found in the **All Classes** section of the **Modes** panel's **Place** mode:
+1. The Physics Constraint Actor can be found in the **All Classes** section of the **Modes** panel's **Place** mode:
 	[REGION:raw]
 	![](PhysCon_ContentBrowser.png)(convert:false)
 	[/REGION]
@@ -35,7 +37,7 @@ A constraint is basically a type of joint. It allows you to connect two Actors t
 	[REGION:lightbox raw]
 	[![](PhysCon_SelActor1.png)(h:230)](PhysCon_SelActor1.png) [![](PhysCon_SelActor2.png)(h:230)](PhysCon_SelActor2.png)
 	[/REGION]
-1. Repeat step 4 for **Constraint Actor 2**.
+1. **Repeat** step 4 for **Constraint Actor 2**.
 
 	![](PhysCon_Selected2.png)(convert:false)
 
@@ -45,7 +47,7 @@ A constraint is basically a type of joint. It allows you to connect two Actors t
 
 	_In this case, all three Actors (2 boxes and the Physics Constraint) were selected at the same time, then rotated 90 degrees so the red constrained box was above the blue constrained box. This also rotates the Physics Constraint so its Angular Swing Motions are on the correct axes._
 
-1. Enable Physics on at least one of the Constrained Actors.
+1. **Enable Physics** on at least one of the Constrained Actors.
 	[REGION:raw]
 	![](PhysCon_FullEnablePhys.png)
 	[/REGION]
@@ -72,7 +74,9 @@ A constraint is basically a type of joint. It allows you to connect two Actors t
 
 ### Skeletal Mesh Considerations
 
-If you are constraining a Skeletal Mesh, you will need to provide a **Constraint Bone** name to the appropriate property. A Bone in this case is actually a Physics Body in the Skeletal Mesh's Physics Asset. The reason for specifying this is while the Physics Bodies will be named after the _Skinned Bone_ they are associated with, a Physics Asset does not need a Physics Body for every _Skinned Bone_.
+If you are constraining a Skeletal Mesh, you will need to provide a **Constraint Bone** name to the appropriate property. A Bone in this case is actually a Physics Body in the 
+Skeletal Mesh's Physics Asset. The reason for specifying this is while the Physics Bodies will be named after the _Skinned Bone_ they are associated with, a Physics Asset does not 
+need a Physics Body for every _Skinned Bone_.
 
 [REGION:lightbox]
 [![](PhysCon_SkelCon.png)(w:920)](PhysCon_SkelCon.png)
@@ -80,15 +84,19 @@ If you are constraining a Skeletal Mesh, you will need to provide a **Constraint
 
 ### Components in Actors
 
-If you would like to constrain a specific Component in an Actor, you can provide a **Component Name** in the appropriate property. If the root of the Actor is of a type that can be constrained, it will be the default Component that is constrained. However, if you provide a valid **Component Name** for either Constrained Actor 1 or 2, that Component will become the target of the Physics Constraint. If that Component is a Skeletal Mesh, you must also provide a Bone name in the appropriate property.
+If you would like to constrain a specific Component in an Actor, you can provide a Component Name in the appropriate property. If the root of the Actor is of a type that can be 
+constrained, it will be the default Component that is constrained. However, if you provide a valid Component Name for either Constrained Actor 1 or 2, that Component will become 
+the target of the Physics Constraint. If that Component is a Skeletal Mesh, you must also provide a Bone name in the appropriate property.
 
 [REGION:raw]
 ![](PhysCon_BPClassCon.png)(convert:false)
 [/REGION]
 
-_Functionally these two Actors and Physics Constraints produce the same effect; however, the Blueprint on the right could produce a different effect by making the Point Lights children of the Root. A Character has a Physics Capsule as a root, but also a Skeletal Mesh, both of which can be targets of the Physics Constraint._
+_Functionally these two Actors and Physics Constraints produce the same effect; however, the Blueprint on the right could produce a different effect by making the Point Lights 
+children of the Root. A Character has a Physics Capsule as a root, but also a Skeletal Mesh, both of which can be targets of the Physics Constraint._
 
-Once you have provided a Component name that the Physics Constraint can attach to, a bounding volume will appear around that component. If a bounding volume does not appear, check the name of your Component and make sure it can actually be constrained by the Physics Constraint.
+Once you have provided a Component name that the Physics Constraint can attach to, a bounding volume will appear around that component. If a bounding volume does not appear, check 
+the name of your Component and make sure it can actually be constrained by the Physics Constraint.
 
 
 
