#UE4doc #docs: Removed extraneous bolding as according to the instructional guide. Altered images to match the instructional guide: green -> yellow.

diff --git a/Engine/Documentation/Source/Engine/UMG/QuickStart/3/UMGQS_3.INT.udn b/Engine/Documentation/Source/Engine/UMG/QuickStart/3/UMGQS_3.INT.udn
--- a/Engine/Documentation/Source/Engine/UMG/QuickStart/3/UMGQS_3.INT.udn
+++ b/Engine/Documentation/Source/Engine/UMG/QuickStart/3/UMGQS_3.INT.udn
@@ -90,7 +90,7 @@ First bit of business is to create the layout of our Main Menu which we will do
 
 	![](Hierarchy11.png)
 
-1.	Select the **Options** Vertical Box and set it�s **Visibility** setting to **Hidden**. 
+1.	Select the **Options** Vertical Box and set its **Visibility** setting to **Hidden**. 
 
 	![](Hierarchy12.png)
 
