#UE4doc #docs:
Updated engine version to 4.9
Changed some text formatting
Removed 'Notes' header as it didn't need to preceed a warning

diff --git a/Engine/Documentation/Source/Programming/Automation/TechnicalGuide/AutomationTechnicalGuide.INT.udn b/Engine/Documentation/Source/Programming/Automation/TechnicalGuide/AutomationTechnicalGuide.INT.udn
--- a/Engine/Documentation/Source/Programming/Automation/TechnicalGuide/AutomationTechnicalGuide.INT.udn
+++ b/Engine/Documentation/Source/Programming/Automation/TechnicalGuide/AutomationTechnicalGuide.INT.udn
@@ -2,14 +2,14 @@ Availability:Public
 Title:Automation Technical Guide
 Crumbs: %ROOT%, Programming, Programming/Automation
 Description:Programmer's guide to creating new automation tests.
-Version: 4.5
+Version: 4.9
 
-Automation tests come in two main flavors: **Simple** and **Complex**. Simple Tests are used to describe 
+Automation tests come in two main flavors: Simple and Complex. Simple Tests are used to describe 
 single atomic tests, while Complex Tests are used to run the same code on a number of inputs. 
 
 Simple tests are useful for confirming that specific functionality is operating as expected. These are 
 generally unit tests or feature tests. For example, simple tests can be used to test if the current map 
-loads in _Play In Editor_ or whether text wrapping is working properly in Slate.
+loads in Play In Editor or whether text wrapping is working properly in Slate.
 
 Complex tests are useful for iterating through a collection of items and perform the same functionality 
 on each item. These are generally content stress tests. For instance, loading all maps or compiling all 
@@ -17,8 +17,8 @@ Blueprints would be good fits for complex automation tests.
 
 ## Architecture and Execution
 
-All automation tests derive from the `FAutomationTestBase` class which defines generic functionality for 
-performing a set of commands. The important functions of the `FAutomationTestBase` class used when setting 
+All automation tests derive from the **FAutomationTestBase** class which defines generic functionality for 
+performing a set of commands. The important functions of the **FAutomationTestBase** class used when setting 
 up a new automation test are:
 
 | Member Functions | Description |
@@ -160,7 +160,6 @@ plate code as follows:
 
 An example would be `FSetResTest` in `EngineAutomationTests.cpp`.  
 
-## Notes
 
 [REGION:warning]
 In the editor, loading maps happens immediately. In game, loading a map happens on the next frame and therefore 
