#UE4doc #docs:
Updated Engine version
Changed font settings on multiple terms

diff --git a/Engine/Documentation/Source/Programming/Automation/UserGuide/AutomationUserGuide.INT.udn b/Engine/Documentation/Source/Programming/Automation/UserGuide/AutomationUserGuide.INT.udn
--- a/Engine/Documentation/Source/Programming/Automation/UserGuide/AutomationUserGuide.INT.udn
+++ b/Engine/Documentation/Source/Programming/Automation/UserGuide/AutomationUserGuide.INT.udn
@@ -2,16 +2,16 @@ Availability:Public
 Title: Automation User Guide
 Crumbs: %ROOT%, Programming, Programming/Automation
 Description:Guide to using the automation system for running tests.
-Version: 4.5
+Version: 4.9
 
 Automation tests are configured and run from the **Automation** panel.
 
 ## User Interface
 
-To summon the automation window, you will need to launch the Editor and select **Session Frontend** from the **Window** menu, or run the [**Unreal Frontend**](Engine/Deployment/UnrealFrontend) standalone app.
+To summon the automation window, you will need to launch the Editor and select **Session Frontend** from the **Window** menu, or run the [Unreal Frontend](Engine/Deployment/UnrealFrontend) standalone app.
 
 [REGION:note]
-You can use your local editor as an automation worker by selecting it in the [**Session Browser**](Engine/Deployment/UnrealFrontend#SessionBrowser).
+You can use your local editor as an automation worker by selecting it in the [Session Browser](Engine/Deployment/UnrealFrontend#SessionBrowser).
 To run automation on external devices  game instances, you will need to create a session by running the
  **New Session** wizard, and then select your new session in the **Session Browser** when it has been created, or select an
  existing session that has been discovered in the **Session Browser**.
