#UE4doc #docs:
Removed Overview header
Unbolded a few terms

diff --git a/Engine/Documentation/Source/Programming/Basics/Basics.INT.udn b/Engine/Documentation/Source/Programming/Basics/Basics.INT.udn
--- a/Engine/Documentation/Source/Programming/Basics/Basics.INT.udn
+++ b/Engine/Documentation/Source/Programming/Basics/Basics.INT.udn
@@ -3,7 +3,6 @@ Title:Programming Basics
 Crumbs: %ROOT%, Programming
 Description:Information for programmers developing with Unreal Engine.
 
-## Overview
 
 An individual game is defined by a game project that contains all of the code, content, and settings specific to that particular game. Gameplay code is contained within one or more 
 gameplay modules, and each game project must contain at least one module. The content, such as art assets, sounds, etc., is imported into the editor and saved into 
@@ -42,9 +41,9 @@ See the [](Programming/UnrealArchitecture/Reference) for a complete explanation
 
 ### Class Implementation
 
-At the most fundamental level, an **Actor** is any gameplay Object that you can place in a level. All Actors extend from the `AActor` class, which is the base class of _spawnable_ gameplay Objects.
+At the most fundamental level, an Actor is any gameplay Object that you can place in a level. All Actors extend from the `AActor` class, which is the base class of _spawnable_ gameplay Objects.
 
-Actors can be thought of, in one sense, as containers that hold special types of **Objects** called [Components](Programming/UnrealArchitecture/Actors/Components). For instance, a _CameraActor_ contains a _CameraComponent_.
+Actors can be thought of, in one sense, as containers that hold special types of Objects called [Components](Programming/UnrealArchitecture/Actors/Components). For instance, a _CameraActor_ contains a _CameraComponent_.
 
 ![](camera_actor.png)
 
