#UE4doc #docs:
Removed Overview header
Bolded a few terms
Updated Engine version

diff --git a/Engine/Documentation/Source/Programming/Basics/CommandLineArguments/CommandLineArguments.INT.udn b/Engine/Documentation/Source/Programming/Basics/CommandLineArguments/CommandLineArguments.INT.udn
--- a/Engine/Documentation/Source/Programming/Basics/CommandLineArguments/CommandLineArguments.INT.udn
+++ b/Engine/Documentation/Source/Programming/Basics/CommandLineArguments/CommandLineArguments.INT.udn
@@ -2,16 +2,16 @@ Availability: Public
 Title: Command-Line Arguments
 Crumbs: %ROOT%, Engine
 Description:Collection of arguments that can be passed to the engine's executable to configure options controlling how it runs.
-Version: 4.5
+Version: 4.9
 
 
 [TOC (start:2 end:3)]
 
 
 
-## Overview
 
-Command line arguments are strings of keywords that you can pass when running the executable via the command line or a shortcut to the executable. Their purpose is to customize the manner in which the engine runs to suit the needs of the developer or user. This can be as simple as causing the editor to run instead of the game, or it can be much more involved such as starting up the game with a certain map running at a specified resolution and framerate while dumping out each frame to individual image files. 
+
+**Command-Line Arguments** are strings of keywords that you can pass when running the executable via the command line or a shortcut to the executable. Their purpose is to customize the manner in which the engine runs to suit the needs of the developer or user. This can be as simple as causing the editor to run instead of the game, or it can be much more involved such as starting up the game with a certain map running at a specified resolution and framerate while dumping out each frame to individual image files. 
 
 
 
@@ -65,7 +65,7 @@ A URL consists of two parts: a map name or server IP address and a series of opt
 
 ## Switches
 
-These arguments can be passed to either the game or the editor, depending on the specific keyword and its intended usage. Some arguments are plain switches (i.e., -UNATTENDED) while others are setting switches which are "key=value" pairs (i.e., -LOG=MyLog.txt). These commands are not case sensitive. The syntax for passing plain switches is to preface each argument with a minus ('-') and then the argument immediately afterward. Setting switches need no leading '-', with the exception of the server switches.
+These arguments can be passed to either the game or the editor, depending on the specific keyword and its intended usage. Some arguments are plain switches (i.e., -UNATTENDED) while others are setting switches which are "key=value" pairs (i.e., -LOG=MyLog.txt). These commands are not case sensitive. The syntax for passing plain switches is to preface each argument with a minus '-' and then the argument immediately afterward. Setting switches need no leading '-', with the exception of the server switches.
 
 *Example:*
 
@@ -85,7 +85,7 @@ These arguments can be passed to either the game or the editor, depending on the
 * `AutomatedMapBuild`: Perform an automated build of a specified map.
 * `BIASCOMPRESSIONFORSIZE`: Override compression settings with respect to size.
 * `BUILDMACHINE`: Set as build machine. Used for deciding if debug output is enabled.
-* `BULKIMPORTINGSOUNDS`: Use when importing sounds in bulk. (Content Browser specific)
+* `BULKIMPORTINGSOUNDS`: Use when importing sounds in bulk. (**Content Browser** specific)
 * `CHECK_NATIVE_CLASS_SIZES`: Enable checking of native class sizes. Note: Native classes on console platforms will cause native class size checks to fail even though they are assumed to be correct.
 * `CODERMODE`: Enables Coder mode.
 * `COMPATSCALE`: Set compatibility settings manually to override PCCompat tool settings.
@@ -109,7 +109,7 @@ These arguments can be passed to either the game or the editor, depending on the
 * `INCLUDEUTGAMECONTENT`: Set UTGameContent packages to be loaded.
 * `INSTALLED`: For development purposes, run the game as if installed.
 * `INSTALLFW` / `UNINSTALLFW`: Set whether the handling of the firewall integration should be performed.
-* `INSTALLGE`: Add the game to the Game Explorer.
+* `INSTALLGE`: Add the game to the **Game Explorer**.
 * `CULTUREFORCOOKING`: Set language to be used for cooking.
 * `LIGHTMASSDEBUG`: Launch lightmass manually with -debug and allows lightmass to be executed multiple times.
 * `LIGHTMASSSTATS`: Force all lightmass agents to report detailed stats to the log.
@@ -137,7 +137,7 @@ These arguments can be passed to either the game or the editor, depending on the
 * `TRACEANIMUSAGE`: Trace animation usage.
 * `TREATLOADWARNINGSASERRORS`: Force load warnings to be treated as errors. 
 * `UNATTENDED`: Set as unattended. Disable anything requiring feedback from user.
-* `UNINSTALLGE`: Remove the game from the Game Explorer.
+* `UNINSTALLGE`: Remove the game from the **Game Explorer**.
 * `USEUNPUBLISHED`: Force packages in the Unpublished folder to be used. (deprecated?)
 * `VADEBUG`: Use the Visual Studio debugger interface.
 * `VERBOSE`: Set script compiler to use verbose output.
