#UE4 Docs: Replaced all the .GIF's with Youtube videos.

diff --git a/Engine/Documentation/Source/Engine/Rendering/Materials/HowTo/Placing_Expressions_Functions/Placing_Expressions_Functions.INT.udn b/Engine/Documentation/Source/Engine/Rendering/Materials/HowTo/Placing_Expressions_Functions/Placing_Expressions_Functions.INT.udn
--- a/Engine/Documentation/Source/Engine/Rendering/Materials/HowTo/Placing_Expressions_Functions/Placing_Expressions_Functions.INT.udn
+++ b/Engine/Documentation/Source/Engine/Rendering/Materials/HowTo/Placing_Expressions_Functions/Placing_Expressions_Functions.INT.udn
@@ -9,15 +9,14 @@ Related: Engine/UI/KeyBindings
 [TOC(start:2)]
 
 
-
 Material Expressions and Material Functions are the building blocks used to create fully functional Materials inside Unreal Engine 4.
 Each expression or function is a self-contained black box that can run little snippets of code on their inputs and output the results. 
 In this How To guide, we will be learning the various ways that we can place Material Expressions and Material Functions inside of the Material Graph. 
 
 ## Material Expressions
 
-Material Expression nodes are tiny snippets of HLSL code that perform very specific tasks inside of a Material. 
-Materials are built using combinations of Material Expression nodes are combined together to do certain tasks. 
+Material Expression nodes contain tiny snippets of HLSL code that perform very specific tasks inside of a Material. 
+Materials are built using combinations of Material Expression nodes that are combined together to do certain tasks. 
 For example, if you want a Texture to move inside of a Material, you can use the Panner Material Expression node. 
 This node contains specific HLSL code that tells the Material to move the UV coordinates of the Texture that it is applied to in a specific direction, making it look like it is moving. 
 As you will soon see,  Material Expression nodes are key to building and interacting with the Materials you create.
@@ -41,19 +40,49 @@ You can place Material Expression into a Material Graph in the following three w
 
 You can drag and drop any Material Expression right form the **Palette** menu directly into the Material graph. To do this, search for the Expression that you want and then click on it with the **Left Mouse Button** and then drag it into the Material Graph. When you release the **Left Mouse Button**, the Material Expression will be placed in the Material graph.
 
-![](Placing_Functions_From_Palette.gif)(convert:false)
+[OBJECT:EmbeddedVideo]
+	[PARAMLITERAL:width]
+	640
+	[/PARAMLITERAL]
+	[PARAMLITERAL:height]
+	360
+	[/PARAMLITERAL]
+	[PARAMLITERAL:videoid]
+	7xw6U6OQy_8
+	[/PARAMLITERAL]
+[/OBJECT]
 
 ### Right-Click Search
 
 You can also add Material Expressions by **Right-Clicking** anywhere in the Material graph and then selecting the Material Expression you want from the menu context menu by clicking on it with the **Left Mouse Button** to place it in the Material graph. You can also use the keyboard to search for the Material Expression you want. When you find the one you are looking for, press the **Enter** key to place the Expression node in the Material Graph.
 
-![](Placing_Functions_From_RightClick.gif)(convert:false)
+[OBJECT:EmbeddedVideo]
+	[PARAMLITERAL:width]
+	640
+	[/PARAMLITERAL]
+	[PARAMLITERAL:height]
+	360
+	[/PARAMLITERAL]
+	[PARAMLITERAL:videoid]
+	AQUNvcgXc5I
+	[/PARAMLITERAL]
+[/OBJECT]
 
 ### Keyboard Shortcuts
 
 There are a number of Keyboard shortcuts that you can use to help you quickly place down Material Expressions that are most commonly used. You can change what key spawns what Material Expression by going to **Edit -> Editor Preferences -> Keyboard Shortcuts -> Material Editor Spawn Nodes**.
 
-![](Placing_Functions_Setting_Shortcuts.gif)(convert:false)
+[OBJECT:EmbeddedVideo]
+	[PARAMLITERAL:width]
+	640
+	[/PARAMLITERAL]
+	[PARAMLITERAL:height]
+	360
+	[/PARAMLITERAL]
+	[PARAMLITERAL:videoid]
+	bhrTE-tAqQo
+	[/PARAMLITERAL]
+[/OBJECT]
 
 Here is a list of the current default keyboard shortcuts and what Material Expression will be created when the keys are pressed.  
 
@@ -89,7 +118,17 @@ Material Functions can be placed inside of a Material Graph in the exact same ma
 
 You can drag and drop any Material Function right from the **Palette** menu directly into the Material graph. To do this, search for the Function that you want and then click on it with the **Left Mouse Button** and drag it into the Material Graph. When you release the **Left Mouse Button**, the Material Expression will be placed in the Material graph.
 
-![](Placing_Functions_Placing_MF_Palette.gif)(convert:false)
+[OBJECT:EmbeddedVideo]
+	[PARAMLITERAL:width]
+	640
+	[/PARAMLITERAL]
+	[PARAMLITERAL:height]
+	360
+	[/PARAMLITERAL]
+	[PARAMLITERAL:videoid]
+	ORfUnnwXdE8
+	[/PARAMLITERAL]
+[/OBJECT]
 
 [region:tip]
 The Palette offers the ability to filter your search so that you can limit your search results to either Expressions, Functions, or both. To set your search filter settings, use the **Left Mouse Button** to click on the **Category** button and then select what you wish to search for from the drop down menu. 
@@ -106,7 +145,17 @@ You can also add Material Functions by **Right-Clicking** anywhere in the Materi
 The **Right-Click** menu cannot filter by Functions or Expression and displays them both when you are searching.
 [/region]
 
-![](Placing_Functions_Function_Search.gif)(convert:false)
+[OBJECT:EmbeddedVideo]
+	[PARAMLITERAL:width]
+	640
+	[/PARAMLITERAL]
+	[PARAMLITERAL:height]
+	360
+	[/PARAMLITERAL]
+	[PARAMLITERAL:videoid]
+	dkdsUXyThiY
+	[/PARAMLITERAL]
+[/OBJECT]
 
 [region:tip]
 
@@ -119,15 +168,34 @@ Material Function nodes will always have the color Blue at the top. You can also
 
 One thing that is exclusive to Material Functions is that you also have the ability to Drag and Drop them from the **Content Browser** into your Material graph. To do this, find the Material Function you want to use in the **Content Browser** then click on it with the **Left Mouse Button** and drag it into your Material graph. Once your Mouse is over the Material graph, release the **Left Mouse Button** to add the Material Function you selected.
 
-![](Placing_Functions_From_Content_Browser.gif)(convert:false)
-
+[OBJECT:EmbeddedVideo]
+	[PARAMLITERAL:width]
+	640
+	[/PARAMLITERAL]
+	[PARAMLITERAL:height]
+	360
+	[/PARAMLITERAL]
+	[PARAMLITERAL:videoid]
+	LcvUYit_tzk
+	[/PARAMLITERAL]
+[/OBJECT]
 [region:note]
 To find the Material Functions in the **Content Browser**, you will need to look in the Engine folder. 
 However, by default the Engine folder is disabled and to see it you will need to enable it. 
 To do this, use the **Left Mouse Button** to click on to View Options that is on the lower right hand corner of the **Content Browser**. 
 Then from the pop-up menu, select the **Show Engine Content** option to display the content in the Engine Folder.  
 
-![](Placing_Functions_Show_Engine_Content.gif)(convert:false)
+[OBJECT:EmbeddedVideo]
+	[PARAMLITERAL:width]
+	640
+	[/PARAMLITERAL]
+	[PARAMLITERAL:height]
+	360
+	[/PARAMLITERAL]
+	[PARAMLITERAL:videoid]
+	XtfxWmM6U9E
+	[/PARAMLITERAL]
+[/OBJECT]
 
 [/region]
 
