#UE4doc #docs:
Removed Overview header
Updated Engine version
Changed font settings for multiple terms/filenames

diff --git a/Engine/Documentation/Source/Programming/Basics/ConfigurationFiles/ConfigurationFiles.INT.udn b/Engine/Documentation/Source/Programming/Basics/ConfigurationFiles/ConfigurationFiles.INT.udn
--- a/Engine/Documentation/Source/Programming/Basics/ConfigurationFiles/ConfigurationFiles.INT.udn
+++ b/Engine/Documentation/Source/Programming/Basics/ConfigurationFiles/ConfigurationFiles.INT.udn
@@ -2,23 +2,21 @@ Availability: Public
 Title: Configuration Files 
 Crumbs: %ROOT%, Engine, Programming, Programming/Basics
 Description:Text files containing property settings for configuring gameplay or engine behavior.
-Version: 4.5
+Version: 4.9
 
 
 [TOC (start:2 end:3)]
 
-## Overview
-
 Configuration files can be used to set values for properties that will be initialized when the project is loaded.  Configuration is determined by key-value pairs, arranged in sections. 
 One or more values can be associated with a given key.
 
-Engine configuration files are used for object and variable default values. User input configuration can be used for key bindings. By default, `DefaultEngine.ini` and `DefaultGame.ini` 
-are created when you create a new blank project with the **Project Wizard**. New projects that begin with templates may also generate `DefaultEditor.ini` and `DefaultInput.ini` configuration files
+Engine configuration files are used for object and variable default values. User input configuration can be used for key bindings. By default, DefaultEngine.ini and DefaultGame.ini 
+are created when you create a new blank project with the **Project Wizard**. New projects that begin with templates may also generate DefaultEditor.ini and DefaultInput.ini configuration files
 if needed.  
 
 The `SaveConfig()` function can be called on a class with the  [](Programming\UnrealArchitecture\Reference\Classes\Specifiers\Config) class specifier.  This saves any properties marked with the 
-[](Programming\UnrealArchitecture\Reference\Properties\Specifiers\Config) property specifier to the appropriate configuration file. Generally, variables saved by `SaveConfig()` are in a section title
-which follows the format **[(package).(classname)]**.  For instance, the **[/Script/Engine.Engine]** section in `DefaultEngine.ini` points to the Engine class stored within the Engine package.
+`Config` property specifier to the appropriate configuration file. Generally, variables saved by **SaveConfig()** are in a section title
+which follows the format `[(package).(classname)]`.  For instance, the `[/Script/Engine.Engine]` section in DefaultEngine.ini points to the Engine class stored within the Engine package.
 There are some exceptions with hard-coded section names.
 
 A number of the settings previously available only through editing configuration files are available in Unreal Editor in the [](Engine\UI\ProjectSettings) editor.
@@ -32,7 +30,7 @@ In order to specify which variables should be read in from configuration files,
 
 [REGION:tip]
 	Note that a category (i.e. Game) must be supplied for the `Config` specifier. This determines which configuration files the class's variables are read from and saved to. All possible categories are 
-	defined in `FConfigCacheIni`. For a list of all configuration file categories, see [Configuration Categories](#configurationcategories).
+	defined in FConfigCache.ini. For a list of all configuration file categories, see [Configuration Categories](#configurationcategories).
 [/REGION]
 
 Decorating a class with the `Config` specifier just indicates that that class can have variables read in from configuration files, and specifies which files the configurations are read from. To
@@ -48,7 +46,7 @@ specify a particular variable as read in and saved to a configuration file, the
 
 	};
 
-No category is supplied to the property's `Config` specifier. The `ExampleVariable` property can now be read from any `Game` configuration file in the [configuration file hierarchy](#filehierarchy),
+No category is supplied to the property's `Config` specifier. The `ExampleVariable` property can now be read from any Game configuration file in the [configuration file hierarchy](#filehierarchy),
 as long as the information is specified with the following syntax:
 
 	[/Script/ModuleName.ExampleClass]
@@ -58,7 +56,7 @@ as long as the information is specified with the following syntax:
 
 The `Config` `UCLASS` and `UPROPERTY` specifiers are inherited.  This means that a child class can either read in or save out all variables specified as `Config` in the parent class, and they
 will be in the same configuration file category. The variables will all be under a section title with the child class's name. For example, the configuration file information for a `ChildExampleClass`
-which inherits from the `ExampleClass` above would look like the following lines, and be saved in the same `Game` configuration files.
+which inherits from the `ExampleClass` above would look like the following lines, and be saved in the same Game configuration files.
 
 	[/Script/ModuleName.ChildExampleClass]
 	ExampleVariable=0.0f
@@ -67,7 +65,7 @@ which inherits from the `ExampleClass` above would look like the following lines
 ### Per-Instance Configurations
 
 Unreal Engine 4 has the ability to save the configuration of an object to any desired configuration file. If the `PerObjectConfig` specifier is used in the `UCLASS` macro, 
-configuration information for this class will be stored per instance, where each instance has a section in the `.ini` file named after the object in the format `[ObjectName ClassName]`. 
+configuration information for this class will be stored per instance, where each instance has a section in the .ini file named after the object in the format `[ObjectName ClassName]`. 
 This keyword is propagated to child classes.
 
 
@@ -91,22 +89,22 @@ Each configuration category has its own hierarchy of files which specify engine-
 
 ### File Hierarchy
 
-The configuration file hierarchy is read in starting with `Base.ini`, with values in later files in the hierarchy overriding earlier values. All files in the `Engine` folder will be applied
-to all projects, while project-specific settings should be in files in the project directory.  Finally, all  project-specific and platform-specific differences are saved out to `[ProjectDirectory]/Saved/Config/[Platform]/[Category].ini`
+The configuration file hierarchy is read in starting with Base.ini, with values in later files in the hierarchy overriding earlier values. All files in the Engine folder will be applied
+to all projects, while project-specific settings should be in files in the project directory.  Finally, all  project-specific and platform-specific differences are saved out to [ProjectDirectory]/Saved/Config/[Platform]/[Category].ini
 
-The below file hierarchy example is for the `Engine` category of configuration files.
+The below file hierarchy example is for the Engine category of configuration files.
 
-1. `Engine/Config/Base.ini`
+1. Engine/Config/Base.ini
 	[REGION:tip]
-	`Base.ini` is usually empty.
+	Base.ini is usually empty.
 	[/REGION]
-1. `Engine/Config/BaseEngine.ini`
-1. `Engine/Config/[Platform]/[Platform]Engine.ini`
-1. `[ProjectDirectory]/Config/DefaultEngine.ini`
-1. `[ProjectDirectory]/Config/[Platform]/[Platform]Engine.ini`
-1. `[ProjectDirectory]/Saved/Config/[Platform]/Engine.ini`
+1. Engine/Config/BaseEngine.ini
+1. Engine/Config/[Platform]/[Platform]Engine.ini
+1. [ProjectDirectory]/Config/DefaultEngine.ini
+1. [ProjectDirectory]/Config/[Platform]/[Platform]Engine.ini
+1. [ProjectDirectory]/Saved/Config/[Platform]/Engine.ini
 	[REGION:tip]
-	The configuration file in the `Saved` directory only stores the project-specific and platform-specific differences in the stack of configuration files.
+	The configuration file in the Saved directory only stores the project-specific and platform-specific differences in the stack of configuration files.
 	[/REGION]
 	
 
@@ -132,7 +130,7 @@ Typical configuration files consist of sections of key-value pairs, arranged as
 * *!* - Removes a property; but you do not have to have an exact match, just the name of the property.
 
 [REGION:note]
-**Note**: **.* is like *+** except it will potentially add a duplicate line. This is useful for the bindings (as seen in `DefaultInput.ini`), for instance, where the bottom-most binding takes effect, so if you add something like:
+**Note**: . is like + except it will potentially add a duplicate line. This is useful for the bindings (as seen in DefaultInput.ini), for instance, where the bottom-most binding takes effect, so if you add something like:
     
     [/Script/Engine.PlayerInput]
     Bindings=(Name="Q",Command="Foo")
