#UE4doc #docs:
Bolded a term
Changed Engine version

diff --git a/Engine/Documentation/Source/Programming/BlueprintFunctionLibraries/BlueprintFunctionLibraries.INT.udn b/Engine/Documentation/Source/Programming/BlueprintFunctionLibraries/BlueprintFunctionLibraries.INT.udn
--- a/Engine/Documentation/Source/Programming/BlueprintFunctionLibraries/BlueprintFunctionLibraries.INT.udn
+++ b/Engine/Documentation/Source/Programming/BlueprintFunctionLibraries/BlueprintFunctionLibraries.INT.udn
@@ -2,9 +2,9 @@ Title: Blueprint Function Libraries
 Crumbs:
 Description:
 Availability:Public
-Version:4.7
+Version:4.9
 
-Often during the development you discover the need for a set of functions that make development easier for your project. These functions typically are stateless and reused across a variety of gameplay code. This same need often presents itself when building with Blueprints. In a previous section of this manual, you saw how to expose your game objects’ functions and properties to Blueprints. However, if you want to expose shared utility functions, you don’t necessarily want to tie them to a specific gameplay object type. For those cases, we have Blueprint Function Libraries.
+Often during the development you discover the need for a set of functions that make development easier for your project. These functions typically are stateless and reused across a variety of gameplay code. This same need often presents itself when building with Blueprints. In a previous section of this manual, you saw how to expose your game objects’ functions and properties to Blueprints. However, if you want to expose shared utility functions, you don’t necessarily want to tie them to a specific gameplay object type. For those cases, we have **Blueprint Function Libraries**.
 
 Blueprint Function Libraries are a collection of static functions that provide utility functionality not tied to a particular gameplay object. These libraries can be grouped into logical function sets, e.g. AI Blueprint Library, or contain utility functions that provide access to many different functional areas, e.g. System Blueprint Library.
 
