#UE4doc #docs:
Removed "Compiling Source" section and replaced with "Next Steps" section.

diff --git a/Engine/Documentation/Source/GettingStarted/DownloadingUnrealEngine/GitHub/GitHub.INT.udn b/Engine/Documentation/Source/GettingStarted/DownloadingUnrealEngine/GitHub/GitHub.INT.udn
--- a/Engine/Documentation/Source/GettingStarted/DownloadingUnrealEngine/GitHub/GitHub.INT.udn
+++ b/Engine/Documentation/Source/GettingStarted/DownloadingUnrealEngine/GitHub/GitHub.INT.udn
@@ -92,8 +92,6 @@ After unziping the source code, you should have a file named `Setup.bat`. Runnin
 In version 4.5.1 and prior, the `Setup.bat` does not exist and the depencies must be downloaded separately. They can be found on the [releases](https://github.com/EpicGames/UnrealEngine/releases) page listed as "Required_1of2" and "Required_2of2". These must be extracted into the Unreal Engine folder.
 [/REGION]
 
-## Compiling Source
+## Next Steps
 
-After the `Setup.bat` has completed, run `GenerateProjectFiles.bat` to generate the Visual Studio / XCode project files.
-
-Once this has completed, open the project file entitled `UE4.sln` or `UE4.xcodeproj` and **Ctrl+Shift+B** once it is done loading to begin compiling the engine. This will take some time.
+After installing the dependencies, you're ready to [build the engine](https://docs.unrealengine.com/latest/INT/Programming/Development/BuildingUnrealEngine/index.html).
