#UE4doc #docs: unbolded hyperlinks according to the instructional guide.

diff --git a/Engine/Documentation/Source/Engine/UMG/QuickStart/6/UMGQS_6.INT.udn b/Engine/Documentation/Source/Engine/UMG/QuickStart/6/UMGQS_6.INT.udn
--- a/Engine/Documentation/Source/Engine/UMG/QuickStart/6/UMGQS_6.INT.udn
+++ b/Engine/Documentation/Source/Engine/UMG/QuickStart/6/UMGQS_6.INT.udn
@@ -23,13 +23,13 @@ Description:The final step in the UMG Quick Start is to provide additional links
 
 ---
 
-For more information on the topics covered in this Quick Start Guide, and across the entire editor, see the [**Unreal Editor Manual**](Engine).
+For more information on the topics covered in this Quick Start Guide, and across the entire editor, see the [Unreal Editor Manual](Engine).
 
 As for coverage related to the topics in this guide:
 
-* For more information on Unreal Motion Graphics, see: [**UMG Overview**](Engine/UMG)
+* For more information on Unreal Motion Graphics, see: [UMG Overview](Engine/UMG)
 
-* For more information on Blueprints, see: [**Blueprints Landing Page**](Engine\Blueprints)
+* For more information on Blueprints, see: [Blueprints Landing Page](Engine\Blueprints)
 
 ---
 
