#UE4doc #docs: ##Overview removed as according to the instructional guide. Altered images to match the instructional guide: green -> yellow.

diff --git a/Engine/Documentation/Source/Engine/UMG/UserGuide/Animation/Animation.INT.udn b/Engine/Documentation/Source/Engine/UMG/UserGuide/Animation/Animation.INT.udn
--- a/Engine/Documentation/Source/Engine/UMG/UserGuide/Animation/Animation.INT.udn
+++ b/Engine/Documentation/Source/Engine/UMG/UserGuide/Animation/Animation.INT.udn
@@ -10,9 +10,6 @@ Description:How to create animated UI elements in UMG.
 
 [TOC(start:2 end:2)]
 
-
-## Overview
-
 Located along the bottom of the **Widget Blueprint Editor** are two windows which allow you to implement and control animations for your UI Widgets. The first, the **Animations** window, allows you to create essentially animation tracks which are used to drive the animation of your Widgets. The second, the **Timeline** window, is how an animation is applied to a Widget over time which is done by placing **Keyframes** at specified times then defining how the attached Widget should appear at that Keyframe (this could be size, shape, location, or even coloring options). 
 
 
