#UE4doc #docs:
Changed font settings on multiple terms
Updated Engine version

diff --git a/Engine/Documentation/Source/Programming/Development/BuildingUnrealEngine/BuildingUnrealEngine.INT.udn b/Engine/Documentation/Source/Programming/Development/BuildingUnrealEngine/BuildingUnrealEngine.INT.udn
--- a/Engine/Documentation/Source/Programming/Development/BuildingUnrealEngine/BuildingUnrealEngine.INT.udn
+++ b/Engine/Documentation/Source/Programming/Development/BuildingUnrealEngine/BuildingUnrealEngine.INT.udn
@@ -2,7 +2,7 @@ Availability:Public
 Title:Building Unreal Engine
 Crumbs: %ROOT%, Programming, Programming/Development
 Description:Compiling the Unreal Engine.
-Version: 4.5
+Version: 4.9
 
 %Globals:OSSelect%
 
@@ -28,17 +28,17 @@ Version: 4.5
 
 		Before you dive right in and start getting your hands dirty with the engine, there are a few basic setup steps to perform:
 
-		1.	Run the `GenerateProjectFiles.bat` batch file located in the root directory.
+		1.	Run the **GenerateProjectFiles.bat** batch file located in the root directory.
 
 			[REGION:warning]
 				All of the project files and even the UE4 solution are intermediate files - located in `[UE4Root]\Engine\Intermediate\ProjectFiles`. You must generate them each time you sync a new build to ensure they exist and are up to date. This also means if you delete your `Intermediate` folder, you must regenerate the project files using the above-mentioned batch file.
 			[/REGION]
 
-		1.	Open the `UE4.sln` Visual Studio solution located in the root directory. This is the main Visual Studio solution for the entire engine. The projects for the engine core and games are contained within this solution. 
+		1.	Open the **UE4.sln** Visual Studio solution located in the root directory. This is the main Visual Studio solution for the entire engine. The projects for the engine core and games are contained within this solution. 
 
 		1.	**Right-click** the **UE4** project and choose **Set As StartUp Project**.
 
-		1.	Save the `UE4.sln`. Make sure to save the changes you have made so the next time you open the solution you will not have to go through the process all over again.
+		1.	Save the **UE4.sln**. Make sure to save the changes you have made so the next time you open the solution you will not have to go through the process all over again.
 	[/PARAM]
 [/OBJECT]
 [OBJECT:ToggleButtonContent]
@@ -52,16 +52,16 @@ Version: 4.5
 	[/PARAMLITERAL]
 	[PARAM:content]
 		[REGION:note]
-		Please make sure you are running Xcode 5.1 before proceeding.
+		Please make sure you are running Xcode 6.1 before proceeding.
 		[/REGION]
 
 		Before you dive right in and start getting your hands dirty with the engine, you will need project files in order to compile: 
 		
-		1. Run `GenerateProjectFiles.command` located in the root directory. It should take less than a minute to complete. 
+		1. Run **GenerateProjectFiles.command** located in the root directory. It should take less than a minute to complete. 
 		
 		1. Close the Terminal window once the operation has completed.
 
-		1. Open the `UE4.xcodeproj` Xcode project located in the root directory. This is the main project for the entire engine. All of the projects for the engine core, tools, and games are contained within this project.
+		1. Open the **UE4.xcodeproj** Xcode project located in the root directory. This is the main project for the entire engine. All of the projects for the engine core, tools, and games are contained within this project.
 
 		1. in Xcode, choose **UE4Editor - Mac** and **My Mac 64-bit** for the active scheme.
 
@@ -71,7 +71,7 @@ Version: 4.5
 
 ## Compiling Unreal Engine
 
-Unreal Engine 4 uses a custom building method via the UnrealBuildTool. This tool handles all the complex aspects of compiling the engine and all of its parts. 
+Unreal Engine 4 (UE4) uses a custom building method via the UnrealBuildTool. This tool handles all the complex aspects of compiling the engine and all of its parts. 
 
 ### Build Configuration
 
@@ -88,7 +88,7 @@ Unreal Engine 4 uses a custom building method via the UnrealBuildTool. This tool
 	[PARAM:content]
 		Not all configurations will appear by default. Edit the [GenerateProjectFiles.bat](Programming/UnrealBuildSystem/ProjectFileGenerator) file in the main UE4 directory to generate the project files for additional configurations as needed.
 
-		Unreal Engine 4 and the UnrealBuildTool use different build configurations to determine how the engine is compiled. Which configuration you use will be determined by the purposes of the build you want to create.
+		UE4 and the UnrealBuildTool use different build configurations to determine how the engine is compiled. Which configuration you use will be determined by the purposes of the build you want to create.
 
 		The available configurations:
 
@@ -127,7 +127,7 @@ Unreal Engine 4 uses a custom building method via the UnrealBuildTool. This tool
 
 ### Target Platforms
 
-Unreal Engine 4 currently supports the following platforms:
+UE4 currently supports the following platforms:
 
 [OBJECT:ToggleButtonContent]
 	[PARAMLITERAL:category]
diff --git a/Engine/Documentation/Source/Programming/Development/CodingStandard/CodingStandard.INT.udn b/Engine/Documentation/Source/Programming/Development/CodingStandard/CodingStandard.INT.udn
--- a/Engine/Documentation/Source/Programming/Development/CodingStandard/CodingStandard.INT.udn
+++ b/Engine/Documentation/Source/Programming/Development/CodingStandard/CodingStandard.INT.udn
@@ -2,12 +2,11 @@ Availability:Public
 Title:Coding Standard
 Crumbs:%ROOT%, Programming, Programming\Development
 Description:Standards and conventions used by Epic Games in the Unreal Engine 4 codebase.
-Version: 4.5
+Version: 4.9
 
 [TOC(start:2)]
 
 
-## Introduction
 
 At Epic, we have a few simple coding standards and conventions.  This document is not meant to be a discussion or work in progress, but rather, reflects the state of Epic's current coding standards.
 
@@ -27,7 +26,7 @@ Classes should be organized with the reader in mind rather than the writer.  Sin
 
 ## Copyright Notice
 
-Any source file (`.h`, `.cpp`, `.xaml`, etc.) provided by Epic for distribution must contain a copyright notice as the first line in the file. The format of the notice must exactly match that shown below:
+Any source file (.h, .cpp, .xaml, etc.) provided by Epic for distribution must contain a copyright notice as the first line in the file. The format of the notice must exactly match that shown below:
 
 	// Copyright 1998-2015 Epic Games, Inc. All Rights Reserved.
 
@@ -35,12 +34,12 @@ If this line is missing or not formatted properly, CIS will generate an error an
 
 ## Naming Conventions
 
-* The first letter of each word in a name (e.g. type or variable) is capitalized, and there is usually no underscore between words.  For example, Health and UPrimitiveComponent, but not lastMouseCoordinates or delta_coordinates.
+* The first letter of each word in a name (e.g. type or variable) is capitalized, and there is usually no underscore between words.  For example, Health and **UPrimitiveComponent**, but not lastMouseCoordinates or delta_coordinates.
 * Type names are prefixed with an additional upper-case letter to distinguish them from variable names.  For example, FSkin is a type name, and Skin is an instance of a FSkin.
     * Template classes are prefixed by T.
-    * Classes that inherit from UObject are prefixed by U.
-    * Classes that inherit from AActor are prefixed by A.
-    * Classes that inherit from SWidget are prefixed by S.
+    * Classes that inherit from **UObject** are prefixed by U.
+    * Classes that inherit from **AActor** are prefixed by A.
+    * Classes that inherit from **SWidget** are prefixed by S.
     * Classes that are abstract interfaces are prefixed by I.
     * Most other classes are prefixed by F, though some subsystems use other letters.
 * Type and variable names are nouns.
@@ -51,7 +50,7 @@ Variable, method, and class names should be clear, unambiguous, and descriptive.
 
 All variables should be declared one at a time so that a comment on the meaning of the variable can be provided. Also, the JavaDocs style requires it. You can use multi-line or single line comments before a variable, and the blank line is optional for grouping variables.
 
-All functions that return a bool should ask a true/false question, such as "IsVisible()", or "ShouldClearBuffer()".  All boolean variables must be prefixed with a "b" character (e.g. "bPendingDestruction", or "bHasFadedIn").
+All functions that return a bool should ask a true/false question, such as IsVisible() or ShouldClearBuffer().  All boolean variables must be prefixed with a "b" character (e.g. bPendingDestruction, or bHasFadedIn).
 
 A procedure (a function with no return value) should use a strong verb followed by an Object. An exception is if the Object of the method is the Object it is in; then the Object is understood from context. Names to avoid include those beginning with "Handle" and "Process"; the verbs are ambiguous.
 
@@ -65,28 +64,20 @@ Functions that return a value should describe the return value; the name should
 
 ### Examples
 
-		/* the tea weight in kilograms */
 		float TeaWeight;
 
-		/* the number of tea leaves */
 		int32 TeaCount;
 
-		/* true indicates tea is smelly */
 		bool bDoesTeaStink;
 
-		/* non-human-readable FName for tea */
 		FName TeaName;
 
-		/* human-readable name of the tea */
 		FString TeaFriendlyName;
 
-		/* Which class of tea to use */
 		UClass* TeaClass;
 
-		/* The sound of pouring tea */
 		USoundCue* TeaSound;
 
-		/* a picture of tea */
 		UTexture* TeaTexture;
 
 
@@ -112,7 +103,7 @@ Do not use the C+\+ int type in portable code, since it is dependent on the comp
 
 ## Comments
 
-Comments are communication; communication is _vital_. Some things to keep in mind about comments (from Kernighan & Pike _The Practice of Programming_):
+Comments are communication; communication is vital. Some things to keep in mind about comments (from Kernighan & Pike _The Practice of Programming_):
 
 
 ### Guidelines
@@ -171,12 +162,10 @@ Note that two different parameter comment styles are supported, as embodied by t
 Unlike the UE3 coding standard, the method comments should only be included once, where the method is publically declared.  The method comments should only contain information relevant to callers of the method, including any information about overrides of the method that may be relevant to the caller.  Details about the implementation of the method and its overrides that are not relevant to callers should be commented within the method implementation.
 
 
-		/* The interface for drinkable objects. */
 		class IDrinkable
 		{
 		public:
 
-			/*
 			 * Called when a player drinks this object.
 			 * @param OutFocusMultiplier - Upon return, will contain a multiplier to apply to the drinker's focus.
 			 * @param OutThirstQuenchingFraction - Upon return, will contain the fraction of the drinker's thirst to quench (0-1).
@@ -184,12 +173,10 @@ Unlike the UE3 coding standard, the method comments should only be included once
 			virtual void Drink(float& OutFocusMultiplier,float& OutThirstQuenchingFraction) = 0;
 		};
 
-		/* A cuppa (tea) */
 		class FTea : public IDrinkable
 		{
 		public:
 
-			/*
 			 * Calculate a delta-taste value for the tea given the volume and temperature of water used to steep.
 			 * @param VolumeOfWater - Amount of water used to brew in mL
 			 * @param TemperatureOfWater - Water temperature in Kelvins
@@ -202,24 +189,19 @@ Unlike the UE3 coding standard, the method comments should only be included once
 				float& OutNewPotency
 				);
 
-			// Adds a sweetener to the tea, quantified by the grams of sucrose that would produce the same sweetness.
 			void Sweeten(float EquivalentGramsOfSucrose);
 
-			/* The value in yen of tea sold in Japan. */
 			float GetPrice() const
 			{
 				return Price;
 			}
 
-			//- IDrinkable interface
 			virtual void Drink(float& OutFocusMultiplier,float& OutThirstQuenchingFraction);
 
 		private:
 
-			/* The value of the tea in Yen. */
 			float Price;
 
-			// The sweetness of the tea, quantified by the grams of sucrose that would produce the same sweetness.
 			float Sweetness;
 		};
 
@@ -334,9 +316,7 @@ Enum classes are supported by all compilers and encouraged as a replacement for
 		{
 			enum Type
 			{
-				/* Description of Thing1 */
 				Thing1,
-				// Description of Thing2
 				Thing2
 			};
 		}
@@ -345,9 +325,7 @@ Enum classes are supported by all compilers and encouraged as a replacement for
 		UENUM()
 		enum class EThing : uint8
 		{
-			/* Description of Thing1 */
 			Thing1,
-			// Description of Thing2
 			Thing2
 		};
 
@@ -500,10 +478,8 @@ You can use namespaces to organize your classes, functions and variables where a
 
 * Example of a namespace-scoped enum:
 
-		/* Defining a enumeration within a namespace to achieve C#-style enum scoping */
 		namespace EColorChannel
 		{
-			/* Declare EColorChannel::Type as the actual type for this enum */
 			enum Type
 			{
 				Red,
@@ -512,7 +488,6 @@ You can use namespaces to organize your classes, functions and variables where a
 			};
 		}
 
-		/* Given a color channel, returns the name of that channel. */
 		FString GetNameForColorChannel(const EColorChannel::Type ColorChannel)
 		{
 			switch(ColorChannel)
@@ -527,12 +502,10 @@ You can use namespaces to organize your classes, functions and variables where a
 
 * Note that for locally-declared enums, you will not be able to use a namespace for scoping.  In these cases, we opt to declare a local struct with no member variables, only a local enum type and use that struct for scoping.
 
-		/* Defining a locally-scoped enumeration using structs*/
 		class FObjectMover
 		{
         public:
 
-            /* Direction to move */
             struct EMoveDirection
             {
 			    enum Type
@@ -542,7 +515,6 @@ You can use namespaces to organize your classes, functions and variables where a
 			    };
             };
 
-            /* Construct an FObjectMover with the specified movement direction */
             FObjectMover( const EMoveDirection::Type Direction );
 		}
 
@@ -557,28 +529,28 @@ You can use namespaces to organize your classes, functions and variables where a
 		<file contents>
 
 * In general, try to minimize physical coupling.
-** If you can use forward declarations instead of including a header, do so.
-** Include as fine grained as possible; do not include Core.h, include the specific headers in Core that you need definitions from.
+* If you can use forward declarations instead of including a header, do so.
+* Include as fine grained as possible; do not include Core.h, include the specific headers in Core that you need definitions from.
 * Try to include every header you need directly, to make fine-grained inclusion easier.
-** Do not rely on a header that is included indirectly by another header you include
-** Do not rely on being included through another header; include everything you need.
+* Do not rely on a header that is included indirectly by another header you include
+* Do not rely on being included through another header; include everything you need.
 * Modules have Private and Public source directories.  Any definitions that are needed by other modules must be in headers in the Public directory, but everything else should be in the Private directory.  Note that in older Unreal modules, these directories may still be called Src and Inc, but those directories are meant to separate private and public code in the same way, and are not meant to separate header files from source files.
 * Do not worry about setting up your headers for precompiled header generation.  UnrealBuildTool can do a better job of this than you can.
 
 
 ## General Style Issues
 
-* *Minimize dependency distance*. When code depends on a variable having a certain value, try to set that variable's value right before using it. Initializing a variable at the top of an execution block, and not using it for a hundred lines of code, gives lots of space for someone to accidentally change the value without realizing the dependency. Having it on the next line makes it clear why the variable is initialized the way it is and where it is used.
-* *Split methods into sub-methods where possible*. Humans are better at looking at a big picture, and drilling down to the interesting details than to start with the details and reconstruct the big picture from them.  In the same way, it is easier to understand a simple method that calls a sequence of several well named sub-methods than to understand an equivalent method that simply contains all the code in those sub-methods.
+* Minimize dependency distance. When code depends on a variable having a certain value, try to set that variable's value right before using it. Initializing a variable at the top of an execution block, and not using it for a hundred lines of code, gives lots of space for someone to accidentally change the value without realizing the dependency. Having it on the next line makes it clear why the variable is initialized the way it is and where it is used.
+* Split methods into sub-methods where possible. Humans are better at looking at a big picture, and drilling down to the interesting details than to start with the details and reconstruct the big picture from them.  In the same way, it is easier to understand a simple method that calls a sequence of several well named sub-methods than to understand an equivalent method that simply contains all the code in those sub-methods.
 * In function declarations or function call sites, do not add a space between the function's name and the parentheses that precedes the argument list.
-* *Address compiler warnings.* Compiler warning messages mean something is not as it should be. Fix what the compiler is complaining about. If you absolutely cannot address it, use #pragma to suppress the warning; this is a remedy of last resort.
-* *Leave a blank line at the end of the file.* All .cpp and .h files should include a blank line to play nice with gcc.
-* *Never allow float to implicit convert to int32.* This is a slow operation, and does not compile on all compilers. Instead, always use the appTrunc() function to convert to int32. This will ensure cross-compiler compatibility as well as generate faster code.
-* *Enforce encapsulation with the protection keywords.* Class members should be declared private unless they are part of the public interface to the class.
+* Address compiler warnings. Compiler warning messages mean something is not as it should be. Fix what the compiler is complaining about. If you absolutely cannot address it, use #pragma to suppress the warning; this is a remedy of last resort.
+* Leave a blank line at the end of the file. All .cpp and .h files should include a blank line to play nice with gcc.
+* Never allow float to implicit convert to int32. This is a slow operation, and does not compile on all compilers. Instead, always use the appTrunc() function to convert to int32. This will ensure cross-compiler compatibility as well as generate faster code.
+* Enforce encapsulation with the protection keywords. Class members should be declared private unless they are part of the public interface to the class.
 * Interface classes (prefixed with "I") should always be abstract and must not have member variables.  Interfaces are allowed to contain methods that are not pure-virtual, and even methods that are non-virtual or static, as long as they are implemented inline.
-* *Use const wherever possible.* Particularly on reference parameters and class methods. const is documentation as much as it is a compiler directive.
-* *Debug code should either be generally useful and polished, or not checked in.* Debug code intermixed with other code makes the other code much harder to read.
-* *Use intermediate variables to simplify complicated expressions.* If you have a complicated expression, it can be easier to understand if you split it into sub-expressions that are assigned to intermediate variables with names describing the meaning of the sub-expression within the parent expression.  For example:
+* Use const wherever possible. Particularly on reference parameters and class methods. const is documentation as much as it is a compiler directive.
+* Debug code should either be generally useful and polished, or not checked in. Debug code intermixed with other code makes the other code much harder to read.
+* Use intermediate variables to simplify complicated expressions. If you have a complicated expression, it can be easier to understand if you split it into sub-expressions that are assigned to intermediate variables with names describing the meaning of the sub-expression within the parent expression.  For example:
 
 		if ((Blah->BlahP->WindowExists->Etc && Stuff) &&
 			!(bPlayerExists && bGameStarted && bPlayerStillHasPawn &&
@@ -597,7 +569,7 @@ You can use namespaces to organize your classes, functions and variables where a
 		}
 
 
-* *Use the virtual and override keywords when declaring an overriding method.* When declaring a virtual function in a derived class that overrides a virtual function in the parent class, you *must* use both the virtual and the override keywords. For example:
+* Use the virtual and override keywords when declaring an overriding method. When declaring a virtual function in a derived class that overrides a virtual function in the parent class, you *must* use both the virtual and the override keywords. For example:
 		
 		class A
 		{
@@ -613,7 +585,7 @@ You can use namespaces to organize your classes, functions and variables where a
 
 Note that there is a lot of existing code that does not follow this yet due to the recent addition of the override keyword, and the override keyword should be added to that code when convenient.
 
-* *Pointers and references should only have one space, which is to the right of the pointer / reference.*  This makes it easy to quickly Find in Files for all pointers or references to a certain type.
+* Pointers and references should only have one space, which is to the right of the pointer / reference.  This makes it easy to quickly Find in Files for all pointers or references to a certain type.
 		
 	Use this:
 		
