#UEDoc #docs - updated image arrow color and bold lettering

diff --git a/Engine/Documentation/Source/Engine/Physics/Constraints/DampingAndFriction/DampingAndFriction.INT.udn b/Engine/Documentation/Source/Engine/Physics/Constraints/DampingAndFriction/DampingAndFriction.INT.udn
--- a/Engine/Documentation/Source/Engine/Physics/Constraints/DampingAndFriction/DampingAndFriction.INT.udn
+++ b/Engine/Documentation/Source/Engine/Physics/Constraints/DampingAndFriction/DampingAndFriction.INT.udn
@@ -8,38 +8,45 @@ Description:Properties and methods used to cause physical objects to resist move
 
 
 
-## Overview
 
-To slow down Physics Bodies, simulate atmospheric drag, or add resistance to a hinge, there are two properties available to Physics Bodies and Physics Constraints: Linear Damping and Angular Damping. Linear Damping controls how much the Physics Body or Constraint resists translation, while Angular Damping controls how much they resist rotating. 
+To slow down Physics Bodies, simulate atmospheric drag, or add resistance to a hinge, there are two properties available to Physics Bodies and Physics Constraints: Linear Damping 
+and Angular Damping. Linear Damping controls how much the Physics Body or Constraint resists translation, while Angular Damping controls how much they resist rotating. 
 
 ## Physics Bodies
 
-Every Physics Body has properties for Linear Damping and Angular Damping. These are primarily used to simulate drag, and even a value of 1.0 has a substantial effect on the Physics Body.
+Every Physics Body has properties for **Linear Damping** and **Angular Damping**. These are primarily used to simulate drag, and even a value of 1.0 has a substantial effect on the Physics 
+Body.
 
 [REGION:note]
-For reference, in normal gravity (9.8m/s^2), a Linear Damping value of 30 is enough to stop any single body Actor from falling under the initial pull of gravity, at the start of the simulation. Once more than one Physics Body is interacting (through the use of Physics Constraints or if a collision is occurring), the amount of damping required to stop the Physics Body goes up.
+For reference, in normal gravity (9.8m/s^2), a Linear Damping value of 30 is enough to stop any single body Actor from falling under the initial pull of gravity, at the start of the 
+simulation. Once more than one Physics Body is interacting (through the use of Physics Constraints or if a collision is occurring), the amount of damping required to stop the Physics 
+Body goes up.
 
 The _minimum_ Linear Damping value to stop a Physics Body that has had force applied to it is around 100.
 
-Finally, with no Angular Damping, a Physics Body will continue to rotate until it is acted upon by an outside force. Even small values can reduce rotation rapidly, while a value of 100 will almost immediately stop any angular motion caused by an outside source.
+Finally, with no Angular Damping, a Physics Body will continue to rotate until it is acted upon by an outside force. Even small values can reduce rotation rapidly, while a value of 
+100 will almost immediately stop any angular motion caused by an outside source.
 [/REGION]
 
 ## Physics Constraints
 
-The Damping for Physics Constraints work a bit differently than with Physics Bodies: they only work on Limited Motions and when a motion attempts to go beyond its limit. Setting a Physics Constraint's Motion(s) to Limited will display the **Soft Limits**:
+The Damping for Physics Constraints work a bit differently than with Physics Bodies: they only work on Limited Motions and when a motion attempts to go beyond its limit. Setting a 
+Physics Constraint's Motion(s) to Limited will display the **Soft Limits**:
 
 [REGION:raw]
 ![](Phys_ConstraintDampProps.png)(convert:false)
 [/REGION]
 
-_The two Angular Swing Motions share Damping values, the Angular Twist receives its own Damping values, and all the Linear axes use the same Damping values. While they may have separate properties, the properties all have the same effect on the Motions they relate to._
+_The two Angular Swing Motions share Damping values, the Angular Twist receives its own Damping values, and all the Linear axes use the same Damping values. While they may have 
+separate properties, the properties all have the same effect on the Motions they relate to._
 
 | Property | Description |
 | -------- | ----------- |
 | **Limit Stiffness** | When a limit is exceeded, this value affects how harshly the constraint attempts to stop the offending motion. The default of 50 will barely affect the motion, a value of 5000 will appear to bounce back into limit, and 50000 will completely deflect the motion. |
 | **Limit Damping** | Like the Damping on Physics Bodies, this defines how much the motion should be slowed. Unlike Damping on Physics Bodies, the damping affect occurs only when the Motion Limit is exceeded, draining momentum until the Physics Body stops. |
 
-The two values work together to create different effects. A high Stiffness and low Damping will look like a hard limit, while having them both high will appear to be move through tar, coming to a stop gently but quickly. A good deal of experimentation is necessary to find the right values for your specific needs.
+The two values work together to create different effects. A high Stiffness and low Damping will look like a hard limit, while having them both high will appear to be move through 
+tar, coming to a stop gently but quickly. A good deal of experimentation is necessary to find the right values for your specific needs.
 
 For more information on Constraints, see: [](Engine/Physics/Constraints/ConstraintsUserGuide)
 
@@ -48,19 +55,31 @@ Or for a reference on the properties available to Physics Constraints,
 
 <!--
 ## PhysicalMaterial Damping
-This can be thought of like the air resistance that is applied to a body. A force is applied to a body to try and slow it down, and the strength of the force is proportional to the speed that the object is moving. There is a separate setting for linear and angular damping, so you can control how much it will try and resist translation and rotation movement. One place where you might want to use this would be a swinging chandelier that you did not want to swing forever. You may not want to use this for a door though, as a lot of the friction on a real door is due to it's hinge. If you set a high damping on the door's PhysicalMaterial and then broke it off the hinge, it would still have that damping applied to it which may look odd. 
+This can be thought of like the air resistance that is applied to a body. A force is applied to a body to try and slow it down, and the strength of the force is proportional to the 
+speed that the object is moving. There is a separate setting for linear and angular damping, so you can control how much it will try and resist translation and rotation movement. 
+One place where you might want to use this would be a swinging chandelier that you did not want to swing forever. You may not want to use this for a door though, as a lot of the 
+friction on a real door is due to it's hinge. If you set a high damping on the door's PhysicalMaterial and then broke it off the hinge, it would still have that damping applied to 
+it which may look odd. 
 
 
 
 ## Joint Drive
-To simulate 'friction' within a joint, you can use the joint drive properties in the RB_ConstraintInstance (you can get to these properties in PhAT by selecting the joint and pressing the 'I' key or the 'i' icon on the toolbar). First you want to turn on the appropriate bTwistVelocityDrive or bSwingVelocityDrive drive, set the desired angular velocity to zero and modify the AngularDriveForceLimit to control how much friction is in the joint. You can find more information on this on the [PhysicsConstraintReference](Engine/Physics/PhysX/Constraints) page. 
+To simulate 'friction' within a joint, you can use the joint drive properties in the RB_ConstraintInstance (you can get to these properties in PhAT by selecting the joint and 
+pressing the 'I' key or the 'i' icon on the toolbar). First you want to turn on the appropriate bTwistVelocityDrive or bSwingVelocityDrive drive, set the desired angular velocity 
+to zero and modify the AngularDriveForceLimit to control how much friction is in the joint. You can find more information on this on the 
+[PhysicsConstraintReference](Engine/Physics/PhysX/Constraints) page. 
 
-The key thing here is that the damping is applied between two jointed object. Imagine two objects jointed together, but then thrown through space. If you used PhysicalMaterial damping here, it would slow down how quickly the pair of objects would move through the air, when what you really want is to modify how they move relative to each other. That is why Joint Drive is the right thing to use in this case.
+The key thing here is that the damping is applied between two jointed object. Imagine two objects jointed together, but then thrown through space. If you used PhysicalMaterial 
+damping here, it would slow down how quickly the pair of objects would move through the air, when what you really want is to modify how they move relative to each other. 
+That is why Joint Drive is the right thing to use in this case.
 
 
 
 ## PhysicsVolume RigidBodyDamping
-Sometimes you want to apply damping or a force to objects only in a certain area of your level. To do this you would use a PhysicsVolume. The ZoneVelocity property of a PhysicsVolume is the desired velocity that objects in that volume should achieve. You can think of this as the speed of a river or wind-tunnel in that area. You can then modify the RigidBodyDamping property to control how strongly a physics body will try and achieve that velocity when in the volume. Setting a zero ZoneVelocity and a high RigidBodyDamping can make the volume act like jelly. Setting a horizontal ZoneVelocity and a moderate RigidBodyDamping can make it act like a flowing river of water.
+Sometimes you want to apply damping or a force to objects only in a certain area of your level. To do this you would use a PhysicsVolume. The ZoneVelocity property of a 
+PhysicsVolume is the desired velocity that objects in that volume should achieve. You can think of this as the speed of a river or wind-tunnel in that area. You can then modify 
+the RigidBodyDamping property to control how strongly a physics body will try and achieve that velocity when in the volume. Setting a zero ZoneVelocity and a high RigidBodyDamping 
+can make the volume act like jelly. Setting a horizontal ZoneVelocity and a moderate RigidBodyDamping can make it act like a flowing river of water.
 -->
  
 
