#UE4doc #docs: Removed extraneous bolding as according to the instructional guide. Altered images to match the instructional guide: green -> yellow.

diff --git a/Engine/Documentation/Source/Engine/UMG/UserGuide/Fonts/Fonts.INT.udn b/Engine/Documentation/Source/Engine/UMG/UserGuide/Fonts/Fonts.INT.udn
--- a/Engine/Documentation/Source/Engine/UMG/UserGuide/Fonts/Fonts.INT.udn
+++ b/Engine/Documentation/Source/Engine/UMG/UserGuide/Fonts/Fonts.INT.udn
@@ -8,8 +8,6 @@ Related: Engine/UMG/UserGuide/WidgetTypeReference
 
 [TOC(start:2 end:2)]
 
-## Overview
-
 This page covers how to import your own fonts, using the **Font Editor** and fonts in the **Unreal Motion Graphics (UMG) UI Designer**. 
 
 ## Font Assets
@@ -76,10 +74,10 @@ A breakdown of the Font Editor Window is presented below:
 
 1.	**Toolbar Menu** - From this menu, you can save any changes you make, find the asset in the **Content Browser**, change the Background Color of the preview window (4) or the Foreground Color (text color) in the preview window. There are options for Updating or Exporting changes made, however these options are only available within the **Offline** cache mode. 
 1.	**Default Font Family** - In this window, you can assign the Default Font Family for use with this Font asset. You can add versions of a particular Font style (for example Normal, Bold, Italics, Underline, etc.) or have a collection of different Font styles as one Composite Font. If you have created a blank Font asset, you can assign a font from inside this window as well. Once a Font has been added, you can also define the hinting algorithm used with the font. 
-1.	**Sub-Font Family** - In this window, you assign the Sub-Font Family for this Font asset to use. Here you can specify a **Character Range** and if a character entered falls within the range, you can specify a different Font style to use instead of the Default (this is useful for when you want to use different Font types for different Languages). 
+1.	**Sub-Font Family** - In this window, you assign the Sub-Font Family for this Font asset to use. Here you can specify a Character Range and if a character entered falls within the range, you can specify a different Font style to use instead of the Default (this is useful for when you want to use different Font types for different Languages). 
 1.	**Preview** - This window allow you to preview your fonts and provides a text entry box for entering sample texts. 
-1.	**Details** - In this window, you can change the **Font Cache Type** as well as change the Font Size and Font Name (for Runtime). 
-	* If you are using the older method, you can still change the parameters for your Font while in **Offline** cache mode. 
+1.	**Details** - In this window, you can change the Font Cache Type as well as change the Font Size and Font Name (for Runtime). 
+	* If you are using the older method, you can still change the parameters for your Font while in Offline cache mode. 
 	* You can also convert any existing Font assets from **Offline** to **Runtime** without having to replace them. 
 
 An example Font asset is provided below. 
