#UE4doc #docs: Removed extraneous bolding as according to the instructional guide. ##Overview removed as according to the instructional guide.

diff --git a/Engine/Documentation/Source/Engine/UMG/UserGuide/PropertyBinding/PropertyBinding.INT.udn b/Engine/Documentation/Source/Engine/UMG/UserGuide/PropertyBinding/PropertyBinding.INT.udn
--- a/Engine/Documentation/Source/Engine/UMG/UserGuide/PropertyBinding/PropertyBinding.INT.udn
+++ b/Engine/Documentation/Source/Engine/UMG/UserGuide/PropertyBinding/PropertyBinding.INT.udn
@@ -5,9 +5,7 @@ Description:A guide to binding properties inside UMG.
 
 [TOC(start:2 end:2)]
 
-## Overview
-
-One of the most useful aspects inside **UMG** is the ability to bind properties of your Widgets to **Functions** or **Properties** inside the **Blueprint**. By binding a property to a Function or Property Variable in your **Blueprint**, anytime that Function is called or Property is updated, it will be reflected in the Widget. 
+One of the most useful aspects inside UMG is the ability to bind properties of your Widgets to Functions or Properties inside the Blueprint. By binding a property to a Function or Property Variable in your Blueprint, anytime that Function is called or Property is updated, it will be reflected in the Widget. 
 
 ## Function Binding
 
