#UEDoc #docs - Removed deprecated properties listing and updated/added sections that were missing

diff --git a/Engine/Documentation/Source/Engine/Physics/Destructibles/DestructibleProperties/DestructibleProperties.INT.udn b/Engine/Documentation/Source/Engine/Physics/Destructibles/DestructibleProperties/DestructibleProperties.INT.udn
--- a/Engine/Documentation/Source/Engine/Physics/Destructibles/DestructibleProperties/DestructibleProperties.INT.udn
+++ b/Engine/Documentation/Source/Engine/Physics/Destructibles/DestructibleProperties/DestructibleProperties.INT.udn
@@ -34,9 +34,11 @@ Description:Descriptions of all the properties on Destructible Actors in the Des
 | Property  | Description |
 | --------- | ----------- |
 | [REGION:tablesection] _Damage Parameters_ [/REGION] ||
-| **Damage Threshold** | The Damage amount which will cause a chunk to fracture. |
-| **Damage Spread** | Controls the distance into the destructible Actor to propagate the damage. The damage applied to the chunk is multiplied by **Damage Spread** to get the propagation distance. |
-| **Impact Damage** | If a chunk is at a depth which has impact damage set, it will take damage equal to **Impact Damage** times the impact force when it collides with another object (assuming **Report Rigid Collisions** are enabled on both objects). |
+| **Damage Threshold** | The Damage amount which will cause a chunk to fracture (break free). |
+| **Damage Spread** | Controls the distance into the destructible Actor to propagate the damage. The damage applied to the chunk is multiplied by Damage Spread to get the propagation distance. |
+| **Enable Impact Damage** | Whether to apply damage to a destructible when colliding with an object. |
+| **Impact Damage** | If a chunk is at a depth which has impact damage set, it will take damage equal to Impact Damage times the impact force when it collides with another object (assuming Report Rigid Collisions are enabled on both objects). |
+| **Custom Impact Resistance** | By default objects that collide with destructibles will bounce back. Custom resistance allows for finer control of how much a destrucible "pushes back" against a colliding object. |
 | **Impact Resistance** | When a chunk takes impact damage due to physical contact, this parameter is the maximum impulse the contact can generate. Weak materials will have this set low, so heavier objects can pass through them during fracture. 0.0 is interpreted at infinite. |
 | **Default Impact Damage** | Impact damage will be taken to this depth. If negative, impact damage is disabled. |
 | [REGION:tablesection] _Debris Parameters_ [/REGION] ||
@@ -70,10 +72,7 @@ Description:Descriptions of all the properties on Destructible Actors in the Des
 | [REGION:tablesection] _Fracture Effects - Per Depth Level_ [/REGION] ||
 | **Particle System** | The particle system to spawn when a chunk breaks. |
 | **Sound** | The sound or sound cue to play when a chunk breaks.|
-| **Body Setup (Parameter)** ||
-| [REGION:tablesection] _Body Setup_ [/REGION] ||
-| **Bone Name** | Only used if a Physics Asset is used, which they are not for destructible Actors. |
-| **Consider for Bounds** | If enabled the body setup will be used for considering bounds. |
+
 
 For the Physics and Collision section of properties, see the [Collision](Engine/Physics/Collision) documentation.
 	
@@ -91,6 +90,17 @@ For the Physics and Collision section of properties, see the [Collision](Engine/
 | **Property** | **Description** |
 | -------- | --------- |
 | **Streaming Distance Multiplier** | Allows artists to adjust the distance where textures using UV 0 are streamed in/out. 1.0 is the default, whereas a higher value increases the streamed in resolution. |
+
+## Chunk Parameters Category
+
+### Chunks Sub Category
+| **Property** | **Description** |
+| -------- | --------- |
+| **Is Support Chunk** | Defines the chunk to be environmentally supported, if the appropriate NxDestructibleParametersFlag flags are set in NxDestructibleParameters.  |
+| **Do Not Fracture** | Defines the chunk to be unfracturable. If this is true, none of its children will be fracturable. |
+| **Do Not Damage** | Defines the chunk to be undamagable. This means this chunk will not fracture, but its children might |
+| **Do not Crumble** | Defines the chunk to be uncrumbleable. This means this chunk will not be broken down into fluid mesh particles no matter how much damage it takes. Note: this only applies to chunks with no children. For a chunk with children then: 1) The chunk may be broken down into its children, and then its children may be crumbled, if the doNotCumble Flag is not set on them, 2) If the destructible module's chunk depth offset LOD may be set such that this chunk effectively has no children. In this case, the doNotCrumble flag will apply to it. |
+
 	
 ## Fracture Settings Category
 
