#UE4doc #docs: Removed extraneous bolding as according to the instructional guide. ##Overview removed as according to the instructional guide. Fixed minor spelling errors.

diff --git a/Engine/Documentation/Source/Engine/UMG/UserGuide/Styling/Styling.INT.udn b/Engine/Documentation/Source/Engine/UMG/UserGuide/Styling/Styling.INT.udn
--- a/Engine/Documentation/Source/Engine/UMG/UserGuide/Styling/Styling.INT.udn
+++ b/Engine/Documentation/Source/Engine/UMG/UserGuide/Styling/Styling.INT.udn
@@ -12,10 +12,7 @@ Related: Engine\UMG\UserGuide\WidgetTypeReference
 
 [TOC(start:2 end:2)]
 
-
-## Overview
-
-When it comes to creating your UI screens through **UMG**, arranging the layout of your elements is only the beginning. For each of your Buttons, Bars, Text Boxes, etc. there are several **Style** options that can be assigned directly within **UMG** from the **Details** panel that can affect the way that they appear. 
+When it comes to creating your UI screens through UMG, arranging the layout of your elements is only the beginning. For each of your Buttons, Bars, Text Boxes, etc. there are several Style options that can be assigned directly within UMG from the **Details** panel that can affect the way that they appear. 
 
 Each of the Widgets below utilize the **Style** option, however their individual Style options may vary: 
 
@@ -35,7 +32,7 @@ Each of the Widgets below utilize the **Style** option, however their individual
 
 ## States
 
-There may be instances, particularly with interactive Widgets, that you want to convey a different look for your Widget based on how it is interacted with or the condition it is in. For example, say you have a button on screen that appears normally and when moused over, changes colors or pulsates, then when clicked does something entirely different. Refered to as **States**, this is the most common form of styling and allows you to specify how the Widget appears based on the current state it is in. 
+There may be instances, particularly with interactive Widgets, that you want to convey a different look for your Widget based on how it is interacted with or the condition it is in. For example, say you have a button on screen that appears normally and when moused over, changes colors or pulsates, then when clicked does something entirely different. Referred to as **States**, this is the most common form of styling and allows you to specify how the Widget appears based on the current state it is in. 
 
 ![](Style1.png)
 
@@ -53,15 +50,15 @@ Examples of each are provided below:
 |Box|Border|Image|
 [/REGION]
 
-* The **Box** draws a 3x3 box where the sides and the middle **Stretch** (orange arrows) based on the **Margins** (dotted blue lines). 
-	*	The green arrows represent constant values based on the **Margin** x **Image Size**. 
-* The **Border** draws a 3x3 border where the sides **Tile** (orange arrows) based on the **Margins** (dotted blue lines). 
+* The **Box** draws a 3x3 box where the sides and the middle Stretch (orange arrows) based on the Margins (dotted blue lines). 
+	*	The green arrows represent constant values based on the Margin x Image Size. 
+* The **Border** draws a 3x3 border where the sides Tile (orange arrows) based on the Margins (dotted blue lines). 
 	*	The middle is not drawn when using a Border. 
-	*	The green arrows represent constant values based on the **Margin** x **Image Size**. 
-* The **Image** will draw the image and will **Stretch** (orange arrow) and the **Margin** will be ignored.
+	*	The green arrows represent constant values based on the Margin x Image Size. 
+* The **Image** will draw the image and will Stretch (orange arrow) and the Margin will be ignored.
 
 #### Texture Based Images
-You can assign **Textures** as your **Image** asset, but there are some flags you will want to set in order for the Texture to show up on every device, and to show up crisp and without compression artifacts. Inside the Texture Edtior under Level Of Detail, set **Mip Gen Settings** to **NoMipmaps** and set the **LOD Group** to **UI**. Under Compression, set the **Compression Setting** to **TC Editor Icon**.  
+You can assign Textures as your Image asset, but there are some flags you will want to set in order for the Texture to show up on every device, and to show up crisp and without compression artifacts. Inside the Texture Edtior under Level Of Detail, set **Mip Gen Settings** to **NoMipmaps** and set the **LOD Group** to **UI**. Under Compression, set the **Compression Setting** to **TC Editor Icon**.  
 
 [REGION:lightbox]
 	[![](TextureSetup.png)(w:480)](TextureSetup.png)
@@ -70,7 +67,7 @@ You can assign **Textures** as your **Image** asset, but there are some flags yo
 _Click for full view._
 
 #### Material Based Images
-**Materials** can also be assigned as the **Image** asset for your Image States, but you will need to change a few settings. Inside the Material, in the **Details** panel under **Usage**, you will need to check the **Used with UI** checkbox. This enables compiling a different shader specifically for Slate. You also need to only output **Emissive Color**
+Materials can also be assigned as the Image asset for your Image States, but you will need to change a few settings. Inside the Material, in the **Details** panel under **Usage**, you will need to check the **Used with UI** checkbox. This enables compiling a different shader specifically for Slate. You also need to only output **Emissive Color**
 
 [REGION:lightbox]
 	[![](UsedWithUI.png)(w:760)](UsedWithUI.png)
@@ -79,13 +76,13 @@ _Click for full view._
 _Click for full view._
 
 #### Dynamic Materials
-In order to modify Material Parameters, you need a **Dynamic Material Instance**. The **Image** Widget makes this especially easy as long as it has a Slate Brush on it with a Material, you can ask for the dynamic image (only creates it once and then caches it).
+In order to modify Material Parameters, you need a **Dynamic Material Instance**. The Image Widget makes this especially easy as long as it has a Slate Brush on it with a Material, you can ask for the dynamic image (only creates it once and then caches it).
 
 ![](UMG_Auto_Dynamic_Material_Image.png)(w:760)
 
 ### Tint / Margin
 
-Also under each State is the option to **Tint** the image, which allows you apply a tint color to the supplied image associated with the state. There is also a **Margin** option which allows you to enter a margin size for the **Box** and **Border** draw modes (in **Image** mode, margin is ignored).  
+Also under each State is the option to **Tint** the image, which allows you apply a tint color to the supplied image associated with the state. There is also a **Margin** option which allows you to enter a margin size for the **Box** and **Border** draw modes (in Image mode, margin is ignored).  
 
 ## Padding
 
@@ -118,7 +115,7 @@ Samples of **Render Transform** settings are illustrated below.
 
 ![](RenderTransforms.png)(w:520)
 
-Each of the Render Transform settings can be Keyframed allowing you to modify them with **Animations** and they can also be modified through **Blueprints** allowing you to make modificiations during or based on gameplay actions performed by the user. Render Transforms are relative to **Layout** Transforms and are not clipped by their parent (for example if you had a Scroll Box full of Buttons, adjusting the **Translation** in the Render Transform settings could push the Button outside of its box). 
+Each of the Render Transform settings can be Keyframed allowing you to modify them with Animations and they can also be modified through Blueprints allowing you to make modifications during or based on gameplay actions performed by the user. Render Transforms are relative to Layout Transforms and are not clipped by their parent (for example if you had a Scroll Box full of Buttons, adjusting the Translation in the Render Transform settings could push the Button outside of its box). 
 
 
 
