#UEDoc #docs - need to link newer properties to the source tool tips

diff --git a/Engine/Documentation/Source/Engine/Physics/PhAT/Properties/Properties.INT.udn b/Engine/Documentation/Source/Engine/Physics/PhAT/Properties/Properties.INT.udn
--- a/Engine/Documentation/Source/Engine/Physics/PhAT/Properties/Properties.INT.udn
+++ b/Engine/Documentation/Source/Engine/Physics/PhAT/Properties/Properties.INT.udn
@@ -44,7 +44,7 @@ Description:A listing of properties and descriptions for Physics Bodies, Physics
 
 [TOC(start:2)]
 
-## Overview
+
 
 ## Physics Asset Tool Properties
 
@@ -57,6 +57,8 @@ Description:A listing of properties and descriptions for Physics Bodies, Physics
 |[REGION:tablesection]Simulation[/REGION]||
 |**Floor Gap** | %doxygen:UPhATSimOptions::FloorGap% |
 |**Grav Scale** | %doxygen:UPhATSimOptions::GravScale% |
+|**Max FPS** | |
+|**Time Dilation** |  |
 |[REGION:tablesection]Mouse Spring[/REGION]||
 |**Handle Linear Damping** | %doxygen:UPhATSimOptions::HandleLinearDamping% |
 |**Handle Linear Stiffness** | %doxygen:UPhATSimOptions::HandleLinearStiffness% |
@@ -81,26 +83,31 @@ Description:A listing of properties and descriptions for Physics Bodies, Physics
 | **Consider for Bounds** | %doxygen:UBodySetup::bConsiderForBounds% |
 |[REGION:tablesection]Physics[/REGION]||
 | **Mass in KG** | Mass of the body in KG. |
-| **Start Awake** | %doxygen:FBodyInstance::bStartAwake% |
+| **Linear Damping** | %doxygen:FBodyInstance::LinearDamping% |
+| **Angular Damping** | %doxygen:FBodyInstance::AngularDamping% |
+| **Enable Gravity** | %doxygen:FBodyInstance::bEnableGravity% |
 | **Physics Type** | %doxygen:UBodySetup::PhysicsType% |
+| **Double Sided Geometry** |  |
 | **Simple Collision Physical Material** | %doxygen:UBodySetup::PhysMaterial% |
-| **Enable Gravity** | %doxygen:FBodyInstance::bEnableGravity% |
+| **Start Awake** | %doxygen:FBodyInstance::bStartAwake% |
+| **Center of Mass Offset** | %doxygen:FBodyInstance::COMNudge% |
+| **Mass Scale** | %doxygen:FBodyInstance::MassScale% |
+| **Max Angular Velocity** |  |
 | **Use Async Scene** | %doxygen:FBodyInstance::bUseAsyncScene% |
+| **Override Max Depenetration Velocity** |  |
+| **Max Depenetration Velocity** |  |
+| **Override Walkable Slope on Instance** |  |
 | **Walkable Slope Override** | %doxygen:UBodySetup::WalkableSlopeOverride% |
 | **Walkable Slope Behavior** | %doxygen:FWalkableSlopeOverride::WalkableSlopeBehavior% |
 | **Walkable Slope Angle** | %doxygen:FWalkableSlopeOverride::WalkableSlopeAngle% |
-| **Cente of Mass Offset** | %doxygen:FBodyInstance::COMNudge% |
 | **Sleep Family** | %doxygen:FBodyInstance::SleepFamily% |
-| **Mass Scale** | %doxygen:FBodyInstance::MassScale% |
-| **Angular Damping** | %doxygen:FBodyInstance::AngularDamping% |
-| **Linear Damping** | %doxygen:FBodyInstance::LinearDamping% |
 | **Position Solver Iteration Count** | %doxygen:FBodyInstance::PositionSolverIterationCount% |
 | **Velocity Solver Iteration Count** | %doxygen:FBodyInstance::VelocitySolverIterationCount% |
 |[REGION:tablesection]Collision[/REGION]||
 | **Simulation Generates Hit Events** | %doxygen:FBodyInstance::bNotifyRigidBodyCollision% |
+| **Phys Material Override** | %doxygen:FBodyInstance::PhysMaterialOverride% |
 | **Collision Response** | %doxygen:UBodySetup::CollisionReponse% |
 | **Use CCD** | %doxygen:FBodyInstance::bUseCCD% |
-| **Phys Material Override** | %doxygen:FBodyInstance::PhysMaterialOverride% |
 
 ## Physics Constraint Properties
 
diff --git a/Engine/Documentation/Source/Engine/Physics/PhAT/Reference/PhAT_Reference.INT.udn b/Engine/Documentation/Source/Engine/Physics/PhAT/Reference/PhAT_Reference.INT.udn
--- a/Engine/Documentation/Source/Engine/Physics/PhAT/Reference/PhAT_Reference.INT.udn
+++ b/Engine/Documentation/Source/Engine/Physics/PhAT/Reference/PhAT_Reference.INT.udn
@@ -45,7 +45,7 @@ Description:Tool to set up physics bodies and joints to use for physical simulat
 [TOC(start:2)]
 
 
-## Overview
+
 
 This document contains descriptions for all interface elements within the **Physics Asset Tool** (**PhAT**).
 
