#UE4doc #docs: Removed extraneous bolding as according to the instructional guide.

diff --git a/Engine/Documentation/Source/Engine/UMG/UserGuide/WidgetBlueprints/WidgetBlueprints.INT.udn b/Engine/Documentation/Source/Engine/UMG/UserGuide/WidgetBlueprints/WidgetBlueprints.INT.udn
--- a/Engine/Documentation/Source/Engine/UMG/UserGuide/WidgetBlueprints/WidgetBlueprints.INT.udn
+++ b/Engine/Documentation/Source/Engine/UMG/UserGuide/WidgetBlueprints/WidgetBlueprints.INT.udn
@@ -3,7 +3,7 @@ Title: Widget Blueprints
 Crumbs:%ROOT%, Engine, Engine/UMG, Engine/UMG/UserGuide
 Description:How to create a Widget Blueprint and Overview of the Widget Blueprint Interface.
 
-In order to start working with **Unreal Motion Graphics**, you will first need to create a **Widget Blueprint** as seen below. 
+In order to start working with Unreal Motion Graphics, you will first need to create a **Widget Blueprint** as seen below. 
 
 1.	Click the **Add New** button inside the **Content Browser**, then under **User Interface**, select the **Widget Blueprint** option. 
 
@@ -42,7 +42,7 @@ _Click for full view._
 | **(1) Menu Bar** | The common Menu bar. |
 | **(2) Tool Bar** | This contains a number of commonly used functions for the Blueprint Editor, such as Compile, Save, and Play. |
 | **(3) Editor Mode** | Switches the UMG Widget Blueprint Editor between Designer and Graph modes. |
-| **(4) Palette** | The list of Widgets that you can drag-and-drop into the Visual Designer. Displays any class inheriting from **UWidget**. |
+| **(4) Palette** | The list of Widgets that you can drag-and-drop into the Visual Designer. Displays any class inheriting from UWidget. |
 | **(5) Hierarchy** | Displays the parenting structure of the User Widget. You can also drag Widgets into this window. |
 | **(6) Visual Designer** | This is the visual representation of the layout. Widgets dragged into the Visual Designer can be manipulated in the window. |
 | **(7) Details** | Displays the properties of the currently selected Widget. |
@@ -60,5 +60,5 @@ The **Graph** tab of the **Widget Blueprint Editor** is pictured below.
 
 _Click for full view._
 
-The **Graph** tab functions just like the default Blueprint Editor. For more information on the basic functionality of the **Graph** tab, refer to the [](Engine/Blueprints/Editor/UIComponents/GraphEditor) documentation. 
+The Graph tab functions just like the default Blueprint Editor. For more information on the basic functionality of the Graph tab, refer to the [](Engine/Blueprints/Editor/UIComponents/GraphEditor) documentation. 
 
