#UE4doc #docs: Removed extraneous bolding as according to the instructional guide.

diff --git a/Engine/Documentation/Source/Engine/UMG/UserGuide/WidgetTypeReference/WidgetTypeReference.INT.udn b/Engine/Documentation/Source/Engine/UMG/UserGuide/WidgetTypeReference/WidgetTypeReference.INT.udn
--- a/Engine/Documentation/Source/Engine/UMG/UserGuide/WidgetTypeReference/WidgetTypeReference.INT.udn
+++ b/Engine/Documentation/Source/Engine/UMG/UserGuide/WidgetTypeReference/WidgetTypeReference.INT.udn
@@ -5,11 +5,9 @@ Description:Reference guide to the different types of Widgets that can be used i
 
 [TOC(start:2 end:2)]
 
-## Overview
+Inside the **Widget Blueprint Editor**, located under the **Palette** window, there are four categories of **Widgets**, each of which contain different Widgets Types that you can drag-and-drop into the **Visual Designer**. By mixing and matching these Widget Types, you can layout the look of your UI on the **Designer** tab and add the functionality behind your Widgets through the settings in the **Details** panel for each Widget and via the **Graph** tab. 
 
-Inside the **Widget Blueprint Editor**, located under the **Palette** window, there are four categories of **Widgets**, each of which contain different **Widgets Types** that you can drag-and-drop into the **Visual Designer**. By mixing and matching these Widget Types, you can layout the look of your UI on the **Designer** tab and add the functionality behind your Widgets through the settings in the **Details** panel for each Widget and via the **Graph** tab. 
-
-Listed below are each of the **Widget Types** under the **Palette** window: 
+Listed below are each of the Widget Types under the **Palette** window: 
 
 [REGION:note]
 Click the Widget Type (where applicable) for a usage guide along with more information on the Widget itself. 
@@ -20,7 +18,7 @@ Click the Widget Type (where applicable) for a usage guide along with more infor
 
 ![](Common.png)
 
-The most frequently used **Widgets** will more than likely fall under this category. 
+The most frequently used Widgets will more than likely fall under this category. 
 
 | Option | Description |
 | -------- | ----------- |
@@ -38,20 +36,20 @@ The most frequently used **Widgets** will more than likely fall under this categ
 
 ![](Input.png)
 
-There are a couple of options on how to allow **Input** from a user which are outlined below. 
+There are a couple of options on how to allow Input from a user which are outlined below. 
 
 | Option | Description |
 | -------- | ----------- |
 | **Combo Box (String)** | The combo box (string) allows you to display a list of options to the user in a dropdown menu from them to select one. |
 | **Spin Box** | A numerical entry box that allows for direct entry of the number or allows the user to click and slide the number. |
-| **Text Box (Multi-Line)** | Similar to the **Text Box** but allows a user to enter multiple lines of text instead of a single line. |
+| **Text Box (Multi-Line)** | Similar to the Text Box but allows a user to enter multiple lines of text instead of a single line. |
 
 
 ## Panel Widgets
 
 ![](Panel.png)
 
-The **Panel** category contains useful **Widgets** for controlling layouts and placing other **Widgets**. 
+The **Panel** category contains useful Widgets for controlling layouts and placing other Widgets. 
 
 | Option | Description |
 | -------- | ----------- |
@@ -59,7 +57,7 @@ The **Panel** category contains useful **Widgets** for controlling layouts and p
 | **Grid Panel** | This is a panel that evenly divides up available space between all of its children widgets. |
 | **Horizontal Box** | Allows child widgets to be laid out in a flow horizontally. |
 | **Overlay** | Allows widgets to be stacked on top of each other and uses a simple flow layout for content on each layer.  |
-| **Scale Box** | Allows you to place content with a desired size and have it scale to meet the constraints placed on this box�s allotted area. If you needed to have a background image scale to fill an area but not become distorted with different aspect rations, or if you need to auto fit some text to an area, this is the control for you. |
+| **Scale Box** | Allows you to place content with a desired size and have it scale to meet the constraints placed on this box's allotted area. If you needed to have a background image scale to fill an area but not become distorted with different aspect rations, or if you need to auto fit some text to an area, this is the control for you. |
 | **Scroll Box** | An arbitrary scrollable collection of widgets. This is great for presenting 10-100 widgets in a list. This does not support virtualization. |
 | **Uniform Grid Panel** | A panel that evenly divides up available space between all of its children. |
 | **Vertical Box** | A vertical box widget is a layout panel allowing child widgets to be automatically laid out vertically. This is great for stacking widgets on top of each other and keeping widgets aligned vertically.|
@@ -77,7 +75,7 @@ The **Widgets** contained in the **Primitive** category provide additional metho
 | -------- | ----------- |
 | **Circular Throbber** | A throbber widget that orients images in a spinning circle.  |
 | **Editable Text** | This a text field that allows for user input without a box background. This only supports a single line of editable text. |
-| **Editable Text (Multi-Line)** | Similar to **Editable Text** but supports multiple lines of text instead of a single line.|
+| **Editable Text (Multi-Line)** | Similar to Editable Text but supports multiple lines of text instead of a single line.|
 | [**Menu Anchor**](Engine\UMG\UserGuide\WidgetTypeReference\MenuAnchor) | [INCLUDE:Engine\UMG\UserGuide\WidgetTypeReference\MenuAnchor#desc] |
 | **Native Widget Host** |This is a container widget that can contain one child Slate widget. This should be used when all you need is to nest a native widget inside a UMG widget. |
 | **Spacer** |A spacer widget provides custom padding between other widgets. The spacer does not have a visual representation and is invisible in game.|
