#UEDoc #docs - updated image and bold lettering

diff --git a/Engine/Documentation/Source/Engine/Physics/PhAT/Reference/PhAT_Reference.INT.udn b/Engine/Documentation/Source/Engine/Physics/PhAT/Reference/PhAT_Reference.INT.udn
--- a/Engine/Documentation/Source/Engine/Physics/PhAT/Reference/PhAT_Reference.INT.udn
+++ b/Engine/Documentation/Source/Engine/Physics/PhAT/Reference/PhAT_Reference.INT.udn
@@ -92,24 +92,24 @@ The PhAT Physics Editor is divided into the following regions:
 
 ### Asset
 
-* **Find in Content Browser...** - Finds the currently opened Physics Asset in the **Content Browser**.
+* **Find in Content Browser...** - Finds the currently opened Physics Asset in the Content Browser.
 
 
 ### Window
 
-* **Viewport** - Shows the **Viewport** Pane if it is closed.
-* **Details** - Shows the **Details** Pane if it is closed.
-* **Hierarchy** - Shows the **Hierarchy** Pane if it is closed.
-* **Toolbar** - Shows the **Toolbar** if it is closed.
-* **Content Browser>** - Opens or shows one of 4 **Content Browser** tabs.
-* **Developer Tools>** - Opens the **Developer Tools** Tab. 
-* **Blueprint Debugger** - Opens the **Blueprint Debugger**.
-* **Class Viewer** - Opens the **Class Viewer**.
-* **Message Log** - Opens or shows the **Message Log**.
-* **Output Log** - Opens the **Output Log**.
-* **Plugins** - Opens the **Plugins** interface.
-* **Device Manger** - Opens the **Device Manager**.
-* **Session Frontend** - Shows the **Session Frontend**.
+* **Viewport** - Shows the Viewport Pane if it is closed.
+* **Details** - Shows the Details Pane if it is closed.
+* **Hierarchy** - Shows the Hierarchy Pane if it is closed.
+* **Toolbar** - Shows the Toolbar if it is closed.
+* **Content Browser>** - Opens or shows one of 4 Content Browser tabs.
+* **Developer Tools>** - Opens the Developer Tools Tab. 
+* **Blueprint Debugger** - Opens the Blueprint Debugger.
+* **Class Viewer** - Opens the Class Viewer.
+* **Message Log** - Opens or shows the Message Log.
+* **Output Log** - Opens the Output Log.
+* **Plugins** - Opens the Plugins interface.
+* **Device Manger** - Opens the Device Manager.
+* **Session Frontend** - Shows the Session Frontend.
 * **Save Layout** - Saves the locations of all attached panes for this editor.
 * **Enable Fullscreen** - Enables Fullscreen for the main Unreal Editor window.
 
@@ -133,16 +133,16 @@ What follows is a description of each of the toolbar buttons, from left to right
 | **Icon**  | **Description** | 
 | --- | --- |
 | ![](icon_SaveAsset_40x.png) | **Save** - Saves the Physics Asset currently open in the Physics Asset Tool. |
-| ![](icon_toolbar_genericfinder_40px.png) | **Find in Content Browser** - Will select the current Physics Asset in the **Content Browser**. |
-| ![](icon_Enable_Simulation_40px.png) | **Simulation** - Simulates the Physics Asset. While simulating, the Poke and Drag functions are available to interact with the simulating Physics Asset. Further, if **Physics Blend** is set to less than 1.0, an animation can be selected and played by using the Animation selector and the **Play** button. |
+| ![](icon_toolbar_genericfinder_40px.png) | **Find in Content Browser** - Will select the current Physics Asset in the Content Browser. |
+| ![](icon_Enable_Simulation_40px.png) | **Simulation** - Simulates the Physics Asset. While simulating, the Poke and Drag functions are available to interact with the simulating Physics Asset. Further, if Physics Blend is set to less than 1.0, an animation can be selected and played by using the Animation selector and the Play button. |
 | ![](Engine\Physics\PhAT\UserGuide\selectedToggle_button.png) | **Selected Simulation** - Toggling this feature on will make the Physics Asset Tool only simulate the selected Physics Bodies and their children. For example, selecting the shoulder of a human would simulate the shoulder down to the fingers. |
-| ![](icon_PhAT_modeSelector.png) | **Mode Selector** - Clicking this enables you to switch between **Bodies Mode** and **Constraints Mode**.|
+| ![](icon_PhAT_modeSelector.png) | **Mode Selector** - Clicking this enables you to switch between Bodies Mode and Constraints Mode.|
 | ![](icon_translate_40x.png) | **Translation** -  Using the Translation tool allows you to move the currently selected Physics Body or Physics Constraint. |
 | ![](icon_rotate_40x.png) | **Rotation** - Using the Rotation tool allows you to rotate the currently selected Physics Body or Physics Constraint. |
 | ![](icon_scale_40x.png) | **Scale** - Using the Scale tool allows you to scale the currently selected Physics Body. |
-| ![](icon_PhAT_Play_40x.png) | **Play** - Plays the currently selected animation. **Physics Blend** must be less that 1.0 for any effect to show. |
+| ![](icon_PhAT_Play_40x.png) | **Play** - Plays the currently selected animation. Physics Blend must be less that 1.0 for any effect to show. |
 |[REGION:tablesection]Animation[/REGION]||
-![](PhAT_Animation.png) **Animation** - While the simulation is running, an animation can be selected here to play on the Skeletal Mesh.
+|![](PhAT_Animation.png) | **Animation** - While the simulation is running, an animation can be selected here to play on the Skeletal Mesh.|
 
 ### Bodies Mode Tool Bar
 
@@ -172,17 +172,17 @@ While in Constraint Mode, these buttons will appear to the right of the **Constr
 
 ## Viewport Pane
 
-The **Viewport** Pane shows a rendered view of the default Skeletal Mesh associated with the Physics Asset. It also shows any physical bodies associated with the ones in the skeleton of the default mesh, any constraints between those bodies, visual tools to manipulate those bodies and constraints, and information about those bodies and constraints. The rendering mode for the mesh, the bodies, and the constraints can all be adjusted individually to get just the right preview for the task at hand. Perhaps most importantly, the physics simulation of the bodies and constraints can be previewed directly in the viewport and the mesh can be manipulated (nudging it or grabbing and moving it) to see how they react to particular actions.
+The Viewport Pane shows a rendered view of the default Skeletal Mesh associated with the Physics Asset. It also shows any physical bodies associated with the ones in the skeleton of the default mesh, any constraints between those bodies, visual tools to manipulate those bodies and constraints, and information about those bodies and constraints. The rendering mode for the mesh, the bodies, and the constraints can all be adjusted individually to get just the right preview for the task at hand. Perhaps most importantly, the physics simulation of the bodies and constraints can be previewed directly in the viewport and the mesh can be manipulated (nudging it or grabbing and moving it) to see how they react to particular actions.
 
 
 ## Details Pane
 
-The **Details** pane is a standard Unreal properties window which shows either a set of basic editor properties if nothing is selected or the properties of the body or constraint currently selected in the **Hierarchy** Pane, depending on what mode the editor is in.
+The Details pane is a standard Unreal properties window which shows either a set of basic editor properties if nothing is selected or the properties of the body or constraint currently selected in the Hierarchy Pane, depending on what mode the editor is in.
 
 
 ## Hierarchy Pane
 
-The **Hierarchy** Pane shows a hierarchical view of the Bones of the Skeletal Mesh. Depending on which mode the editor is currently in, selecting a joint name from this list will select either a Body or a Constraint. Bones with Bodies or Constraints directly associated with them are shown in bold to distinguish them. Selecting Bones in the **Hierarchy** Pane will select their corresponding Body or Constraint (depending which mode is enabled).
+The Hierarchy Pane shows a hierarchical view of the Bones of the Skeletal Mesh. Depending on which mode the editor is currently in, selecting a joint name from this list will select either a Body or a Constraint. Bones with Bodies or Constraints directly associated with them are shown in bold to distinguish them. Selecting Bones in the Hierarchy Pane will select their corresponding Body or Constraint (depending which mode is enabled).
 
 
 ## Controls
@@ -214,13 +214,13 @@ The **Hierarchy** Pane shows a hierarchical view of the Bones of the Skeletal Me
 | **]** | Disable collision with the selected body and another. |
 | [REGION:tablesection]**Constraints Mode**[/REGION] ||
 | **Q** | Cycle constraint X -> Y -> Z. |
-| **1** | Switches **Angular Swing 1** of the selected **Physics Constraint** between **Limited** and **Locked**. |
-| **2** | Switches **Angular Swing 2** of the selected **Physics Constraint** between **Limited** and **Locked**. |
-| **3** | Switches **Angular Twist** of the selected **Physics Constraint** between **Limited** and **Locked**. |
-| **4** | Cycles the angular motions of the selected **Physics Constraint** locking 2 of them and limiting the 3rd. |
+| **1** | Switches **Angular Swing 1** of the selected Physics Constraint between **Limited** and **Locked**. |
+| **2** | Switches **Angular Swing 2** of the selected Physics Constraint between **Limited** and **Locked**. |
+| **3** | Switches **Angular Twist** of the selected Physics Constraint between **Limited** and **Locked**. |
+| **4** | Cycles the angular motions of the selected Physics Constraint locking 2 of them and limiting the 3rd. |
 
 [REGION:note]
-As of version 4.5,  PhAT allows you to copy a **Physics Body** from one bone to another based on the left/right naming convention.  If you use an `_L` `_R` naming convention on your joints, hitting the **M** key will copy the **Physics Body** to the bone that shares the same name, but has the opposing suffix.  However, this only copies the Physics Body and does not currently copy constraints. Also note that currently this does not mirror rotations.
+As of version 4.5,  PhAT allows you to copy a Physics Body from one bone to another based on the left/right naming convention.  If you use an `_L` `_R` naming convention on your joints, **hitting** the **M** key will copy the Physics Body to the bone that shares the same name, but has the opposing suffix.  However, this only copies the Physics Body and does not currently copy constraints. Also note that currently this does not mirror rotations.
 [/REGION:note]
 
 
