#UEDoc #docs - Updated doc images and bold lettering where needed.

diff --git a/Engine/Documentation/Source/Engine/Physics/PhAT/UserGuide/UserGuide.INT.udn b/Engine/Documentation/Source/Engine/Physics/PhAT/UserGuide/UserGuide.INT.udn
--- a/Engine/Documentation/Source/Engine/Physics/PhAT/UserGuide/UserGuide.INT.udn
+++ b/Engine/Documentation/Source/Engine/Physics/PhAT/UserGuide/UserGuide.INT.udn
@@ -44,15 +44,16 @@ Description:How to use the Physics Asset Tool to create Physics Asset for Skelet
 
 [TOC(start:2)]
 
-## Overview
 
-The **Physics Asset Tool** is used for creating **Physics Assets** for use with **Skeletal Meshes**. One can start from nothing and build to a full ragdoll setup or use the automation tools to create a basic set of Physics Bodies and Physics Constraints.
+The **Physics Asset Tool** is used for creating **Physics Assets** for use with **Skeletal Meshes**. One can start from nothing and build to a full ragdoll setup or use the automation 
+tools to create a basic set of Physics Bodies and Physics Constraints.
 
 This document will cover step by step procedures for the most common activities in the Physics Asset Tool.
 
 ## New Physics Asset
 
-When a Skeletal Mesh is imported, there is an option to generate a Physics Asset for it as it is imported. Once the imported file is processed, you will be presented with the "New Asset" dialog:
+When a Skeletal Mesh is imported, there is an option to generate a Physics Asset for it as it is imported. Once the imported file is processed, you will be presented with the 
+"New Asset" dialog:
 
 [REGION:raw]
 ![](PhATUG_ImportPA.png)
@@ -62,16 +63,16 @@ However, you can follow these steps if you need to create a Physics Asset for a
 
 1. Find the Skeletal Mesh asset you wish to add a Physics Asset to in the 
 **Content Browser**.
-1. **Right-click** on the **Skeletal Mesh** in the **Content Browser** to bring up the **Right Click Context Menu** -> **Create** -> **Create Physics Asset**.
+1. **Right-click** on the Skeletal Mesh in the Content Browser to bring up the **Right Click Context Menu** -> **Create** -> **Create Physics Asset**.
 1. Adjust setting to your liking.
 
 	![](Shared\PhAT\PhAT_restetDefault.png)
 
-1. Click **Ok**.
+1. **Click** 'Ok'.
 
 ## Editing a Physics Asset
 
-1. Find the Physics Asset in the **Content Browser**.
+1. Find the Physics Asset in the Content Browser.
 
 	![](Engine/Physics\PhysicalMaterials\PhysMatUserGuide\PhysMat_PhAT_02.png)
 
@@ -86,14 +87,18 @@ However, you can follow these steps if you need to create a Physics Asset for a
 
 1. Enable Body Editing Mode ![](Shared\Icons\icon_PHatMode_Body_40x.png).
 1. Select a Physics Body in the Viewport or in the **Hierarchy Pane**.
-	* If it is difficult to see the Physics Body you desire, use the **Viewport** -> **Modes** -> **MeshRender Mode** -> **Wireframe/None** to adjust the visibility of the Skeletal Mesh.
-1. Move, rotate, and scale the Physics Body using the Translation, Rotation, and Scale tools so it fits the Skeletal Mesh and Skeletal Joint.
-1. Edit the Physics Body's properties in the **Details** Pane.
-1. Enable or Disable Collision between the current Physics Bodies and other in the Physics Asset using the _Collision On_ and _Collision Off_ tools. Any Physics Body that can collide with the currently selected Physics Body will be blue, if they cannot collide, they will be gray.
-1. If a series of bodies should act as one, such as with the twist joints that form a wrist, use the _Weld_ tool to prevent undesired physics simulation behavior. Welded Physics Bodies will be yellow in color.
+	* If it is difficult to see the Physics Body you desire, use the **Viewport** -> **Modes** -> **MeshRender Mode** -> **Wireframe/None** to adjust the visibility of the 
+	Skeletal Mesh.
+1. **Move, rotate, and scale** the Physics Body using the Translation, Rotation, and Scale tools so it fits the Skeletal Mesh and Skeletal Joint.
+1. Edit the Physics Body's properties in the Details Pane.
+1. **Enable** or **Disable** Collision between the current Physics Bodies and other in the Physics Asset using the _Collision On_ and _Collision Off_ tools. Any Physics Body that can 
+collide with the currently selected Physics Body will be blue, if they cannot collide, they will be gray.
+1. If a series of bodies should act as one, such as with the twist joints that form a wrist, use the _Weld_ tool to prevent undesired physics simulation behavior. Welded Physics 
+Bodies will be yellow in color.
 1. Save ![](Shared\Icons\icon_SaveAsset_40x.png) often.
 
-See the [Physics Asset Tool Reference](Engine/Physics\PhAT\Reference) and the [Physics Asset Properties Reference](Engine/Physics\PhAT\Properties) pages for more information on each button and property available in the Physics Asset Tool.
+See the **[Physics Asset Tool Reference](Engine/Physics\PhAT\Reference)** and the **[Physics Asset Properties Reference](Engine/Physics\PhAT\Properties)** pages for more information on each 
+button and property available in the Physics Asset Tool.
 
 [PUBLISH:licensee]
 [REGION:note]
@@ -103,8 +108,8 @@ As of change list #1977557, multi select is enabled for Physics Bodies and Physi
 
 ### Deleting Physics Bodies
 
-1. Enable Body Editing Mode ![](Shared\Icons\icon_PHatMode_Body_40x.png).
-1. Select a Physics Body you wish to delete.
+1. **Enable** Body Editing Mode ![](Shared\Icons\icon_PHatMode_Body_40x.png).
+1. **Select** a Physics Body you wish to delete.
 	* If it is difficult to see the Physics Body you desire, use the **Viewport -> Modes -> MeshRender Mode -> Wireframe/None** to adjust the visibility of the Skeletal Mesh.
 1. Press the **Delete** key.
 
@@ -114,10 +119,10 @@ As of change list #1977557, multi select is enabled for Physics Bodies and Physi
 
 ![](PhATUG_addBody.png)(w:900)
 
-1. **Right-click** an already existing Physics Body or a Bone in the **Hierarchy** Panel.
-1. Select the type of Physics Body you wish to add from the Context Menu using the left mouse button.
+1. **Right-click** an already existing Physics Body or a Bone in the Hierarchy Panel.
+1. **Select** the type of Physics Body you wish to add from the Context Menu using the left mouse button.
 
-	If the Bone in the **Hierarchy** Panel does not already have a Physics Body on it, the only option in the **right-click** context menu is "Add Body."
+	If the Bone in the Hierarchy Panel does not already have a Physics Body on it, the only option in the **right-click** context menu is "Add Body."
 
 	Finally, you can just **double-click** a bone that does not have a Physics Body in the **Hierarchy** Panel to "Add Body."
 
@@ -148,8 +153,8 @@ As of change list #1977557, multi select is enabled for Physics Bodies and Physi
 [REGION:raw]
 **Enabling Collision**
 
-1. Select 2 or more Physics Bodies using **Ctrl + Left Mouse click**.
-1. Click the "Collision On" ![](Shared\Icons\icon_PhAT_EnableCollision_40x.png) button.
+1. **Select** 2 or more Physics Bodies using **Ctrl + Left Mouse click**.
+1. **Click** the "Collision On" ![](Shared\Icons\icon_PhAT_EnableCollision_40x.png) button.
 
 
 [/REGION]
@@ -158,8 +163,8 @@ As of change list #1977557, multi select is enabled for Physics Bodies and Physi
 [EXCERPT:ColOff]
 **Disabling Collision** 
 
-1. Select 2 or more Physics Bodies using **Ctrl + Left Mouse click**.
-1. Click the "Collision Off" ![](Shared\Icons\icon_PhAT_DisableCollision_40x.png) button.
+1. **Select** 2 or more Physics Bodies using **Ctrl + Left Mouse click**.
+1. **Click** the "Collision Off" ![](Shared\Icons\icon_PhAT_DisableCollision_40x.png) button.
 
 [/EXCERPT:ColOff]
 
@@ -170,8 +175,8 @@ As of change list #1977557, multi select is enabled for Physics Bodies and Physi
 
 Welding Physics Bodies will make multiple bodies interact as one, as well as lock together the Skeletal Mesh Joints they are associated with. To Weld one Physics Body to another:
 
-1. Select 2 or more Physics Bodies using **Ctrl + Left Mouse click**.
-1. Click the "Weld" ![](Shared\Icons\icon_PhAT_Weld_40x.png) button.
+1. **Select** 2 or more Physics Bodies using **Ctrl + Left Mouse click**.
+1. **Click** the "Weld" ![](Shared\Icons\icon_PhAT_Weld_40x.png) button.
 
 
 Any Physics Body welded to the currently selected Physics Body will show in Yellow.
@@ -181,10 +186,11 @@ Any Physics Body welded to the currently selected Physics Body will show in Yell
 
 Each Physics Body in a Physics Asset can have a Physical Material assigned to it. To apply a Physics Material to a single Physics Body:
 
-1. Enable Body Editing Mode ![](Shared\Icons\icon_PHatMode_Body_40x.png).
-1. Select a Physics Body in the Viewport or in the **Hierarchy** Pane.
-	* If it is difficult to see the Physics Body you desire, use the **Viewport** -> **Modes** -> **MeshRender Mode** -> **Wireframe/None** to adjust the visibility of the Skeletal Mesh.
-1. In the **Details** pane, locate the Simple Collision Physical Material property and assign a Physics Material.
+1. **Enable** Body Editing Mode ![](Shared\Icons\icon_PHatMode_Body_40x.png).
+1. **Select** a Physics Body in the Viewport or in the **Hierarchy** Pane.
+	* If it is difficult to see the Physics Body you desire, use the **Viewport** -> **Modes** -> **MeshRender Mode** -> **Wireframe/None** to adjust the visibility of the Skeletal 
+	Mesh.
+1. In the Details pane, locate the Simple Collision Physical Material property and assign a Physics Material.
 
 [REGION:raw]
 ![](Engine/Physics\PhysicalMaterials\PhysMatUserGuide\PhysMat_PhAT_04.png)(w:920)
@@ -194,25 +200,27 @@ Each Physics Body in a Physics Asset can have a Physical Material assigned to it
 ## Constraints
 
 [REGION:note]
-The use of Physics Constraints is covered in the [Physics Constraints User Guide](Engine/Physics\Constraints\ConstraintsUserGuide) and their properties are covered in the [Physics Constraint Reference](Engine/Physics\Constraints\ConstraintsReference). This section will only cover workflows that are specific to the Physics Asset Tool or ones that have deviated substantially from the norm.
+The use of Physics Constraints is covered in the **[Physics Constraints User Guide](Engine/Physics\Constraints\ConstraintsUserGuide)*8 and their properties are covered in the **[Physics Constraint Reference](Engine/Physics\Constraints\ConstraintsReference)**. This section will only cover workflows that are specific to the Physics Asset Tool or ones that have deviated substantially from the norm.
 [/REGION]
 
 ### Editing Physics Constraints
 
-1. Enable Constraint Editing Mode ![](Shared\Icons\icon_PHatMode_Joint_40x.png).
-1. Select a Physics Constraint in the Viewport or in the **Hierarchy** Pane.
-1. Move and rotate the Physics Constraint using the Translation and Rotation to create the rotational point of the "joint" the Physics Constraint will form.
-1. Edit the Physics Constraint's properties in the **Details** Pane.
+1. **Enable** Constraint Editing Mode ![](Shared\Icons\icon_PHatMode_Joint_40x.png).
+1. **Select** a Physics Constraint in the Viewport or in the Hierarchy Pane.
+1. **Move and rotate** the Physics Constraint using the Translation and Rotation to create the rotational point of the "joint" the Physics Constraint will form.
+1. Edit the Physics Constraint's properties in the Details Pane.
 
 	[REGION:tip]
-	You can switch Swing1, Swing2, and Twist from Limited to Locked quickly by using the "**1**", "**2**", and "**3**" keys respectively. The "**4**" key can be used cycle through them setting one to limited and the other two to locked.
+	You can switch Swing1, Swing2, and Twist from Limited to Locked quickly by using the "**1**", "**2**", and "**3**" keys respectively. The "**4**" key can be used cycle through 
+	them setting one to limited and the other two to locked.
 	[/REGION]
 
 1. Save ![](Shared\Icons\icon_SaveAsset_40x.png) often.
 
 
 
-See the [Physics Asset Tool Reference](Engine/Physics\PhAT\Reference) and the [Physics Asset Properties Reference](Engine/Physics\PhAT\Properties) pages for more information on each button and property available in the Physics Asset Tool.
+See the **[Physics Asset Tool Reference](Engine/Physics\PhAT\Reference)** and the **[Physics Asset Properties Reference](Engine/Physics\PhAT\Properties)** pages for more information on each 
+button and property available in the Physics Asset Tool.
 
 
 
@@ -223,9 +231,10 @@ If you are using the Physics Constraint's Linear or Angular limits, you will be
 
 ![](PhATUG_Limits.png)(w:300 convert:false)
 
-By then translating and rotating the Physics Constraint, you can align the limits to give the desired results. In the most basic of terms, when a Physics Constraint is limited, you can see a yellow line suspended in a green arc or cone structure. The line will be "constrained" within this arc or cone.
+By then translating and rotating the Physics Constraint, you can align the limits to give the desired results. In the most basic of terms, when a Physics Constraint is limited, 
+you can see a yellow line suspended in a green arc or cone structure. The line will be "constrained" within this arc or cone.
 
-For more information on Physics Constraints and their properties, see the [Physics Constraints User Guide](Engine/Physics\Constraints\ConstraintsUserGuide).
+For more information on Physics Constraints and their properties, see the **[Physics Constraints User Guide](Engine/Physics\Constraints\ConstraintsUserGuide)**.
 
 
 ### Deleting Physics Constraints
@@ -234,9 +243,9 @@ For more information on Physics Constraints and their properties, see the [Physi
 There is no easy way to recreate a Physics Constraint, be mindful of this when deleting them.
 [/REGION]
 
-1. Enable Constraint Editing Mode ![](Shared\Icons\icon_PHatMode_Joint_40x.png).
-1. Select the Physics Constraint you wish to delete.
-1. Press the **Delete** key.
+1. **Enable** Constraint Editing Mode ![](Shared\Icons\icon_PHatMode_Joint_40x.png).
+1. **Select** the Physics Constraint you wish to delete.
+1. **Press** the **Delete** key.
 
 
 ### Recreating a Physics Constraint
@@ -245,7 +254,9 @@ There is no easy way to recreate a Physics Constraint, be mindful of this when d
 There is no easy way to recreate a Physics Constraint, be mindful of this when deleting them.
 [/REGION]
 
-Physics Constraints are only created on generation of a Physics Body and are only generated upstream. So if you remove a Physics Constraint from the shoulder, you would have to remove the upper arm Physics Body (which will remove the elbow Physics Constraint), then re-create the upper arm Physics Body. This would create the shoulder Physics Constraint, but not the elbow Physics Constraint, so you would have to continue this process down the arm.
+Physics Constraints are only created on generation of a Physics Body and are only generated upstream. So if you remove a Physics Constraint from the shoulder, you would have to 
+remove the upper arm Physics Body (which will remove the elbow Physics Constraint), then re-create the upper arm Physics Body. This would create the shoulder Physics Constraint, 
+but not the elbow Physics Constraint, so you would have to continue this process down the arm.
 
 
 ### Constraint Snapping
@@ -261,10 +272,10 @@ You can snap the currently selected Physics Constraint(s) by using the Snap Cons
 
 While in either mode, properties from one Physics Body or Constraint can be copied to other Physics Bodies or Constraints:
 
-1. Select the Physics Body or Constraint you wish to copy data from.
-1. Press **Ctrl + C**.
-1. Select the Physics Bodies or Constraints you wish to apply data to.
-1. Press **Ctrl + V**.
+1. **Select** the Physics Body or Constraint you wish to copy data from.
+1. **Press Ctrl + C**.
+1. **Select** the Physics Bodies or Constraints you wish to apply data to.
+1. **Press Ctrl + V**.
 
 
 This will not overwrite the shape of a Physics Body.
@@ -272,11 +283,11 @@ This will not overwrite the shape of a Physics Body.
 
 ## Testing
 
-Clicking the **Simulation** ![](Shared\Icons\icon_PhAT_PlaySim_40x.png) button will allow you to test your Physics Asset.
+Clicking the Simulation ![](Shared\Icons\icon_PhAT_PlaySim_40x.png) button will allow you to test your Physics Asset.
 
 ![](Shared\PhAT\PhAT_Sim.png)
 
-The downward triangle next to the **Simulation** button will give you access to **Real Simulation** and **No Gravity**:
+The downward triangle next to the Simulation button will give you access to **Real Simulation** and **No Gravity**:
 
 * **Real Simulation** - Simulates the entire Physics Asset with gravity enabled.
 * **No Gravity** - Simulates the entire Physics Asset but does not turn on gravity, leaving you to **ctrl+click** to poke the Physics Bodies in a zero gravity environment.
@@ -285,29 +296,37 @@ The downward triangle next to the **Simulation** button will give you access to
 
 You can also simulate a chain of joints by toggling **Selected Simulation** ![](selectedToggle_button.png) on, then clicking the **Simulation** ![](Shared\Icons\icon_PhAT_PlaySim_40x.png) button.
 
-* **Selected Simulation** - Does mostly what it sounds like, it only simulates the **Physics Bodies** you have selected (using **ctrl+left click** you can select more than one Physics Body) **and** those down the hierarchy from the selected Physics Bodies. So if you select the shoulder, the entire arm will simulate.
+* **Selected Simulation** - Does mostly what it sounds like, it only simulates the Physics Bodies you have selected (using **ctrl+left click** you can select more than one 
+Physics Body) and those down the hierarchy from the selected Physics Bodies. So if you select the shoulder, the entire arm will simulate.
 
 
 
 ## Exploding
 
-This is caused by two Physics Bodies interpenetrating and the physics system attempting to correct it by applying a massive amount of force to expel the Physics Bodies from each other. If a Physics Constraint is keeping the offending Physics Bodies together, the physics system will continually apply force to separate them, causing very erratic and extreme movement.
+This is caused by two Physics Bodies interpenetrating and the physics system attempting to correct it by applying a massive amount of force to expel the Physics Bodies from each other. 
+If a Physics Constraint is keeping the offending Physics Bodies together, the physics system will continually apply force to separate them, causing very erratic and extreme movement.
 
 ![](brokenPhAT.png)
 
-This is normally fixed by disabling collision between the two offending Physics Bodies, or changing their positions and/or scale to make sure they do not interpenetrate. Also, welded Physics Bodies act as one regardless of interpenetration.
+This is normally fixed by disabling collision between the two offending Physics Bodies, or changing their positions and/or scale to make sure they do not interpenetrate. Also, welded 
+Physics Bodies act as one regardless of interpenetration.
 
 ## Jitter
 
-If the Physics Asset is mostly "collapsed" but still shaking and twitching on the ground, but not violently bouncing around, there are a few values you can tweak to get it to stop and sleep. 
+If the Physics Asset is mostly "collapsed" but still shaking and twitching on the ground, but not violently bouncing around, there are a few values you can tweak to get it to stop 
+and sleep. 
 
-Before doing anything, try simulating with **No Gravity**, this will show you if any of your Physics Constraints are miss aligned and trying to correct before the Physics Asset has even hit the ground.
+Before doing anything, try simulating with No Gravity, this will show you if any of your Physics Constraints are miss aligned and trying to correct before the Physics Asset has 
+even hit the ground.
 
-Often times, a small amount of **Linear** and **Angular Damping** on all the Physics Asset's Physics Bodies is more than enough to stop the asset from shaking. However, using large values in **Linear Damping** will cause a Physics Body slow its movement through the world, even due to gravity. **Damping**, as it is implemented here, is not "drag" per say. It is more like moving through a viscous substance.
+Often times, a small amount of Linear and Angular Damping on all the Physics Asset's Physics Bodies is more than enough to stop the asset from shaking. However, using large 
+values in Linear Damping will cause a Physics Body slow its movement through the world, even due to gravity. Damping, as it is implemented here, is not "drag" per se. It is 
+more like moving through a viscous substance.
 
-If jitter is still an issue, perhaps you have a number of really small Physics Bodies, check the [Substepping](Engine/Physics/Substepping) documentation for information on having the physics simulation calculate in-between solutions.
+If jitter is still an issue, perhaps you have a number of really small Physics Bodies, check the **[Substepping](Engine/Physics/Substepping)** documentation for information on having the 
+physics simulation calculate in-between solutions.
 
-See the [](Engine/Physics/PhAT/Reference) and [](Engine/Physics/Constraints/ConstraintsReference) for more details on the properties for Physics Constraints.
+See the **[](Engine/Physics/PhAT/Reference)** and **[](Engine/Physics/Constraints/ConstraintsReference)** for more details on the properties for Physics Constraints.
 
 
 
