#UE4doc #docs Updated Level Design Quick Start to meet Instructional Guide (v 1.0) requirements.

diff --git a/Engine/Documentation/Source/Engine/QuickStart/1/LDQS_1.INT.udn b/Engine/Documentation/Source/Engine/QuickStart/1/LDQS_1.INT.udn
--- a/Engine/Documentation/Source/Engine/QuickStart/1/LDQS_1.INT.udn
+++ b/Engine/Documentation/Source/Engine/QuickStart/1/LDQS_1.INT.udn
@@ -25,12 +25,12 @@ Related:GettingStarted/Installation
 
 %Steps%
 
-Before we get started, the first thing that we will need to do is create a [**Project**](Engine/Basics/Projects) to work from for the duration of this guide. In the steps below we will use Unreal Engine 4's [**Project Browser**](Engine/Basics/Projects/Browser) to create our project. 
+Before we get started, the first thing that we will need to do is create a [Project](Engine/Basics/Projects) to work from for the duration of this guide. In the steps below we will use Unreal Engine 4's [Project Browser](Engine/Basics/Projects/Browser) to create our project. 
 After completing this step, our project will be ready and you will gone through the process of creating a Project and using a Template which you can then use for your very own future projects. 
 
 ## Steps
 
-1. After [**Installing Unreal Engine**](GettingStarted/Installation) and launching the Unreal Editor, you will first be presented with the **Project Browser**. 
+1. After [Installing Unreal Engine](GettingStarted/Installation) and launching the Unreal Editor, you will first be presented with the **Project Browser**. 
 
 1. Click the **New Project** tab and select the **Blank** Blueprint Template, enter a Save Location and Name, then click **Create Project**. 
 
diff --git a/Engine/Documentation/Source/Engine/QuickStart/2/LDQS_2.INT.udn b/Engine/Documentation/Source/Engine/QuickStart/2/LDQS_2.INT.udn
--- a/Engine/Documentation/Source/Engine/QuickStart/2/LDQS_2.INT.udn
+++ b/Engine/Documentation/Source/Engine/QuickStart/2/LDQS_2.INT.udn
@@ -26,12 +26,12 @@ Related:Engine/UI/LevelEditor/Viewports/ViewportControls
 
 %Steps%
 
-With the project open and ready to go, the first thing you may notice is the [**Viewport**](Engine/UI/LevelEditor/Viewports) in the center of the Unreal Editor. 
+With the project open and ready to go, the first thing you may notice is the [Viewport](Engine/UI/LevelEditor/Viewports) in the center of the Unreal Editor. 
 
 ![](TableAndChairs.png)(w:720)
 
-Inside the Viewport is where you will do most of your level construction. The template project that we selected in the previous step includes a small sample [**Level**](Engine/Levels) and some assets for us to get us started with. 
-Using this little area as a point of reference, take a moment to get used to the [**Viewport Camera Controls**](Engine/UI/LevelEditor/Viewports/ViewportControls) by using the most common methods of navigating the Viewport in Unreal Engine 4. 
+Inside the Viewport is where you will do most of your level construction. The template project that we selected in the previous step includes a small sample [Level](Engine/Levels) and some assets for us to get us started with. 
+Using this little area as a point of reference, take a moment to get used to the [Viewport Camera Controls](Engine/UI/LevelEditor/Viewports/ViewportControls) by using the most common methods of navigating the Viewport in Unreal Engine 4. 
 
 
 ## Standard Controls
@@ -39,7 +39,7 @@ Using this little area as a point of reference, take a moment to get used to the
 [INCLUDE:Engine\UI\LevelEditor\Viewports\ViewportControls#StandardControls]
 
 [REGION:tip]
-To switch to an [**Orthographic View**](https://wiki.unrealengine.com/Videos/Player?series=PLZlv_N0_O1gasd4IcOe9Cx9wHoBB7rxFl&video=RoiQOwCg-4Q), click the **Perspective** drop-down then select an Orthographic View Mode. 
+To switch to an [Orthographic View](https://wiki.unrealengine.com/Videos/Player?series=PLZlv_N0_O1gasd4IcOe9Cx9wHoBB7rxFl&video=RoiQOwCg-4Q), click the **Perspective** drop-down then select an Orthographic View Mode. 
 
 ![](OrthographicViews.png)
 [/REGION]
diff --git a/Engine/Documentation/Source/Engine/QuickStart/3/LDQS_3.INT.udn b/Engine/Documentation/Source/Engine/QuickStart/3/LDQS_3.INT.udn
--- a/Engine/Documentation/Source/Engine/QuickStart/3/LDQS_3.INT.udn
+++ b/Engine/Documentation/Source/Engine/QuickStart/3/LDQS_3.INT.udn
@@ -25,8 +25,8 @@ Related:Engine/Content/Browser
 
 %Steps%
 
-Next we will create a new [**Level**](Engine/Levels) that we will use to build our game environment in. 
-While there are several different ways in which you can create a new Level, we will use the most common method by using the [**Content Browser**](Engine/Content/Browser) which is used to create several different types of assets inside the Unreal Editor. 
+Next we will create a new [Level](Engine/Levels) that we will use to build our game environment in. 
+While there are several different ways in which you can create a new Level, we will use the most common method by using the [Content Browser](Engine/Content/Browser) which is used to create several different types of assets inside the Unreal Editor. 
 
 
 ## Steps
diff --git a/Engine/Documentation/Source/Engine/QuickStart/4/LDQS_4.INT.udn b/Engine/Documentation/Source/Engine/QuickStart/4/LDQS_4.INT.udn
--- a/Engine/Documentation/Source/Engine/QuickStart/4/LDQS_4.INT.udn
+++ b/Engine/Documentation/Source/Engine/QuickStart/4/LDQS_4.INT.udn
@@ -29,13 +29,13 @@ Related:Engine/Rendering/LightingAndShadows
 
 %Steps%
 
-In this step we will begin placing [**Actors**](Engine/Actors) (which are things like lights or geometry for example) into our empty level. 
-We will cover the two most common ways of adding Actors to the level, one of which is through [**Place Mode**](Engine/UI/LevelEditor/Modes/PlaceMode) and the other is through the [**Content Browser**](Engine/Content/Browser). 
+In this step we will begin placing [Actors](Engine/Actors) (which are things like lights or geometry for example) into our empty level. 
+We will cover the two most common ways of adding Actors to the level, one of which is through [Place Mode](Engine/UI/LevelEditor/Modes/PlaceMode) and the other is through the [Content Browser](Engine/Content/Browser). 
 After completing this step, you will know how to place Actors inside your own levels and can begin manipulating those Actors to create your environment. 
 
 ## Steps
 
-1. In the **Modes Panel**, with **Place Mode** enabled, click the [**BSP**](Engine/Actors/Brushes) category and select the **Box**. 
+1. In the **Modes Panel**, with **Place Mode** enabled, click the [BSP](Engine/Actors/Brushes) category and select the **Box**. 
 
 	![](ModesPanel.png)
 
@@ -47,7 +47,7 @@ After completing this step, you will know how to place Actors inside your own le
 
 	![](place2.png)
 
-1. With the **Box** still selected, in the [**Details**](Engine/UI/LevelEditor/Details) panel (lower-right window of the editor), set **Location** and **Rotation** all to **0**. 
+1. With the **Box** still selected, in the [Details](Engine/UI/LevelEditor/Details) panel (lower-right window of the editor), set **Location** and **Rotation** all to **0**. 
 
 1. Set the **Scale** to **4 x 4 x 0.1**. 
 
@@ -59,29 +59,29 @@ After completing this step, you will know how to place Actors inside your own le
 
 	![](placeLight.png)
 	
-1. On the on the [**Translation Tool**](https://docs.unrealengine.com/latest/INT/Engine/UI/LevelEditor/Viewports/ViewportControls/index.html#transformcontrols), click and drag the Z Axis (Blue) gizmo up away from the surface of the floor.
+1. On the on the [Translation Tool](https://docs.unrealengine.com/latest/INT/Engine/UI/LevelEditor/Viewports/ViewportControls/index.html#transformcontrols), click and drag the Z Axis (Blue) gizmo up away from the surface of the floor.
 
 	![](moveLight.png)
 
 	If the **Directional Light** becomes unselected, you can re-select it by **Left-clicking** on it in the **Level Viewport**. 
 
-1. In the **Modes Panel**, select the **Visual Effects** tab and drag-and-drop an [**Atmospheric Fog**](Engine/Actors/FogEffects/AtmosphericFog) to the level. 
+1. In the **Modes Panel**, select the **Visual Effects** tab and drag-and-drop an [Atmospheric Fog](Engine/Actors/FogEffects/AtmosphericFog) to the level. 
 
 	![](placeFog.png)
 	
 	The **Atmospheric Fog** Actor will add a basic sky to the level and the level will become illuminated instead of dark. 
 
-1. From the **Modes Panel** on the **Basic** tab, drag-and-drop a [**Player Start**](Engine/Actors/PlayerStart) into the level. 
+1. From the **Modes Panel** on the **Basic** tab, drag-and-drop a [Player Start](Engine/Actors/PlayerStart) into the level. 
 
 	![](placeStart.png)
 
-1. From the **Modes Panel** on the **Volumes** tab, drag-and-drop a [**Lightmass Importance Volume**](Engine/Rendering/LightingAndShadows/Lightmass/Basics) into the level. 
+1. From the **Modes Panel** on the **Volumes** tab, drag-and-drop a [Lightmass Importance Volume](Engine/Rendering/LightingAndShadows/Lightmass/Basics) into the level. 
 
 	![](placeVolume.png)
 
 	The **Lightmass Importance Volume** is used to control and concentrate lighting and shadowing affects within the volume. When placing the Lightmass Importance Volume in the level, the default size of the volume does not cover our playable area so we will need to scale it up. 
 	
-1. Inside the **Level Viewport**, press **R** to switch to the [**Scale Tool**](https://docs.unrealengine.com/latest/INT/Engine/UI/LevelEditor/Viewports/ViewportControls/index.html#transformcontrols). 
+1. Inside the **Level Viewport**, press **R** to switch to the [Scale Tool](https://docs.unrealengine.com/latest/INT/Engine/UI/LevelEditor/Viewports/ViewportControls/index.html#transformcontrols). 
 
 1. Click and drag the white box in the center of the **Scale Tool** so that the **Lightmass Importance Volume** encapsulates the floor. 
 
@@ -97,7 +97,7 @@ After completing this step, you will know how to place Actors inside your own le
 
 	![](addChair.png)
 
-1. With the **SM_Chair** selected inside the **Level Viewport**, press **E** to access the [**Rotation Tool**](https://docs.unrealengine.com/latest/INT/Engine/UI/LevelEditor/Viewports/ViewportControls/index.html#transformcontrols). 
+1. With the **SM_Chair** selected inside the **Level Viewport**, press **E** to access the [Rotation Tool](https://docs.unrealengine.com/latest/INT/Engine/UI/LevelEditor/Viewports/ViewportControls/index.html#transformcontrols). 
 	
 	![](Rotate1.png)(w:440)
 
diff --git a/Engine/Documentation/Source/Engine/QuickStart/5/LDQS_5.INT.udn b/Engine/Documentation/Source/Engine/QuickStart/5/LDQS_5.INT.udn
--- a/Engine/Documentation/Source/Engine/QuickStart/5/LDQS_5.INT.udn
+++ b/Engine/Documentation/Source/Engine/QuickStart/5/LDQS_5.INT.udn
@@ -27,8 +27,8 @@ Related:
 
 %Steps%
 
-With several different **Actors** placed inside our level, the next step involves [**Editing Actor Properites** ](GettingStarted/HowTo/DetailsPanel) which can change the look or the way the Actor functions in the level, giving us a more customized looking level. 
-We will start with editing the properties of our **Directional Light Actor** then shift our focus to applying [**Materials**](Engine/Rendering/Materials) to some of the [**Static Mesh Actors**](Engine/Actors/StaticMeshActor) that you have placed in your level. 
+With several different **Actors** placed inside our level, the next step involves [Editing Actor Properites ](GettingStarted/HowTo/DetailsPanel) which can change the look or the way the Actor functions in the level, giving us a more customized looking level. 
+We will start with editing the properties of our **Directional Light Actor** then shift our focus to applying [Materials](Engine/Rendering/Materials) to some of the [Static Mesh Actors](Engine/Actors/StaticMeshActor) that you have placed in your level. 
 
 Once you have finished this step, you will have seen where to access and modify the properties of Actors so that you can begin editing and experimenting with different settings inside your own levels. 
 
diff --git a/Engine/Documentation/Source/Engine/QuickStart/7/LDQS_7.INT.udn b/Engine/Documentation/Source/Engine/QuickStart/7/LDQS_7.INT.udn
--- a/Engine/Documentation/Source/Engine/QuickStart/7/LDQS_7.INT.udn
+++ b/Engine/Documentation/Source/Engine/QuickStart/7/LDQS_7.INT.udn
@@ -39,23 +39,23 @@ Using the methods that were provided during the course of this guide, try do to
 * Add different kinds of lights with different colors. 
 * Use different Materials on some of your Actors. 
 
-For more information on the topics covered in this quick start guide, and across the entire editor, see the [**Unreal Editor Manual**](Engine).
+For more information on the topics covered in this quick start guide, and across the entire editor, see the [Unreal Editor Manual](Engine).
 
 As for specifics covered in this quick start:
 
-* For more information on the Level Editor, see: [**Level Editor**](Engine\UI\LevelEditor)
+* For more information on the Level Editor, see: [Level Editor](Engine\UI\LevelEditor)
 
-* For more information on Viewports, see: [**Viewports**](Engine\UI\LevelEditor\Viewports)
+* For more information on Viewports, see: [Viewports](Engine\UI\LevelEditor\Viewports)
 
-* For more information on the Editing Modes available in Unreal Editor, 4 see: [**Level Editor Modes**](Engine\UI\LevelEditor\Modes)
+* For more information on the Editing Modes available in Unreal Editor, 4 see: [Level Editor Modes](Engine\UI\LevelEditor\Modes)
 
-* For more information on the **Content Browser**, see: [**Content Browser**](Engine\Content\Browser)
+* For more information on the **Content Browser**, see: [Content Browser](Engine\Content\Browser)
 
-* For more information on the **Details** Panel, see: [**Details Panel**](Engine\UI\LevelEditor\Details)
+* For more information on the **Details** Panel, see: [Details Panel](Engine\UI\LevelEditor\Details)
 
-* For more information on Building, see: [**Lightmass**](Engine\Rendering\LightingAndShadows\Lightmass)
+* For more information on Building, see: [Lightmass](Engine\Rendering\LightingAndShadows\Lightmass)
 
-* For more information on Lighting, see: [**Lighting Quick Start Guide**](Engine/Rendering/LightingAndShadows/QuickStart)
+* For more information on Lighting, see: [Lighting Quick Start Guide](Engine/Rendering/LightingAndShadows/QuickStart)
 
 %Steps%
 
