TypoCorrection

diff --git a/Engine/Documentation/Source/Engine/Rendering/PostProcessEffects/Blendables/Blendables.INT.udn b/Engine/Documentation/Source/Engine/Rendering/PostProcessEffects/Blendables/Blendables.INT.udn
--- a/Engine/Documentation/Source/Engine/Rendering/PostProcessEffects/Blendables/Blendables.INT.udn
+++ b/Engine/Documentation/Source/Engine/Rendering/PostProcessEffects/Blendables/Blendables.INT.udn
@@ -1,4 +1,4 @@
-Availability:Public
+﻿Availability:Public
 Title:Blendables
 Crumbs: %ROOT%, Engine, Engine/Rendering/PostProcessEffects
 Description:Blendables assets can be smoothly interpolated and used to affecting the rendering (e.g. post processing, fog, Ambient Cubemap, ambient occlusion). 
@@ -40,7 +40,7 @@ The benefits of the new system:
 * Indirection over packages allows to adjust the content without level access (version control)
 * Indirection means a single asset can be reused in many occasions (less redundancy, more power)
 * Custom UI for each blendable is possible (much harder to do with a single struct)
-* Each Blendable reference can have it's own weight, the asset can have a weight (see LightPropagationVolumeBlendable) and a per property weight would be easy to do.
+* Each Blendable reference can have its own weight, the asset can have a weight (see LightPropagationVolumeBlendable) and a per-property-weight would be easy to do.
 * Breaking up the large struct make interaction with blueprints more efficient and simpler.
 
 ## The Blendables Container
@@ -58,8 +58,8 @@ the reference can be to an asset living in a package (created with the content b
 
 When you create new elements in the array you can choose to create a blendable of a specitic type (currently only LightPropagationVolumeBlendable) or using an asset reference
 (e.g. LightPropagationVolumeBlendable, Material, Material Instance). Over time we intend to create more blendable types (e.g. Bloom, SceneColor, DepthOfField, ...). The reference
-can be of any type that is a Blendable (implements the IBlendableInterface). The order in the array is opposite to how layers stack because they blending is applied form top to bottom
-and blending is overwriting the data that was tehre before. Keep in mind the data of many volumes (or other objects) are getting combined taking weight and priority into account.
+can be of any type that is a Blendable (implements the IBlendableInterface). The order in the array is opposite to how layers stack because their blending is applied from top to bottom
+and blending is overwriting the data that was there before. Keep in mind the data of many volumes (or other objects) are getting combined taking weight and priority into account.
 
 Note: It's a good practice to have a unbound PostProcessVolume with lower priority in the level called "global". To get get full control over an existing level you can add
 an unbound volume with high priority. To check if a blendable has effect you can quickly adjust it's weight to 0 and back.
@@ -73,7 +73,7 @@ coding them in the code. The code method is more involved and makes it harder fo
 [REGION:imagetable]
 | ![](AssetWithDifferentOuter.png) |
 | ------ |
-| The LightPropagationVolumeBlendable details as they can be seen in the editor. No matter if the blendable was created in the content browser (left) or created in the object (e.g. post process volume) the uiser interface is similar. It's good practice to give each property a checkbox (weight=0 or weight=1) and give the whole struct a blend weight. |
+| The LightPropagationVolumeBlendable details as they can be seen in the editor. No matter if the blendable was created in the content browser (left) or created in the object (e.g. post process volume) the user interface is similar. It's good practice to give each property a checkbox (weight=0 or weight=1) and give the whole struct a blend weight. |
 [/REGION]
 
 Note: The LightPropagationVolumeBlendable properties still exist in the PostProcessSettings for backwards compatibility reasons but we plan to remove the settings from there and make the blendables the only method.
@@ -87,7 +87,7 @@ The method **GetSingleFinalDataConst()** is used to get the data after it was bl
 
 The **AddOrUpdateBlendable** blueprint function is exposed where you can find PostProcessSettings. It allows convenient access to the blendables container. You pass in the object that has the
 Blendables container, the weight and a reference to the blendable. If the reference was already found in the container it simply updates the weight. It doesn't remove container elements as that
-could confuse other code traversing the container and has implications to the garbage collection. There is no real performance cost a blendable reference of a weight of 0 ss removing the element
+could confuse other code traversing the container and has implications to the garbage collection. There is no real performance cost a blendable reference of a weight of 0 as removing the element
 should never be needed.
 
 Here you can see how to reference a blendable asset in the content browser:
