TypoCorrection

diff --git a/Engine/Documentation/Source/Gameplay/Tools/VisualLogger/VisualLogger.INT.udn b/Engine/Documentation/Source/Gameplay/Tools/VisualLogger/VisualLogger.INT.udn
--- a/Engine/Documentation/Source/Gameplay/Tools/VisualLogger/VisualLogger.INT.udn
+++ b/Engine/Documentation/Source/Gameplay/Tools/VisualLogger/VisualLogger.INT.udn
@@ -1,4 +1,4 @@
-Title:Visual Logger
+﻿Title:Visual Logger
 Availability:Public
 Crumbs:
 Description: Tool that captures state from actors and then displays it visually in game or editor.
