TypoCorrection

diff --git a/Engine/Documentation/Source/Gameplay/Tools/GameplayDebugger/GameplayDebugger.INT.udn b/Engine/Documentation/Source/Gameplay/Tools/GameplayDebugger/GameplayDebugger.INT.udn
--- a/Engine/Documentation/Source/Gameplay/Tools/GameplayDebugger/GameplayDebugger.INT.udn
+++ b/Engine/Documentation/Source/Gameplay/Tools/GameplayDebugger/GameplayDebugger.INT.udn
@@ -1,4 +1,4 @@
-Title:Gameplay Debugger
+﻿Title:Gameplay Debugger
 Availability:Public
 Crumbs:
 Description: Tool that enables analyzing realtime gameplay data at runtime.
@@ -43,7 +43,7 @@ While working in editor GDT can be used in PIE or in Simulate. Binded key or *En
 
 ## Basic extension
 
-Gameplay Debugger can be only extended with C++ code. For Blueprint projects it can only be used as is - to display basic debug information for now. It's quite easy to extend Gameplay Debugger, to collect and display game specific data. Custom classes inherited from *UGameplayDebuggingComponent *class and *AGameplayDebuggingHUDComponent *class are needed for that. FIrst class is used to collect and eventually replicated data and second class is used to display all collected data on screen.
+Gameplay Debugger can be only extended with C++ code. For Blueprint projects it can only be used as is - to display basic debug information for now. It's quite easy to extend Gameplay Debugger, to collect and display game specific data. Custom classes inherited from *UGameplayDebuggingComponent *class and *AGameplayDebuggingHUDComponent *class are needed for that. First class is used to collect and eventually replicate the data and second class is used to display all collected data on screen.
 
 Below is a simple class which collects game specific data:
 
