TypoCorrection

diff --git a/Engine/Documentation/Source/Gameplay/Tools/NetworkProfiler/NetworkProfiler.INT.udn b/Engine/Documentation/Source/Gameplay/Tools/NetworkProfiler/NetworkProfiler.INT.udn
--- a/Engine/Documentation/Source/Gameplay/Tools/NetworkProfiler/NetworkProfiler.INT.udn
+++ b/Engine/Documentation/Source/Gameplay/Tools/NetworkProfiler/NetworkProfiler.INT.udn
@@ -1,4 +1,4 @@
-Title:Network Profiler
+﻿Title:Network Profiler
 Availability:Public
 Crumbs:
 Description:Tool for displaying network traffic and performance information captured at runtime.
@@ -23,7 +23,9 @@ The data file will be saved to <project directory>/Saved/Profiling/<project name
 
 The network profiler is a standalone application located at <UE install directory>/Engine/Binaries/DotNET/NetworkProfiler.exe.
 
-At the top of the window, click the "Open File" button to select a .nprof file and view the data contained within.![image alt text](image_0.png)
+At the top of the window, click the "Open File" button to select a .nprof file and view the data contained within.
+
+![image alt text](image_0.png)
 
 ### Chart, Filters, Details tab
 
@@ -48,11 +50,11 @@ This is a list of all the stats tracked by the network profiler. Toggling the ch
 
 #### Summary View
 
-This is a summary of the data for the time span that has been selected in the chart. It shows raw totals as well as per-second information. Note that the "Outgoing bandwidth" will be greater than the “Game socket send size” because “Outgoing bandwidth” includes an approximation of the IP and UDP header sizes as well.
+This is a summary of the data for the time span that has been selected in the chart. It shows raw totals as well as per-second information. Note that the "Outgoing bandwidth" will be greater than the "Game socket send size" because "Outgoing bandwidth" includes an approximation of the IP and UDP header sizes as well.
 
 #### Frame Details
 
-This view breaks down the data that was sent on the currently selected frame. You can see which actors, properties, and remote functions were sent, as well as how many bytes each one contributed. Note that replicated actors and RPCs are both included in bunches, so the "NumBytes" field in “SEND BUNCH” sections includes the bytes from “SEND RPC” and “REPLICATE ACTOR” sections. Ultimately, the “SOCKET SEND TO” entries will tell you exactly how many bytes the engine sent to the outgoing socket. Also note that the standard UDP packet overhead is not included here.
+This view breaks down the data that was sent on the currently selected frame. You can see which actors, properties, and remote functions were sent, as well as how many bytes each one contributed. Note that replicated actors and RPCs are both included in bunches, so the "NumBytes" field in "SEND BUNCH" sections includes the bytes from "SEND RPC" and "REPLICATE ACTOR" sections. Ultimately, the "SOCKET SEND TO" entries will tell you exactly how many bytes the engine sent to the outgoing socket. Also note that the standard UDP packet overhead is not included here.
 
 #### Filters
 
@@ -68,7 +70,7 @@ While a single frame is selected in the graph, this view will show a list of act
 
 Each of these tabs show a summary of the respective replication data for the entire profiling session. Note that the data displayed in these tabs is not affected by the current frame or time span selected in the chart on the "Chart, Filters, Details" tab.
 
-The tabs show a list of actors, properties, or RPCs, respectively. The "Total Size" and “Average Size” columns show how much bandwidth was needed for a particular item, and the “Count” column represents how many times a particular item was replicated during the profiling session. You can toggle between sorting by a column in ascending or descending order by clicking on the column heading.
+The tabs show a list of actors, properties, or RPCs, respectively. The "Total Size" and "Average Size" columns show how much bandwidth was needed for a particular item, and the "Count" column represents how many times a particular item was replicated during the profiling session. You can toggle between sorting by a column in ascending or descending order by clicking on the column heading.
 
 ## Server and client considerations
 
