#UE4doc #docs: Singular vs Plural grammar change: item -> items

diff --git a/Engine/Documentation/Source/Engine/UMG/UserGuide/WidgetTypeReference/WrapBox/WrapBox.INT.udn b/Engine/Documentation/Source/Engine/UMG/UserGuide/WidgetTypeReference/WrapBox/WrapBox.INT.udn
--- a/Engine/Documentation/Source/Engine/UMG/UserGuide/WidgetTypeReference/WrapBox/WrapBox.INT.udn
+++ b/Engine/Documentation/Source/Engine/UMG/UserGuide/WidgetTypeReference/WrapBox/WrapBox.INT.udn
@@ -31,7 +31,7 @@ In the **Details** panel of a Widget that is placed inside a **Wrap Box**, there
 
 ## Usage Example
 
-The **Wrap Box** is useful when you want to arrange items within a container and have those item appear aligned and evenly spaced and automatically format the contents inside it based on the width of the container. 
+The **Wrap Box** is useful when you want to arrange items within a container and have those items appear aligned and evenly spaced and automatically format the contents inside it based on the width of the container. 
 
 Take for example the image below where we have created a "shop" screen of sorts where several items are presented to the player. 
 
