#UE4doc #docs: ##Overview removed as according to the instructional guide.

diff --git a/Engine/Documentation/Source/Gameplay/ClassCreation/ClassCreationBasics.INT.udn b/Engine/Documentation/Source/Gameplay/ClassCreation/ClassCreationBasics.INT.udn
--- a/Engine/Documentation/Source/Gameplay/ClassCreation/ClassCreationBasics.INT.udn
+++ b/Engine/Documentation/Source/Gameplay/ClassCreation/ClassCreationBasics.INT.udn
@@ -5,9 +5,6 @@ Description:Examples showing how to create classes with Blueprints alone, C++ al
 
 ![](LightSwitchActor.png)(w:500 convert:false) 
 
-
-## Overview
-
 These examples show how to create a new class using Blueprints only, C++ code only, and a combination of C++ code and Blueprints. The goal is to create new LightSwitch classes with the same properties and behavior
 using all three workflows, and then add an instance of each new class to the level, so there are three new LightSwitch Actors. 
 
