#UE4doc #docs:
Unbolded a few terms.
Removed redundant 'To compile the project: ' line after 'Building the project'
Added 'Community' to list of Visual Studio 2013 versions to use and specified that Xcode needs to be 6.1 or higher to be compatible.

diff --git a/Engine/Documentation/Source/Programming/Development/CompilingProjects/CompilingProjects.INT.udn b/Engine/Documentation/Source/Programming/Development/CompilingProjects/CompilingProjects.INT.udn
--- a/Engine/Documentation/Source/Programming/Development/CompilingProjects/CompilingProjects.INT.udn
+++ b/Engine/Documentation/Source/Programming/Development/CompilingProjects/CompilingProjects.INT.udn
@@ -2,28 +2,28 @@ Availability:Public
 Title:Compiling Game Projects
 Crumbs: %ROOT%, Programming, Programming/Development
 Description:Compiling individual game projects using Visual Studio on Windows, or Xcode on Mac.
-Version: 4.5
+Version: 4.9
 
 [TOC (start:2 end:3)]
 
 [EXCERPT:Main]
 
-Unreal Engine uses a custom building method via the UnrealBuildTool which handles all the complex aspects of compiling 
+Unreal Engine 4 (UE4) uses a custom building method via the UnrealBuildTool which handles all the complex aspects of compiling 
 the project and linking it with the engine. This process occurs transparently allowing you to simply build the project 
 through the standard Visual Studio build workflow.
 
-UnrealBuildTool uses the `*.Build.cs` and `*.Target.cs` files to build the game project. These are automatically 
+UnrealBuildTool uses the .Build.cs and .Target.cs files to build the game project. These are automatically 
 generated when a project is created using a C++ template, or when the [](Programming\Development\ManagingGameCode\CppClassWizard) is used to add code to a 
 project created from a Blueprints Only template.
 
 ## Build Configuration
 
-Unreal Engine 4 and the UnrealBuildTool use different build configurations to determine how the engine is compiled. Which configuration you use will be determined by the purposes of the build you want to create.
+UE4 and the UnrealBuildTool use different build configurations to determine how the engine is compiled. Which configuration you use will be determined by the purposes of the build you want to create.
 
 [PUBLISH:Rocket]
 Each build configuration contains two keywords.  The first of these indicates the state of the engine and your game project. For instance, if you compile using a **Debug** configuration, you will
 be able to debug your game's code.  The second keyword indicates the target you are building for.  If you would like to open a project in Unreal, you need to build with
-the **Editor** target keyword, but if you are building an executable version of your game, you would need to build using the **empty** target keyword.
+the **Editor** target keyword, but if you are building an executable version of your game, you would need to build using the empty target keyword.
 
 |Build Configuration - State | Description |
 |-------------|-------------|
@@ -39,7 +39,7 @@ the **Editor** target keyword, but if you are building an executable version of
 [/PUBLISH:Rocket]
 [PUBLISH:Licensee]
 Each build configuration contains two keywords.  The first of these indicates the state of the engine and your game project. For instance, if you compile using a **DebugGame** configuration, you will
-be able to debug the game code, but not the engine code.  The second keyword indicates the target you are building for.  If you would like to open a project in Unreal Editor, you need to build with
+be able to debug the game code, but not the engine code.  The second keyword indicates the target you are building for.  If you would like to open a project in UE4, you need to build with
 the **Editor** target keyword, but if you are building an executable version of your project that will run on a server, you would need to build using the **Server** target keyword.
 
 |Build Configuration - State | Description |
@@ -77,7 +77,7 @@ The build configuration can be set in the Visual Studio toolbar.
 ### Setting the Solution Platform
 
 [PUBLISH:Rocket]
-Unreal Engine currently supports the following platforms:
+UE4 currently supports the following platforms:
 
 * **Windows 64-bit**
 
@@ -86,7 +86,7 @@ Unreal Engine currently supports the following platforms:
 [/PUBLISH:Rocket]
 
 [PUBLISH:Licensee]
-Unreal Engine 4 currently supports the following platforms:
+UE4 currently supports the following platforms:
 
 * **Windows 32-bit**
 * **Windows 64-bit**
@@ -102,7 +102,7 @@ has instructions for generating project files for additional platforms.
 
 [PUBLISH:Rocket]
 [REGION:note]
-Please make sure you are running Visual Studio 2013 Professional or Visual Studio Express 2013 for Windows Desktop installed before proceeding.  If you are using Mac, make sure to have Xcode installed.
+Please make sure you are running Visual Studio 2013 Express, Community, or Professional for Windows Desktop installed before proceeding.  If you are using Mac, make sure to have Xcode 6.1 or higher installed.
 [/REGION]
 [REGION:tip]
 When building in Visual Studio, you are only compiling the game project, not the Unreal Editor.  
@@ -114,8 +114,6 @@ When building in Visual Studio, you are compiling the engine in addition to the
 [/REGION]
 [/PUBLISH:Licensee]
 
-**To compile the project:**
-
 [PUBLISH:Licensee]
 1.	Set the **Solution Configuration** to the configuration you want to build. In this example, it is set to **Development Editor**. Refer to the [Build Configuration](#BuildConfiguration) section for descriptions of each available configuration.
 
