#UE4doc #docs: ##Overview removed as according to the instructional guide. Altered images to match the instructional guide: green -> yellow.

diff --git a/Engine/Documentation/Source/Gameplay/ClassCreation/BlueprintOnly/BlueprintOnly.INT.udn b/Engine/Documentation/Source/Gameplay/ClassCreation/BlueprintOnly/BlueprintOnly.INT.udn
--- a/Engine/Documentation/Source/Gameplay/ClassCreation/BlueprintOnly/BlueprintOnly.INT.udn
+++ b/Engine/Documentation/Source/Gameplay/ClassCreation/BlueprintOnly/BlueprintOnly.INT.udn
@@ -3,8 +3,6 @@ Title:Blueprints Only
 Crumbs: %ROOT%, Gameplay, Gameplay/ClassCreation
 Description:Introductory information for gameplay programmers getting started with Unreal Engine.
 
-## Overview
-
 A [Blueprints Class](Engine/Blueprints/UserGuide/Types/ClassBlueprint) allows you to set up new classes using the [](Engine/Blueprints) system.  After you [create a new Blueprint Class](Engine/Blueprints/UserGuide/Types/ClassBlueprint/Creation),
 you can add components, set up functions and other gameplay or design behavior with visual scripting, and set default values for class variables. The LightSwitch class set up using only Blueprints is named
 **LightSwitch_BPOnly**, and is explained below.
