#UE4doc #docs Removed 'Overview' heading per Instructional Guide v.1.

diff --git a/Engine/Documentation/Source/Engine/Audio/WAV/WAV.INT.udn b/Engine/Documentation/Source/Engine/Audio/WAV/WAV.INT.udn
--- a/Engine/Documentation/Source/Engine/Audio/WAV/WAV.INT.udn
+++ b/Engine/Documentation/Source/Engine/Audio/WAV/WAV.INT.udn
@@ -48,8 +48,6 @@ Description:The basic requirements for getting sound files into Unreal Engine 4.
 
 [TOC (start:2 end:3)]
 
-## Overview
-
 [OBJECT:topiclistnotitle]
 	[PARAM:icon]
 		![](WAV_TitleIcon.png)(w:200 convert:false) 
