#UE4 Docs: Setting the doc to public so that they will go out in the next doc update.

diff --git a/Engine/Documentation/Source/Platforms/GearVR/GlobalMenu/1/GlobalMenu_HudCreation.INT.udn b/Engine/Documentation/Source/Platforms/GearVR/GlobalMenu/1/GlobalMenu_HudCreation.INT.udn
--- a/Engine/Documentation/Source/Platforms/GearVR/GlobalMenu/1/GlobalMenu_HudCreation.INT.udn
+++ b/Engine/Documentation/Source/Platforms/GearVR/GlobalMenu/1/GlobalMenu_HudCreation.INT.udn
@@ -1,4 +1,4 @@
-Availability:Docs
+Availability:Public
 Title:Samsung Gear VR Global Menu HUD Creation
 Crumbs:%ROOT%, Platforms/GearVR
 Description:Instructions on how to create a new HUD to use with your Gear VR project.
diff --git a/Engine/Documentation/Source/Platforms/GearVR/GlobalMenu/2/GlobalMenu_HUDAssignment.INT.udn b/Engine/Documentation/Source/Platforms/GearVR/GlobalMenu/2/GlobalMenu_HUDAssignment.INT.udn
--- a/Engine/Documentation/Source/Platforms/GearVR/GlobalMenu/2/GlobalMenu_HUDAssignment.INT.udn
+++ b/Engine/Documentation/Source/Platforms/GearVR/GlobalMenu/2/GlobalMenu_HUDAssignment.INT.udn
@@ -1,4 +1,4 @@
-Availability:Docs
+Availability:Public
 Title:Samsung Gear VR Global Menu HUD Assignment
 Crumbs:%ROOT%, Platforms/GearVR
 Description: How to assign the new HUD to your project
diff --git a/Engine/Documentation/Source/Platforms/GearVR/GlobalMenu/3/GlobalMenu_HUDBPSetup.INT.udn b/Engine/Documentation/Source/Platforms/GearVR/GlobalMenu/3/GlobalMenu_HUDBPSetup.INT.udn
--- a/Engine/Documentation/Source/Platforms/GearVR/GlobalMenu/3/GlobalMenu_HUDBPSetup.INT.udn
+++ b/Engine/Documentation/Source/Platforms/GearVR/GlobalMenu/3/GlobalMenu_HUDBPSetup.INT.udn
@@ -1,4 +1,4 @@
-Availability:Docs
+Availability:Public
 Title:Samsung Gear VR Global Menu Blueprint Setup
 Crumbs:%ROOT%, Platforms/GearVR
 Description: How to setup a Blueprint to display the Gear VR Global Menu
diff --git a/Engine/Documentation/Source/Platforms/GearVR/GlobalMenu/GearVRGlobalMenu.INT.udn b/Engine/Documentation/Source/Platforms/GearVR/GlobalMenu/GearVRGlobalMenu.INT.udn
--- a/Engine/Documentation/Source/Platforms/GearVR/GlobalMenu/GearVRGlobalMenu.INT.udn
+++ b/Engine/Documentation/Source/Platforms/GearVR/GlobalMenu/GearVRGlobalMenu.INT.udn
@@ -1,4 +1,4 @@
-Availability:Docs
+Availability:Public
 Title:Samsung Gear VR Global Menu
 Crumbs:%ROOT%, Platforms
 Description:How to setup your Gear VR project to access the global menu.
