#UE4doc #docs: Removed extraneous bolding as according to the instructional guide. ##Overview removed as according to the instructional guide. Altered images to match the instructional guide: green -> yellow.

diff --git a/Engine/Documentation/Source/Gameplay/ClassCreation/CodeAndBlueprints/CodeAndBlueprints.INT.udn b/Engine/Documentation/Source/Gameplay/ClassCreation/CodeAndBlueprints/CodeAndBlueprints.INT.udn
--- a/Engine/Documentation/Source/Gameplay/ClassCreation/CodeAndBlueprints/CodeAndBlueprints.INT.udn
+++ b/Engine/Documentation/Source/Gameplay/ClassCreation/CodeAndBlueprints/CodeAndBlueprints.INT.udn
@@ -4,8 +4,6 @@ Crumbs: %ROOT%, Gameplay, Gameplay/ClassCreation
 Description:Introductory information for gameplay programmers getting started with Unreal Engine.
 Version: 4.6
 
-## Overview
-
 C++ classes can be extended with Blueprints, allowing programmers to set up new gameplay classes in code that can be built upon and changed with Blueprints by level designers.
 There are a number of specifiers that change how a C++ class interacts with the Blueprint system, and some of them will be highlighted in this example.
 
@@ -174,18 +172,18 @@ inherited from the parent class when checked.
 | ![](BP_Only_MyBlueprint.png) | ![](Both_MyBlueprint.png) |
 
 
-There are two graphs used to set up the **LightSwitchBoth_BP** class behavior. The first is the [](Engine/Blueprints\UserGuide\UserConstructionScript) graph, which contains a specialized **Construction Script** 
-event. Without the **Construction Script** setup, a new **LightSwitchBoth_BP** Actor would just use the LightSwitchBoth constructor.  However, the **Construction Script** executes when the 
-Actor is moved within the level, as well as when **DesiredIntensity** is changed. Using the **Construction Script** means that Actor variables exposed to Blueprints can be changed easily and the 
+There are two graphs used to set up the LightSwitchBoth_BP class behavior. The first is the [](Engine/Blueprints\UserGuide\UserConstructionScript) graph, which contains a specialized Construction Script 
+event. Without the Construction Script setup, a new LightSwitchBoth_BP Actor would just use the LightSwitchBoth constructor.  However, the Construction Script executes when the 
+Actor is moved within the level, as well as when DesiredIntensity is changed. Using the Construction Script means that Actor variables exposed to Blueprints can be changed easily and the 
 effects of these changes can be seen quickly.
 
 In the **LightSwitchBoth_BP** class, the **Construction Script** event is connected to a **Set Brightness** node, so that the brightness of **Point Light 1** (the PointLightComponent) is set to 
-the value of **DesiredIntensity** when the Actor is added to or moved within the level or **DesiredIntensity** changes. 
+the value of **DesiredIntensity** when the Actor is added to or moved within the level or DesiredIntensity changes. 
 
 ![](Both_ConstructionScript.png)
 
-The other graph set up in the **LightSwitch_BPOnly** class is the [](Engine/Blueprints\UserGuide\EventGraph). Execution in the EventGraph begins from events. In this case, the **OnOverlap** event
-executes whenever the C++ function `OnOverlap` would be called. In the LightSwitchBoth source file, delegates were setup so **OnOverlap** will execute when an Actor enters or leaves the SphereComponent:
+The other graph set up in the LightSwitch_BPOnly class is the [](Engine/Blueprints\UserGuide\EventGraph). Execution in the EventGraph begins from events. In this case, the OnOverlap event
+executes whenever the C++ function `OnOverlap` would be called. In the LightSwitchBoth source file, delegates were setup so OnOverlap will execute when an Actor enters or leaves the SphereComponent:
 
 		Sphere1->OnComponentBeginOverlap.AddDynamic(this, &ALightSwitchBoth::OnOverlapBegin);		// set up a notification for when this component overlaps something
 		Sphere1->OnComponentEndOverlap.AddDynamic(this, &ALightSwitchBoth::OnOverlapEnd);		// set up a notification for when this component overlaps something
@@ -198,13 +196,13 @@ For more on events and working with graphs, see the [](Engine/Blueprints\UserGui
 ![](Both_EventGraph_2.png)
 
 The **DesiredIntensity** variable is set to EditAnywhere in the LightSwitchBoth header file, so it is visible in the **Blueprint Editor**'s Defaults and can be edited by clicking the ![](Engine\Blueprints\UserGuide\Variables\classDefaults_button.png) to show the Class Defaults in the Details Panel.
-This also means that the variable can be changed for each instance of the class, so each Actor can have its own **DesiredIntensity**. Because **DesiredIntensity** is also **BlueprintReadWrite** and used in the 
-**Construction Script**, updating it also causes the **Construction Script** to execute again.
+This also means that the variable can be changed for each instance of the class, so each Actor can have its own DesiredIntensity. Because DesiredIntensity is also BlueprintReadWrite and used in the 
+Construction Script, updating it also causes the Construction Script to execute again.
 
 ![](Both_Defaults.png)
 
 Blueprinted classes can be extended with other Blueprint Classes, either using the dropdown button next to the class in the **Class Viewer** to create a new Blueprint, or by
 **right-clicking** on the Blueprint and selecting **Create New Blueprint Based on This**.
 
-The Blueprint Class **LightSwitchBoth_BP** is in the [](Engine/Content/Browser), and can be dragged into the level from there. It is also in the [](Engine/UI/ClassViewer).
-For more information on placing Actors in a level using the **Content Browser** or **Class Viewer**, see the [](Engine/Actors/Placement) documentation.
+The Blueprint Class LightSwitchBoth_BP is in the [](Engine/Content/Browser), and can be dragged into the level from there. It is also in the [](Engine/UI/ClassViewer).
+For more information on placing Actors in a level using the Content Browser or Class Viewer, see the [](Engine/Actors/Placement) documentation.
