#UE4doc #docs:
Updated Engine version
Changed Visual Studio and Xcode versions needed to reflect latest requirements.
Changed some terms due to inaccuracy.
Removed step about UE4 asking the user if they wish to edit the C++ class they just added as it doesn't occur any longer.

diff --git a/Engine/Documentation/Source/Programming/Development/ManagingGameCode/CppClassWizard/CppClassWizard.INT.udn b/Engine/Documentation/Source/Programming/Development/ManagingGameCode/CppClassWizard/CppClassWizard.INT.udn
--- a/Engine/Documentation/Source/Programming/Development/ManagingGameCode/CppClassWizard/CppClassWizard.INT.udn
+++ b/Engine/Documentation/Source/Programming/Development/ManagingGameCode/CppClassWizard/CppClassWizard.INT.udn
@@ -2,7 +2,7 @@ Availability:Public
 Title:C++ Class Wizard
 Crumbs: %ROOT%, Programming, Programming\Development
 Description:An overview of the C++ Class Wizard in UE4.
-Version: 4.5
+Version: 4.9
 
 %Globals:OSSelect%
 
@@ -13,7 +13,7 @@ The _C++ Class Wizard_ provides a fast and easy way to add native C++ code class
 with your own functionality, if you wish. This converts a content-only project into a code project. You can access the C++ Class Wizard like so:
 
 [REGION:note]
-On Windows, please make sure you have Visual Studio 2013 Professional or Express for Windows Desktop installed before proceeding.  On Mac, make sure to have Xcode installed.
+On Windows, please make sure you have Visual Studio 2013 Express, Community, or Professional for Windows Desktop installed before proceeding.  On Mac, make sure to have Xcode 6.1 or higher installed.
 [/REGION:note]
 
 
@@ -40,11 +40,7 @@ On Windows, please make sure you have Visual Studio 2013 Professional or Express
 		Class names may only contain alphanumeric characters and cannot contain spaces. The field will notify you if you enter an invalid name.
 	[/REGION]
 
-1. After a moment of calculation, the Editor will ask if you would like to edit the code now. Click ![](button_Yes.png).
-
-	![](EditNow.png)
-
-	[OBJECT:ToggleButtonContent]
+1. [OBJECT:ToggleButtonContent]
 		[PARAMLITERAL:category]
 		OS
 		[/PARAMLITERAL]
