#UE4doc #docs:
Changes didn't commit for the previous submit

diff --git a/Engine/Documentation/Source/Programming/Development/ManagingGameCode/ManagingGameCode.INT.udn b/Engine/Documentation/Source/Programming/Development/ManagingGameCode/ManagingGameCode.INT.udn
--- a/Engine/Documentation/Source/Programming/Development/ManagingGameCode/ManagingGameCode.INT.udn
+++ b/Engine/Documentation/Source/Programming/Development/ManagingGameCode/ManagingGameCode.INT.udn
@@ -2,7 +2,7 @@ Availability:Public
 Title:Managing Game Code
 Crumbs: %ROOT%, Programming, Programming\Development
 Description:Adding code and generating project files. 
-Version: 4.5
+Version: 4.9
 
 %Globals:OSSelect%
 
@@ -64,7 +64,7 @@ Version: 4.5
 	active_button_content
 	[/PARAMLITERAL]
 	[PARAM:content]
-	If your project is already open in the editor, you can easily open it in Visual Studio by selecting **Open in Visual Studio** from the **File** menu.
+	If your project is already open in the editor, you can easily open it in Visual Studio by selecting **Open Visual Studio** from the **File** menu.
 
 	![](open_visual_studio.png)
 
