#UE4doc #docs: Removed extraneous bolding as according to the instructional guide. ##Overview removed as according to the instructional guide.

diff --git a/Engine/Documentation/Source/Gameplay/ClassCreation/CodeOnly/CodeOnly.INT.udn b/Engine/Documentation/Source/Gameplay/ClassCreation/CodeOnly/CodeOnly.INT.udn
--- a/Engine/Documentation/Source/Gameplay/ClassCreation/CodeOnly/CodeOnly.INT.udn
+++ b/Engine/Documentation/Source/Gameplay/ClassCreation/CodeOnly/CodeOnly.INT.udn
@@ -4,13 +4,11 @@ Crumbs: %ROOT%, Gameplay, Gameplay/ClassCreation
 Description:Introductory information for gameplay programmers getting started with Unreal Engine.
 Version: 4.6
 
-## Overview
-
 Using the [](Programming\Development\ManagingGameCode\CppClassWizard), you can easily add new C++ classes to your project. After you choose the class you would like your new class to derive from, the wizard
 sets up the necessary header file and source file. If this is the first code you have added to the project, the project is converted to a code project, which creates the game module that contains your 
 source code.  It also lets Unreal Editor know that the game module exists, so it can load your C++ changes from Visual Studio or Xcode, and minor code changes can be compiled within Unreal Editor.
 
-The LightSwitch class set up using only C++ is named **LightSwitchCodeOnly**, and is explained below.
+The LightSwitch class set up using only C++ is named LightSwitchCodeOnly, and is explained below.
 
 ## Class Setup
 
@@ -214,5 +212,5 @@ If you are adding code to an existing C++ project, you can use the Hot Reload fu
 C++ classes can be extended with C++ classes as well as with Blueprint Classes - they show up after checking the **Show All Classes** checkbox in the **C++ Class Wizard** and in the **Custom Classes** section of 
 the **Pick Parent Class** window during Blueprint Class creation.
 
-The LightSwitchCodeOnly class is in the [](Engine/UI/ClassViewer), and can be dragged into the level from there. For more information on placing Actors in a level using the **Class Viewer**, see the [](Engine/Actors/Placement) documentation.
+The LightSwitchCodeOnly class is in the [](Engine/UI/ClassViewer), and can be dragged into the level from there. For more information on placing Actors in a level using the Class Viewer, see the [](Engine/Actors/Placement) documentation.
  
