#UE4doc #docs: Removed extraneous bolding as according to the instructional guide. ##Overview removed as according to the instructional guide.

diff --git a/Engine/Documentation/Source/Gameplay/Framework/GameplayFramework.INT.udn b/Engine/Documentation/Source/Gameplay/Framework/GameplayFramework.INT.udn
--- a/Engine/Documentation/Source/Gameplay/Framework/GameplayFramework.INT.udn
+++ b/Engine/Documentation/Source/Gameplay/Framework/GameplayFramework.INT.udn
@@ -3,22 +3,19 @@ Title:Gameplay Framework
 Crumbs:%ROOT%, Gameplay
 Description:Core systems, such as game rules, player input and controls, cameras, and user interfaces.
 
-
-## Overview
-
 To show how the different pieces of game framework interlock, here is a simple example. Imagine a simple concept for a game, 
-where a rabbit races a snail.  The base of the game framework is the **GameMode**.  The **GameMode** sets the rules for the game,
+where a rabbit races a snail.  The base of the game framework is the GameMode.  The **GameMode** sets the rules for the game,
 like the rule that whichever player crosses the finish line first is the winner.  It also handles spawning the players. 
  
-A player is set up in a **PlayerController**, which can possess a **Pawn**.  The **Pawn** is the physical representation of a player in  the game, while the **Controller** possesses the **Pawn** and can set rules for its behavior. 
-In our example, there would be two **Pawns**, one for the snail and one for the rabbit.  The rabbit would actually be a **Character**, a special subclass of **Pawn** which has built-in movement functionality including running and jumping.  The snail, on the other hand, has a different style of movement, so it could extend directly from the **Pawn** class.
+A player is set up in a **PlayerController**, which can possess a Pawn.  The **Pawn** is the physical representation of a player in  the game, while the Controller possesses the Pawn and can set rules for its behavior. 
+In our example, there would be two Pawns, one for the snail and one for the rabbit.  The rabbit would actually be a **Character**, a special subclass of Pawn which has built-in movement functionality including running and jumping.  The snail, on the other hand, has a different style of movement, so it could extend directly from the Pawn class.
 
-A **Pawn** can contain its own rules for movement and other game logic, but that functionality can also be included in a **Controller**.  A **Controller** can either be a **PlayerController** taking input from a human player or an **AIController** 
-with automated control by the computer. In this example, the player would be controlling the snail, so the snail **Pawn** would be possessed by a **PlayerController**.  The rabbit would be controlled by AI, with rules
-for when it should stop, sprint, or nap all being set up in an **AIController** which possesses the rabbit **Character**.  Only human players care about the view provided by a **Camera**, so only one CameraComponent in the snail
-**Pawn** would be used by the PlayerCamera.
+A Pawn can contain its own rules for movement and other game logic, but that functionality can also be included in a Controller.  A Controller can either be a PlayerController taking input from a human player or an AIController 
+with automated control by the computer. In this example, the player would be controlling the snail, so the snail Pawn would be possessed by a PlayerController.  The rabbit would be controlled by AI, with rules
+for when it should stop, sprint, or nap all being set up in an AIController which possesses the rabbit Character.  Only human players care about the view provided by a Camera, so only one CameraComponent in the snail
+Pawn would be used by the PlayerCamera.
 
-During gameplay, input from the player would move the snail around the map, while the **HUD** was overlaid on the view provided by the **Camera**, showing information about who was currently in first place, and the 
+During gameplay, input from the player would move the snail around the map, while the HUD was overlaid on the view provided by the Camera, showing information about who was currently in first place, and the 
 race time that had elapsed.  
 
 For quick descriptions of the basic framework elements, and a chart of how they fit together, check out the [](Gameplay/Framework/QuickReference). Otherwise, you can use the links below to jump directly 
