#UE4doc #docs: Removed extraneous bolding as according to the instructional guide. Changed minor grammar error with commas.

diff --git a/Engine/Documentation/Source/Gameplay/Framework/Camera/Camera.INT.udn b/Engine/Documentation/Source/Gameplay/Framework/Camera/Camera.INT.udn
--- a/Engine/Documentation/Source/Gameplay/Framework/Camera/Camera.INT.udn
+++ b/Engine/Documentation/Source/Gameplay/Framework/Camera/Camera.INT.udn
@@ -5,11 +5,11 @@ Description:The Camera represents the player's point of view; how the player see
 
 The **Camera** represents the player's point of view; how the player sees the world. For this reason, 
 cameras only have relevance to human-controlled players. The **PlayerController** specifies a camera class 
-and instantiates a **Camera** Actor which is used to calculate the position and orientation the player views 
+and instantiates a Camera Actor which is used to calculate the position and orientation the player views 
 the world from.
 
 [REGION:note]
-For basic examples of how to work with **Cameras**, refer to the [Using Cameras](Gameplay\HowTo\UsingCameras) documentation. You can also layer animations onto cameras using the [CameraAnim feature](Gameplay/Framework/Camera/Animations).
+For basic examples of how to work with Cameras, refer to the [Using Cameras](Gameplay\HowTo\UsingCameras) documentation. You can also layer animations onto cameras using the [CameraAnim feature](Gameplay/Framework/Camera/Animations).
 [/REGION]
 
 ## CameraComponent and CameraActor
@@ -18,7 +18,7 @@ All of the camera's properties and behavior are set up in the CameraComponent; t
 level rather than within another class.  
 
 In the CameraComponent, it is possible to set whether the camera is in perspective mode or orthographic mode.  The vertical field of view (FOV) can be set for perspective mode, 
-and the width in world units can be set for orthographic mode.  For both modes, the aspect ratio can be designated, and preset aspect ratios for common devices and display types are 
+and the width in world units can be set for orthographic mode.  For both modes, the aspect ratio can be designated and preset aspect ratios for common devices and display types are 
 provided.  [Post process effects](Engine/Rendering\PostProcessEffects) can be added to the camera, and it is also possible to scale the strength of the post process effects.
 
 Two components are added to the CameraComponent to aid in visual placement in the Editor, although they will not be visible during gameplay.  A FrustumComponent shows where the 
