#UE4doc #docs: unbolded hyperlinks according to the instructional guide.

diff --git a/Engine/Documentation/Source/Gameplay/Framework/Pawn/Pawn.INT.udn b/Engine/Documentation/Source/Gameplay/Framework/Pawn/Pawn.INT.udn
--- a/Engine/Documentation/Source/Gameplay/Framework/Pawn/Pawn.INT.udn
+++ b/Engine/Documentation/Source/Gameplay/Framework/Pawn/Pawn.INT.udn
@@ -13,9 +13,9 @@ Related: Gameplay/Framework/GameMode
 The **Pawn** class is the base class of all Actors that can be controlled by players or AI. A Pawn is the physical representation of a player or AI entity within the world. 
 This not only means that the Pawn determines what the player or AI entity looks like visually, but also how it interacts with the world in terms of collisions and other physical 
 interactions. This can be confusing in certain circumstances as some types of games may not have a visible player mesh or avatar within the game. Regardless, the Pawn still 
-represents the physical location, rotation, etc. of a player or entity within the game. A [**Character**](Gameplay/Framework\Pawn\Character) is a special type of Pawn that has the ability to walk around.  
+represents the physical location, rotation, etc. of a player or entity within the game. A [Character](Gameplay/Framework\Pawn\Character) is a special type of Pawn that has the ability to walk around.  
 
-By default, there is a one-to-one relationship between [**Controllers**](Gameplay/Framework/Controller) and Pawns; meaning, each Controller controls only one Pawn at any given time. Also, Pawns spawned during gameplay 
+By default, there is a one-to-one relationship between [Controllers](Gameplay/Framework/Controller) and Pawns; meaning, each Controller controls only one Pawn at any given time. Also, Pawns spawned during gameplay 
 are not automatically possessed by a Controller.
 
 In **Blueprints**, the best way to add movement to your Pawn-derived class is with **SetActorLocation**. With **SetActorLocation**, you can decide if you want to teleport to the location, or sweep. 
@@ -38,7 +38,7 @@ three variables are also accessible in DefaultPawn-derived Blueprints.
 
 ### Spectator Pawn
 
-The _SpectatorPawn_ class is a subclass of DefaultPawn.  Through a [**GameMode**](Gameplay/Framework/GameMode), different classes can be specified as defaults for Pawn and SpectatorPawn, and this class provides a simple framework ideal for
+The _SpectatorPawn_ class is a subclass of DefaultPawn.  Through a [GameMode](Gameplay/Framework/GameMode), different classes can be specified as defaults for Pawn and SpectatorPawn, and this class provides a simple framework ideal for
 spectating functionality.  Like DefaultPawn, it has a spherical _CollisionComponent_, although the _StaticMeshComponent_ is not created due to the implementation of `.DoNotCreateDefaultSubobject(Super::MeshComponentName)`.
 The movement for the SpectatorPawn class is handled in the _SpectatorPawnMovementComponent_; the no-gravity flying behavior is the same as in the _DefaultPawnMovementComponent_, with added code for 
 handling or ignoring time dilation as necessary.
