#UE4doc #docs:
Updated Engine version
Changed font settings for multiple terms
Commented out a line for clarification/revision.

diff --git a/Engine/Documentation/Source/Programming/Development/Tools/ConsoleManager/ConsoleManager.INT.udn b/Engine/Documentation/Source/Programming/Development/Tools/ConsoleManager/ConsoleManager.INT.udn
--- a/Engine/Documentation/Source/Programming/Development/Tools/ConsoleManager/ConsoleManager.INT.udn
+++ b/Engine/Documentation/Source/Programming/Development/Tools/ConsoleManager/ConsoleManager.INT.udn
@@ -2,17 +2,16 @@ Availability: Public
 Title:Console Manager: Console Variables in C++
 Crumbs:%ROOT%, Engine, Programming, Programming/Basics
 Description:Overview of the Console Manager and implementation details for creating console variables.
-Version: 4.6
+Version: 4.9
 
 [TOC(start:2)]
 
 
 
-## Overview
 
 A **console command** is simply a user input string that is sent to the engine and the engine can react is some way (e.g. console/log response, changing internal state).
 A **console variable** additionally stores some state that can be changed through the console. By registering console commands and variables in the console manager you get
-auto completion and enumeration to get a list of all console objects (console command **Help** or **DumpConsoleVariables**). Because of this, you should avoid the old Exec
+auto completion and enumeration to get a list of all console objects (console command Help or DumpConsoleVariables). Because of this, you should avoid the old Exec
 interface. The console manager in the central point to control all those things and more (e.g. user input history).
 
 
@@ -43,8 +42,8 @@ The variable needs to be registered early when the engine is created. This examp
 		TEXT("  3: high quality (e.g. color fringe, not yet implemented)"),
 		ECVF_Scalability | ECVF_RenderThreadSafe);
 
-Here we register a console variables of the type int32, with the name **r.RefractionQuality**, the default value of 2 and some multi line help text and some flags.
-There are many flags, the most important one is **ECVF_Cheat**. That is explained in detail in `IConsoleManager.h`.
+Here we register a console variables of the type int32, with the name r.RefractionQuality, the default value of 2 and some multi line help text and some flags.
+There are many flags, the most important one is ECVF_Cheat. That is explained in detail in IConsoleManager.h.
 The help text is shown when the user uses "?" after the console variable.
 
 If needed you also can generate a console variable inside a function:
@@ -80,21 +79,21 @@ Here they type is deducted from the variable type.
 
 ## Getting the state of a console variable
 
-Getting the state of console variables created with `RegisterConsoleVariableRef` can be done efficiently by using the variable that it was registered with. e.g.
+Getting the state of console variables created with **RegisterConsoleVariableRef** can be done efficiently by using the variable that it was registered with. e.g.
 
 	// only needed if you are not in the same cpp file
 	extern TAutoConsoleVariable<int32> CVarRefractionQuality;
 	// get the value on the game thread
 	int32 MyVar = CVarRefractionQuality.GetValueOnGameThread();
 
-Using Getter functions (ie. !GetInt(), !GetFloat(), !GetString()) to determine a console variable�s state results in a slightly slower implementation (virtual function call, possibly cache miss, etc.).
+Using Getter functions (ie. !GetInt(), !GetFloat(), !GetString()) to determine a console variables state results in a slightly slower implementation (virtual function call, possibly cache miss, etc.).
 For best performance you should use same type the variable was registered with. In order to get the pointer to the variable, you can either store the return
-argument of the register function or call `FindConsoleVariable` just before you need the variable. Example:
+argument of the register function or call **FindConsoleVariable** just before you need the variable. Example:
 
     static const auto CVar = IConsoleManager::Get().FindConsoleVariable(TEXT("TonemapperType")); 
 	int32 Value = CVar->GetInt();
    
-The **static** there ensures the name search (implemented as map) is only done the first time this code is called.
+The static there ensures the name search (implemented as map) is only done the first time this code is called.
 This is correct as the variable will never move and only gets destructed on engine shutdown.
 
 
@@ -102,7 +101,7 @@ This is correct as the variable will never move and only gets destructed on engi
 
 If you want to execute some custom code if a console variable changes, you have 3 methods you can choose from. 
 
-Often the **simplest method** is the best: You can store the old state in your subsystem and check each frame if they differ. Here you control when this happens very freely e.g. render thread or game thread, streaming thread, before/after tick or rendering. When you detect the difference, you copy
+Often the simplest method is the best: You can store the old state in your subsystem and check each frame if they differ. Here you control when this happens very freely e.g. render thread or game thread, streaming thread, before/after tick or rendering. When you detect the difference, you copy
 the console variable state and do your custom code. e.g.
 
 	void MyFunc()
@@ -117,7 +116,7 @@ the console variable state and do your custom code. e.g.
 		}
 	}
 
-You also can register a **console variable sink** e.g. 
+You also can register a console variable sink e.g. 
 
 	static void MySinkFunction()
 	{
@@ -139,10 +138,10 @@ You also can register a **console variable sink** e.g.
 The sink is called at a specific point on the main thread before rendering. The function does not get the console variable name/pointer as this often would lead to the wrong behavior.
 If multiple console variables (e.g. r.SceneColorFormat, r.GBufferFormat) should all trigger the change, it is best to call the code after all have been changed, not one after another.
 
-The last method, using the **callback**, you should avoid as it can cause problems if not used carefully:
+The last method, using the callback, you should avoid as it can cause problems if not used carefully:
 
 * A cycle can cause a deadlock (We could prevent the deadlock but which callback to favour is not clear).
-* The callback can come back at any point in time whenever !Set() is getting called. You code has to work in all cases (during init, during serialization).
+* The callback can come back at any point in time whenever **!Set()** is getting called. You code has to work in all cases (during init, during serialization).
   You can assume it is always on the main thread side.
 
 We suggest to not use this method unless you cannot solve it with the other methods mentioned. 
@@ -178,7 +177,7 @@ Please read `IConsoleManager.h` for find more details on this.
 
 ## Loading console variables
 
-On engine startup, the state of console variables can be loaded from the file "Engine/Config/ConsoleVariables.ini". This place is reserved for the local developer - it should not be used for project settings.
+On engine startup, the state of console variables can be loaded from the file **Engine/Config/ConsoleVariables.ini**. This place is reserved for the local developer - it should not be used for project settings.
 More details can be found in the file itself:
 
     ; This file allows you to set console variables on engine startup (In undefined order).
@@ -200,7 +199,7 @@ More details can be found in the file itself:
     
     [Startup]
     
-You also can put the settings "Engine/Config/BasEngine.ini" e.g.
+You also can put the settings **Engine/Config/BasEngine.ini** e.g.
 
 	[SystemSettings]
 	r.MyCvar = 2
@@ -208,7 +207,7 @@ You also can put the settings "Engine/Config/BasEngine.ini" e.g.
 	[SystemSettingsEditor]
 	r.MyCvar = 3
 
-Setting also can come from "Script/Engine.RendererSettings". These project settings are exposed like this:
+Setting also can come from **Script/Engine.RendererSettings**. These project settings are exposed like this:
 
 	UPROPERTY(config, EditAnywhere, Category=Optimizations, meta=(
 		ConsoleVariable="r.EarlyZPassMovable",DisplayName="Movables in early Z-pass",
@@ -217,7 +216,7 @@ Setting also can come from "Script/Engine.RendererSettings". These project setti
 
 Those settings can be changed in the editor UI. Project settings should not intermix with scalability settings (to prevent priority issues).
 
-Other settings can come from the Scalability feature. Look at "Config/BaseScalability.ini" or the Scalability documentation for more info.
+Other settings can come from the Scalability feature. Look at **Config/BaseScalability.ini** or the Scalability documentation for more info.
 
 ## Priority
 
@@ -264,17 +263,19 @@ The priority is also helpful to see who set the variable the last time. You can
 
 ## Future
 
-* Currently console variable can only be created in C++ but that might change.
+* Currently console variables can only be created in C++ but that might change.
 * We thought about adding an enum and bool type but there are many problems attached to it. For now we suggest to use int, or if needed, strings.
 * The help text is convenient but to save executable size or to make it harder for cheaters, we consider adding a define to prevent the help text to go into the executable.
 
 
 ## Unregistering console variables
 
-The `UnregisterConsoleVariable` method allows you to remove the console variable. At least this is what is happening from the user's perspective.
+The **UnregisterConsoleVariable** method allows you to remove the console variable. At least this is what is happening from the user's perspective.
 The variable is still kept (with the unregistered flags) to not crash when pointers access the data. If a new variable is registered with the same name,
-the old variable is restored and help and flags get copied from the new variable. This way DLL loading and unloading can work even without losing the variable state.
-Note this will not work for console variable references. This can be fixed by giving up on one: do not store pointers, do not unregistered or not use references.
+the old variable is restored and flags get copied from the new variable. This way DLL loading and unloading can work even without losing the variable state.
+Note that this will not work for console variable references. 
+<!--Sentence below deprecated for clarification.-->
+<!--This can be fixed by giving up on one: do not store pointers, do not unregistered or not use references.-->
 
 
 
