#UE4doc #docs: Removed extraneous bolding as according to the instructional guide. Changed bolded word to header: **HUD** -> ## HUD.

diff --git a/Engine/Documentation/Source/Gameplay/Framework/UIAndHUD/UIAndHUD.INT.udn b/Engine/Documentation/Source/Gameplay/Framework/UIAndHUD/UIAndHUD.INT.udn
--- a/Engine/Documentation/Source/Gameplay/Framework/UIAndHUD/UIAndHUD.INT.udn
+++ b/Engine/Documentation/Source/Gameplay/Framework/UIAndHUD/UIAndHUD.INT.udn
@@ -32,19 +32,19 @@ Unreal Engine 4 provides multiple means of creating UIs and HUDs. The Canvas cla
 The game's Interface is used to convey information to the player and provide a means of prompting the user for directed input. 
 A game interface generally consists of two main elements: the heads-up display (HUD) and menus or user interfaces (UIs).
 
-The **HUD** refers to the status and information overlaid on the screen during gameplay. The purpose of the HUD 
+The HUD refers to the status and information overlaid on the screen during gameplay. The purpose of the HUD 
 is to inform the player of the current state of the game, i.e. the score, their health, the amount of time remaining, 
 etc. The HUD is usually non-interactive, meaning the player does not click on elements of the HUD, though this 
 becomes a gray area in certain types of games where the HUD and user interfaces are hard to separate.
 
 **User Interfaces** refer to menus and other interactive elements. These elements are usually drawn overlaid on 
-the screen much like the **HUD**, but in certain circumstances they could be part of the game world itself rendered 
+the screen much like the HUD, but in certain circumstances they could be part of the game world itself rendered 
 onto a surface in the world. The most obvious examples of UIs are the main menu displayed when the game starts up 
 or the pause menu shown when the player has paused the game. However, other UIs may be displayed during play. These 
 could be used to show dialog between characters in the game or in more complex situations, such as in an RTS or RPG, 
 they may be integral to the game play itself allowing the player to choose weapons, armor, units to build, etc.
 
-**HUD**
+## HUD
 
 The **HUD** is the base object for displaying elements overlaid on the screen. Every human-controlled player in the 
 game has their own instance of the `AHUD` class which draws to their individual Viewport. In the case of splitscreen 
