#UE4doc #docs: Changed title of How To to match instructional guide.

diff --git a/Engine/Documentation/Source/Gameplay/HowTo/ActorInput/ActorInputOverview.INT.udn b/Engine/Documentation/Source/Gameplay/HowTo/ActorInput/ActorInputOverview.INT.udn
--- a/Engine/Documentation/Source/Gameplay/HowTo/ActorInput/ActorInputOverview.INT.udn
+++ b/Engine/Documentation/Source/Gameplay/HowTo/ActorInput/ActorInputOverview.INT.udn
@@ -1,11 +1,8 @@
 Availability: Public
-Title:How to Set Up Input on an Actor
+Title:Setting Up Input on an Actor
 Crumbs: %ROOT%, Gameplay
 Description:A How To Guide to Setting up Input for an Actor in Unreal Engine 4. 
 
-
-##Overview
-
 When developing your game, there may be instances when you want to allow the player to perform some form of **Input** on an Actor in your level. For example, perhaps you have a treasure chest that you want the player to be able to open when they approach it, or a light you want them to be able to turn on/off, or some other form of interactable object you want to affect in some way when the player presses a button. 
 
 Setting up Input for an Actor gives you a level of control over how and when the Actor responds to player input. For more information, see the link provided in the **Setup** section below. 
