#UE4doc #docs: Removed extraneous bolding as according to the instructional guide. Altered images to match the instructional guide: green -> yellow. ##Overview removed as according to the instructional guide.

diff --git a/Engine/Documentation/Source/Gameplay/HowTo/ActorInput/Blueprints/ActorInput_BP.INT.udn b/Engine/Documentation/Source/Gameplay/HowTo/ActorInput/Blueprints/ActorInput_BP.INT.udn
--- a/Engine/Documentation/Source/Gameplay/HowTo/ActorInput/Blueprints/ActorInput_BP.INT.udn
+++ b/Engine/Documentation/Source/Gameplay/HowTo/ActorInput/Blueprints/ActorInput_BP.INT.udn
@@ -1,15 +1,13 @@
 Availability: Public
-Title:How to Set Up Input on an Actor in Blueprints
+Title:Setting Up Input on an Actor in Blueprints
 Crumbs: %ROOT%, Gameplay, Gameplay\HowTo\ActorInput
 Description:A How To Guide for setting up input on an Actor in Blueprints. 
 
 [TOC (start:2 end:2)]
 
-## Overview
-
 This page covers the basic methods on how to **Enable Input** or **Disable Input** for an Actor. By Enabling Input on an Actor, you can allow a player to press a button or key and execute events that affect the Actor in some way (be it turn a light on or off, open or close a door, activate something, etc.). 
 
-The **Enable/Disable Input Usage** section also covers an approach to Enabling/Disabling Input using a Trigger Volume to determine if a player is near the Actor in the level. This is useful when you want a player to affect the Actor only when they are near or in range of the Actor to conceivably affect it in a real-world scenario. 
+The Enable/Disable Input Usage section also covers an approach to Enabling/Disabling Input using a Trigger Volume to determine if a player is near the Actor in the level. This is useful when you want a player to affect the Actor only when they are near or in range of the Actor to conceivably affect it in a real-world scenario. 
  
 
 ## Enabling Actor Input
