#UE4doc #docs: ##Overview removed as according to the instructional guide. Changed title of How To to match instructional guide.

diff --git a/Engine/Documentation/Source/Gameplay/HowTo/AddingComponents/ComponentsOverview.INT.udn b/Engine/Documentation/Source/Gameplay/HowTo/AddingComponents/ComponentsOverview.INT.udn
--- a/Engine/Documentation/Source/Gameplay/HowTo/AddingComponents/ComponentsOverview.INT.udn
+++ b/Engine/Documentation/Source/Gameplay/HowTo/AddingComponents/ComponentsOverview.INT.udn
@@ -1,16 +1,13 @@
 Availability: Public
-Title:How to Add Components to an Actor
+Title:Adding Components to an Actor
 Crumbs: %ROOT%, Gameplay
 Description:A How To Guide for Adding Components to Actors in Unreal Engine 4. 
 
-
-##Overview
-
 **Components** are a special type of **Object** designed to be used as sub-objects within **Actors**. These are generally used in situations where it is necessary to have easily swappable parts in order to change the behavior or functionality of some particular aspect of the owning Actor. For instance, a car is controlled and moves very differently from an aircraft, which is controlled and moves differently from a boat, etc.; yet all of these are vehicles that share other commonalities. By using a Component to handle the controls and movement, the same vehicle can easily be made to behave like any one of these specific types. 
 
 ## Setup
 
-To learn more about **Adding Components to an Actor** using Blueprints, follow the link below:
+To learn more about Adding Components to an Actor using Blueprints, follow the link below:
 
 [REGION:buttonlist]
 * [](Gameplay\HowTo\AddingComponents\Blueprints)
