#UE4doc #docs:
Updated Engine version
Removed Overview header
Bolded a few terms
Added references to Visual Studio Community edition

diff --git a/Engine/Documentation/Source/Programming/Development/VisualStudioSetup/VisualStudioSetup.INT.udn b/Engine/Documentation/Source/Programming/Development/VisualStudioSetup/VisualStudioSetup.INT.udn
--- a/Engine/Documentation/Source/Programming/Development/VisualStudioSetup/VisualStudioSetup.INT.udn
+++ b/Engine/Documentation/Source/Programming/Development/VisualStudioSetup/VisualStudioSetup.INT.udn
@@ -2,17 +2,16 @@ Availability:Public
 Title:Setting Up Visual Studio for UE4
 Crumbs: %ROOT%, Programming, Programming\Development
 Description:Tips, tricks, and techniques for setting up Visual Studio for development with Unreal Engine 4.
-Version: 4.5
+Version: 4.9
 
 [TOC(start:2)]
 
 
-## Overview
 
 [EXCERPT:Intro]
 
-Unreal Engine 4 is designed to integrate smoothly with Visual Studio 2013. This allows you to quickly and easily make code changes in your projects and see the results immediately upon compilation. 
-Setting up Visual Studio to work with Unreal Engine 4 can help to improve efficiency and the overall user experience for developers using Unreal Engine 4.
+Unreal Engine 4 (UE4) is designed to integrate smoothly with Visual Studio 2013 (VS2013). This allows you to quickly and easily make code changes in your projects and see the results immediately upon compilation. 
+Setting up Visual Studio to work with UE4 can help to improve efficiency and the overall user experience for developers using UE4.
 [/EXCERPT:Intro]
 
 In this document, we will cover the basics for setting up your UE4-to-VS2013 workflow.
@@ -37,7 +36,7 @@ The following are recommended settings for developers using Visual Studio with U
 	
 		![](menu_dropdown.png)
 	1. Click the dropdown button, hover over **Add or Remove Buttons**, and then click on **Solution Platforms** to add the menu to the toolbar.
-1. **Turn off the Error List window**. If enabled, the **Error List** window pops up automatically when you get a legitimate error in your code. However, when working with Unreal Engine, the **Error List** can display false error information. It is best to disable this window and use the **Output** window to see your code errors instead. This window will only show your real errors. 
+1. Turn off the **Error List** window. If enabled, the **Error List** window pops up automatically when you get a legitimate error in your code. However, when working with Unreal Engine, the **Error List** can display false error information. It is best to disable this window and use the **Output** window to see your code errors instead. This window will only show your real errors. 
 	1. Close the **Error List** window if it is open. 
 	1. From the **Tools** menu, open the **Options** dialog.
 	1. Select **Projects and Solutions** and uncheck **Always show error list if build finishes with error**.
@@ -60,7 +59,7 @@ Other configuration steps you may find useful are:
 (#Intellisense)
 ## Intellisense, Live Errors, and Squiggles
 
-UE4 projects now have proper IntelliSense support, including a live Error List and "squiggles"! 
+UE4 projects now have proper IntelliSense support, including a live Error List and squiggles.
 (See below for how to enable it.)
 
 IntelliSense recompiles C++ as you type. This is a lot more powerful than only VAX's syntax 
