#UE4doc #docs Updated Importing Audio protocol to reflect UE 4.8.

diff --git a/Engine/Documentation/Source/Engine/Audio/WAV/WAV.INT.udn b/Engine/Documentation/Source/Engine/Audio/WAV/WAV.INT.udn
--- a/Engine/Documentation/Source/Engine/Audio/WAV/WAV.INT.udn
+++ b/Engine/Documentation/Source/Engine/Audio/WAV/WAV.INT.udn
@@ -92,28 +92,28 @@ Compression settings are tweaked to maintain similar quality across platforms. A
 
 In the **Content Browser**:
 
-1. Click the **New Asset** ![](CB_button_NewAsset.png) button.
-1. Select _Import_ from the menu.
-1. Select the .WAV file you wish to import.
+1.	 From the Content Browser, click the **Import** button (pictured below).
+	![](ImportButton_UI.png)
+1.	 Select the .WAV file you wish to import.
 
 Or
 
-1. Select your sound file in a Windows File Explorer window.
-1. Click and drag it into the **Content Browser** in Unreal Editor 4.
+1.	 Select your sound file in a Windows File Explorer window.
+1.	 Click and drag it into the **Content Browser** in Unreal Editor 4.
 
 ### Multichannel Import
 
 In the **Content Browser**:
 
-1. Click the **New Asset** ![](CB_button_NewAsset.png) button.
-1. Select _Import_ from the menu.
-1. Select the files that have the appropriate naming convention that will form your sound asset, such as: "Surround_fl.wav", "Surround_fr.wav", "Surround_sl.wav", "Surround_sr.wav".
-1. This will create a four channel asset called "Surround."
+1.	 From the Content Browser, click the **Import** button (pictured below).
+	![](ImportButton_UI.png)
+1.	 Select the files that have the appropriate naming convention that will form your sound asset, such as: "Surround_fl.wav", "Surround_fr.wav", "Surround_sl.wav", "Surround_sr.wav".
+1.	 This will create a four channel asset called "Surround."
 
 Or
 
-1. Select all your sound files in a Windows File Explorer window.
-1. Click and drag them into the **Content Browser** in Unreal Editor 4.
+1.	 Select all your sound files in a Windows File Explorer window.
+1.	 Click and drag them into the **Content Browser** in Unreal Editor 4.
 
 [/EXCERPT:wav01]
 
