#UE4doc #docs: Removed extraneous bolding as according to the instructional guide. ##Overview removed as according to the instructional guide. Changed title of How To to match instructional guide. Altered images to match the instructional guide: green -> yellow.

diff --git a/Engine/Documentation/Source/Gameplay/HowTo/AddingComponents/Blueprints/Components_BP.INT.udn b/Engine/Documentation/Source/Gameplay/HowTo/AddingComponents/Blueprints/Components_BP.INT.udn
--- a/Engine/Documentation/Source/Gameplay/HowTo/AddingComponents/Blueprints/Components_BP.INT.udn
+++ b/Engine/Documentation/Source/Gameplay/HowTo/AddingComponents/Blueprints/Components_BP.INT.udn
@@ -1,13 +1,11 @@
 Availability: Public
-Title:How to Add Components in Blueprints
+Title:Adding Components in Blueprints
 Crumbs: %ROOT%, Gameplay, Gameplay\HowTo\AddingComponents
 Description:A How To Guide for Adding Components in Blueprints. 
 
 [TOC (start:2 end:2)]
 
-## Overview
-
-On this page, you will learn how to add **Components** to a **Blueprint** and how Components can be accessed and modified inside a Blueprint. To show this, we will create a simple burning bush Blueprint Actor that is comprised of three components: the **Static Mesh** Component (the Bush), an **Audio** Component (fire sound effect) and a **Particle Effect** (the visible fire effect). 
+On this page, you will learn how to add **Components** to a Blueprint and how Components can be accessed and modified inside a Blueprint. To show this, we will create a simple burning bush Blueprint Actor that is comprised of three components: the **Static Mesh** Component (the Bush), an **Audio** Component (fire sound effect) and a **Particle Effect** (the visible fire effect). 
 
 ## Adding Components to a Blueprint
 
@@ -19,7 +17,7 @@ The steps below will show you how to Add Components to a Blueprint Actor.
 
 	![](Components1.jpg)
 
-	This will create a **Blueprint** based on this asset and will add it to the Blueprint as a **Component**. 
+	This will create a Blueprint based on this asset and will add it to the Blueprint as a Component. 
 
 1.	Click **Ok** when the **Create Blueprint** window opens. 
 
@@ -35,13 +33,13 @@ The steps below will show you how to Add Components to a Blueprint Actor.
 
 	![](Components3.png)
 
-	This window allows you to add **Components** to the Blueprint, in this case we can add a fire and audio effect. 
+	This window allows you to add Components to the Blueprint, in this case we can add a fire and audio effect. 
 
 1.	Click **Audio** to add an Audio Component and notice that the **Details** window is populated below the **Components** window. 
 
 	![](Components4.png)
 
-	This is where you specify the details regarding the **Component**. Depending upon the Component added, varying options will be presented that you can define. In this case, we want to define the actual Sound Asset that the Audio Component will use. 
+	This is where you specify the details regarding the Component. Depending upon the Component added, varying options will be presented that you can define. In this case, we want to define the actual Sound Asset that the Audio Component will use. 
 
 1.	In the **Details** panel, click the **None** box under **Sound** then select the **Fire01_Cue** asset. 
 
@@ -73,7 +71,7 @@ The steps below will show you how to Add Components to a Blueprint Actor.
 
 You should now see that we have a bush that is on fire and plays a fire audio effect. As you move away from the bush, the audio should fade out and it should fade in as you move close to the bush. 
 
-We have created a **Blueprint Actor** that is comprised of three **Components**, we could have placed the bush in the level, then placed the audio inside the bush, then placed the fire in the bush; however creating a Blueprint that has these Components is a far more manageable solution than having individual pieces. 
+We have created a Blueprint Actor that is comprised of three Components, we could have placed the bush in the level, then placed the audio inside the bush, then placed the fire in the bush; however creating a Blueprint that has these Components is a far more manageable solution than having individual pieces. 
 
 
 ## Working with Components in Blueprints
