#UE4doc #docs: Removed extraneous bolding as according to the instructional guide. ##Overview removed as according to the instructional guide. Changed title of How To to match instructional guide.

diff --git a/Engine/Documentation/Source/Gameplay/HowTo/CharacterMovement/CharacterMovementOverview.INT.udn b/Engine/Documentation/Source/Gameplay/HowTo/CharacterMovement/CharacterMovementOverview.INT.udn
--- a/Engine/Documentation/Source/Gameplay/HowTo/CharacterMovement/CharacterMovementOverview.INT.udn
+++ b/Engine/Documentation/Source/Gameplay/HowTo/CharacterMovement/CharacterMovementOverview.INT.udn
@@ -1,18 +1,15 @@
 Availability: Public
-Title:How to Set Up Character Movement
+Title:Setting Up Character Movement
 Crumbs: %ROOT%, Gameplay
 Description:A How To Guide for Setting Up a Character with Character Movement in Unreal Engine 4. 
 
-
-##Overview
-
 To create **Character Movement**, several parts are needed, each working together to produce the end result of a character moving around in your game. You will need a character, inputs for that character, how the character responds to inputs from a player, animations for the character moving, transitions to and from animations, and the logic behind the transitioning between animations. 
 
 In this "How To" series, we intend to show you how you can set all this up from scratch in either Blueprints or C++ where at the end of the series, you will have a playable character that exhibits different forms of Character Movement. 
 
 ## Setup
 
-To learn more about setting up **Character Movement** in Blueprints, follow the link below:
+To learn more about setting up Character Movement in Blueprints, follow the link below:
 
 [REGION:buttonlist]
 * [](Gameplay\HowTo\CharacterMovement\Blueprints)
