Documentation: Removed some dead command line arguments from the doc

diff --git a/Engine/Documentation/Source/Programming/Basics/CommandLineArguments/CommandLineArguments.INT.udn b/Engine/Documentation/Source/Programming/Basics/CommandLineArguments/CommandLineArguments.INT.udn
--- a/Engine/Documentation/Source/Programming/Basics/CommandLineArguments/CommandLineArguments.INT.udn
+++ b/Engine/Documentation/Source/Programming/Basics/CommandLineArguments/CommandLineArguments.INT.udn
@@ -163,10 +163,7 @@ These arguments can be passed to either the game or the editor, depending on the
 * `EXEC`: Executes the specified exec file.
 * `FPS`: Set the frames per second for benchmarking.
 * `FULLSCREEN`: Set game to run in fullscreen mode.
-* `PIXO`: Use pixomatic software rendering.
 * `SECONDS`: Set the maximum tick time.
-* `SM3`: Force the use of shader model 3. 
-* `SM4` / `D3D10`: Force use of DirectX 10 if available.
 * `WINDOWED`: Set game to run in windowed mode.
 
 
