#UE4doc #docs: Removed extraneous bolding as according to the instructional guide. Altered images to match the instructional guide: green -> yellow.

diff --git a/Engine/Documentation/Source/Gameplay/HowTo/CharacterMovement/Blueprints/Input/Input.INT.udn b/Engine/Documentation/Source/Gameplay/HowTo/CharacterMovement/Blueprints/Input/Input.INT.udn
--- a/Engine/Documentation/Source/Gameplay/HowTo/CharacterMovement/Blueprints/Input/Input.INT.udn
+++ b/Engine/Documentation/Source/Gameplay/HowTo/CharacterMovement/Blueprints/Input/Input.INT.udn
@@ -64,7 +64,7 @@ Here we will set up the **Input** (Game Controls) to control the character and c
 	![](GameMode_2.png)
 
 	[REGION:note]
-	For more details on **Game Modes**, see [](Gameplay\HowTo\SettingUpAGameMode) for more information. 
+	For more details on Game Modes, see [](Gameplay\HowTo\SettingUpAGameMode) for more information. 
 	[/REGION]
 
 1.	After setting both the Input and Game Mode settings, close the Project Settings Window. 
diff --git a/Engine/Documentation/Source/Gameplay/HowTo/CharacterMovement/Blueprints/Overview.INT.udn b/Engine/Documentation/Source/Gameplay/HowTo/CharacterMovement/Blueprints/Overview.INT.udn
--- a/Engine/Documentation/Source/Gameplay/HowTo/CharacterMovement/Blueprints/Overview.INT.udn
+++ b/Engine/Documentation/Source/Gameplay/HowTo/CharacterMovement/Blueprints/Overview.INT.udn
@@ -1,5 +1,5 @@
 Availability: Public
-Title:How to Set Up Character Movement in Blueprints
+Title:Setting Up Character Movement in Blueprints
 Crumbs: %ROOT%, Gameplay, Gameplay\HowTo\CharacterMovement
 Description:A How To Guide for Setting Up Character Movement in Blueprints. 
 
