#UEDoc #docs - Updated doc and images

diff --git a/Engine/Documentation/Source/Engine/Physics/PhysicalMaterials/PhysMatUserGuide/PhysMatUserGuide.INT.udn b/Engine/Documentation/Source/Engine/Physics/PhysicalMaterials/PhysMatUserGuide/PhysMatUserGuide.INT.udn
--- a/Engine/Documentation/Source/Engine/Physics/PhysicalMaterials/PhysMatUserGuide/PhysMatUserGuide.INT.udn
+++ b/Engine/Documentation/Source/Engine/Physics/PhysicalMaterials/PhysMatUserGuide/PhysMatUserGuide.INT.udn
@@ -45,55 +45,47 @@ Description:Physical Materials User Guide: the creation and application of Physi
 
 [TOC(start:2)]
 
-## Overview
 
 This document will cover the creation and usage of **Physical Materials** as well as enabling or editing **SurfaceTypes** for your project.
 
 ## Creation
 
-1. From the **Content Browser**, click ![](Engine/Audio\SoundCues\Editor\button_CB_NewAsset.png) **New Asset** -> **Physics** -> **Physical Material** or **right-click** in the **Content Browser** -> **Physics** -> **Physical Material**.
+1. From the Content Browser. **click** ![](Engine/Audio\SoundCues\Editor\button_CB_NewAsset.png) **New Asset** -> **Physics** -> **Physical Material** or **right-click** in the **Content Browser** -> **Physics** -> **Physical Material**.
 
 	![](PhysMat_Creation_01.png)
 
-1. **Double click** the new **Physical Material** to edit its properties.
+1. **Double click** the new Physical Material to edit its properties.
 
 	![](PhysMat_Creation_02.png)
 
-1. Adjust Properties.
+1. **Adjust Properties**.
 
 	![](PhysMat_Creation_03.png)
 
-1. Save ![](Shared\Icons\icon_SaveAsset_40x.png)
+1. **Click Save** ![](Shared\Icons\icon_SaveAsset_40x.png)
 
 See the [Physical Material Reference](Engine/Physics\PhysicalMaterials\PhysMatReference) for information on the properties in the Physical Material.
 
 ### Surface Type
 
-By default, Unreal Engine 4 supports 30 Surface Types, which can be labeled anything you desire, stored in your project's `DefaultEngine.ini`, which can be found in `YourProjectRoot\Config\DefaultEngine.ini`. Most project templates do not include the `\[PhysicalMaterial.SurfaceTypes\]` category nor any definitions for Surface Types, but it can easily be enabled by adding the following to your `DefaultEngine.ini`:
+By default, Unreal Engine 4 supports 62 Surface Types, which can be labeled anything you desire, stored in your project's `DefaultEngine.ini`, which can be found in `YourProjectRoot\Config\DefaultEngine.ini`. 
+
 
-	[PhysicalMaterial.SurfaceTypes]
-	SurfaceType1=Glass
-	SurfaceType2=Metal
-	SurfaceType3=Wood
-	SurfaceType4=Concrete
-	...
-	SurfaceType30=CollapsedStar
 
-Once added, save the `DefaultEngine.ini` file. You will need to re-load your project for the Surface Types to appear on any Physical Materials or to be used as an enumeration in Blueprints.
 
 ## Usage
 
 ### Material
 
-1. Open or create a new Material.
+1. **Open** or **create** a new Material.
 
 	![](PhysMat_Material_01.png)
 
-1. Select the main material node.
+1. **Select** the main material node.
 
 	![](PhysMat_Material_02.png)
 
-1. Change the Physical Material.
+1. **Change** the Physical Material.
 
 	![](PhysMat_Material_03.png)
 
@@ -101,11 +93,11 @@ Once added, save the `DefaultEngine.ini` file. You will need to re-load your pro
 
 ### Material Instance
 
-1. Open or create a new Material Instance.
+1. **Open** or **create** a new Material Instance.
 
 	![](PhysMat_MaterialInst_01.png)
 
-1. Change the Physical Material.
+1. **Change** the Physical Material.
 
 	![](PhysMat_MaterialInst_02.png)
 
@@ -116,15 +108,15 @@ Once added, save the `DefaultEngine.ini` file. You will need to re-load your pro
 When adjusting the Physical Material for a Physics Asset, it is best to start by assigning the most commonly used Physical Material to all of the Physics Bodies in the Physics Asset.
 
 [EXCERPT:PhysMatSkelMesh1]
-1. Open a Physics Asset in the Physics Asset Tool.
+1. **Open** a Physics Asset in the Physics Asset Tool.
 
 	![](PhysMat_PhAT_02.png)
 
-1. Select a Physical Material in the **Content Browser**.
+1. **Select** a Physical Material in the Content Browser.
 
 	![](PhysMat_PhAT_03.png)
 
-1. In the Physics Asset Tool, click the Apply PhysMat button: ![](Shared\Icons\icon_PhAT_PhysMat_40x.png).
+1. In the Physics Asset Tool, **click** the **Apply PhysMat button**: ![](Shared\Icons\icon_PhAT_PhysMat_40x.png).
 	[REGION:raw]
 	![](PhysMat_PhAT_01.png)(w:860)
 	[/REGION]
@@ -133,19 +125,21 @@ When adjusting the Physical Material for a Physics Asset, it is best to start by
 
 If specific Physics Bodies need differing Physical Materials, you can then adjust them individually.
 
-1. Enable Body Edit Mode in the Physics Asset Tool.
-1. Select a Physics Body.
-1. In the **Details** pane, in the Physics Category, change the **Simple Collision Physical Material**.
+1. **Enable Body Edit Mode** in the Physics Asset Tool.
+1. **Select** a Physics Body.
+1. In the Details pane, in the Physics Category, change the **Simple Collision Physical Material**.
 
 
 [REGION:raw]
 ![](PhysMat_PhAT_04.png)(w:920)
 [/REGION]
 
-The default behavior for Physics interactions with Skeletal Meshes is to only interact with its associated Physics Asset, so the Skeletal Mesh's material's Physical Material will not be used.
+The default behavior for Physics interactions with Skeletal Meshes is to only interact with its associated Physics Asset, so the Skeletal Mesh's material's Physical Material will 
+not be used.
 
 [REGION:note]
-With Physics Assets, though you are setting the Simple Collision Physical Material, when you trace against a Physics Asset, you will need to do a complex trace and the complex trace will return the Physical Material listed in the **Simple Collision Physical Material** property of the hit Physics Body.
+With Physics Assets, though you are setting the Simple Collision Physical Material, when you trace against a Physics Asset, you will need to do a complex trace and the complex trace 
+will return the Physical Material listed in the Simple Collision Physical Material property of the hit Physics Body.
 [/REGION]
 
 ### StaticMeshes
@@ -154,27 +148,27 @@ StaticMeshes have both Simple collision (the Physics Bodies that you create in y
 
 | Collision | Physical Material Order |
 | ----- | ----- |
-| **Simple** | When any collision or trace uses **Simple Collision**, it will reference the Physical Material on the StaticMesh that is set in the StaticMesh Editor, or if the _Phys Material Override_ on the **StaticMesh Actor** is set to something other than `None`, it will use what is listed in that property. |
-| **Complex** | When any collision or trace uses **Complex Collision**, it will reference the Physical Material on the **Material** or **Material Instances** applied to the **StaticMesh Actor**, or if the Phys Material Override on the StaticMesh Actor is set to something other than `None`, it will use what is listed in that property. |
+| **Simple** | When any collision or trace uses Simple Collision, it will reference the Physical Material on the StaticMesh that is set in the StaticMesh Editor, or if the _Phys Material Override_ on the StaticMesh Actor is set to something other than `None`, it will use what is listed in that property. |
+| **Complex** | When any collision or trace uses Complex Collision, it will reference the Physical Material on the Material or Material Instances applied to the StaticMesh Actor, or if the Phys Material Override on the StaticMesh Actor is set to something other than `None`, it will use what is listed in that property. |
 
-To set the Simple Collision Physical Material for a **StaticMesh**:
+To set the Simple Collision Physical Material for a StaticMesh:
 
-1. **Double Click** a Static Mesh in the **Content Browser** to bring up the Static Mesh Editor.
+1. **Double Click** a Static Mesh in the Content Browser to bring up the Static Mesh Editor.
 
 	![](PhysMat_StaticMeshEd_01.png)
 
-1. Change the **Simple Collision Physical Material** property in the **Static Mesh Settings** category to the desired Physical Material.
+1. Change the Simple Collision Physical Material property in the **Static Mesh Settings** category to the desired Physical Material.
 
 	![](PhysMat_StaticMeshEd_02.png)
 
-1. Save ![](Shared\Icons\icon_SaveAsset_40x.png).
+1. **Click Save** ![](Shared\Icons\icon_SaveAsset_40x.png).
 
 
 ### Misc
 
 ![](PhysMat_Override.png)
 
-The **Phys Material Override** property exists on everything with a **Physics** category. This can be used to completely override Simple Collision Physical Materials on an **Actor** or **Component** with the selected Physical Material.
+The _Phys Material Override_ property exists on everything with a **Physics** category. This can be used to completely override Simple Collision Physical Materials on an Actor or Component with the selected Physical Material.
 
 * Overriding a StaticMesh's simple collision Physical Material.
 * Since Skeletal Meshes Physics Assets always return Simple Collision, you can use this to override all Physical Materials on a placed Skeletal Mesh Actor.
