#UEDoc #docs - Updated image and bold lettering.

diff --git a/Engine/Documentation/Source/Engine/Physics/Substepping/Substepping.INT.udn b/Engine/Documentation/Source/Engine/Physics/Substepping/Substepping.INT.udn
--- a/Engine/Documentation/Source/Engine/Physics/Substepping/Substepping.INT.udn
+++ b/Engine/Documentation/Source/Engine/Physics/Substepping/Substepping.INT.udn
@@ -42,20 +42,21 @@ Description:
 [/OBJECT]
 [/VAR]
 
-## Overview
 
-One of the lesser known features in UE4 is the ability to turn on physics sub-stepping. By doing this, you can get physics simulations that are more accurate and stable. The most noticeable improvement will be with ragdoll jitter and other complex physical assets.
+
+One of the lesser known features in UE4 is the ability to turn on physics **Sub-stepping**. By doing this, you can get physics simulations that are more accurate and stable. 
+The most noticeable improvement will be with ragdoll jitter and other complex physical assets.
 
 ## Enabling Sub-Stepping
 
-Sub-stepping can be turned on in the project specific physics settings:
+Sub-stepping can be turned on in the **Project Setting** -> **Physics** tab:
 
 [REGION:raw]
 ![](SS_Enabled.png)(w:940)
 [/REGION]
 
 [REGION:note]
-This feature is still a work in progress, and is not fully supported by APEX destruction. It is also currently not supported on mobile devices. However, we are working towards this as it can greatly improve the look and feel.
+This feature is still a work in progress, and is not fully supported by APEX Destruction. It is also currently not supported on mobile devices. However, we are working towards this as it can greatly improve the look and feel.
 [/REGION]
 
 | Property | Description |
@@ -67,7 +68,7 @@ This feature is still a work in progress, and is not fully supported by APEX des
 
 ## Technical Details
 
-UE4 uses a variable frame rate. While this is good for hardware scalability, it creates a challenge for the physics engine, which works best with small fixed time steps. Sub-stepping takes the total frame time and divides it into sub-steps. The physics simulation is then ticked multiple times per frame. The number of sub-steps taken will depend on how small your max sub-step delta time is set to. The smaller the max sub-step time the more stable your simulation will be, but at a greater CPU cost.
+Unreal Engine 4 uses a variable frame rate. While this is good for hardware scalability, it creates a challenge for the physics engine, which works best with small fixed time steps. Sub-stepping takes the total frame time and divides it into sub-steps. The physics simulation is then ticked multiple times per frame. The number of sub-steps taken will depend on how small your max sub-step delta time is set to. The smaller the max sub-step time the more stable your simulation will be, but at a greater CPU cost.
 
 Sub-stepping is completely hidden from the user, which means that some calls to the physics engine have to be interpolated or maintained. For example, if a user applies a force to an Actor for a single frame, and the frame is internally sub-stepped N times, we need to apply the force for N consecutive simulation steps in order to achieve the same acceleration. Similarly, if the user sets the target location of an Actor, we must interpolate the target location over multiple sub-steps in order to maintain the desired speed. All of these details are already handled internally by UE4, but there are some CPU and memory costs associated with the required bookkeeping.
 
