#UEDoc #docs - Updated images and bold lettering

diff --git a/Engine/Documentation/Source/Engine/Physics/Vehicles/VehcileContentCreation/Vehicle_Content.INT.udn b/Engine/Documentation/Source/Engine/Physics/Vehicles/VehcileContentCreation/Vehicle_Content.INT.udn
--- a/Engine/Documentation/Source/Engine/Physics/Vehicles/VehcileContentCreation/Vehicle_Content.INT.udn
+++ b/Engine/Documentation/Source/Engine/Physics/Vehicles/VehcileContentCreation/Vehicle_Content.INT.udn
@@ -25,9 +25,11 @@ Navigation:topic
 [/VAR]
 
 
-## Overview
 
-The basic, bare minimum, art setup for a vehicle is just a **Skeletal Mesh**. The type of vehicle will decide how complicated an art setup you will need, and special consideration may need to be given to the suspension. A tank does not need to have a crazy suspension setup because it will almost always go unseen, but a dune buggy, such as the one in **Vehicle Game**, will require more joints to make all the exposed components move in a believable way.
+
+The basic, bare minimum, art setup for a vehicle is just a **Skeletal Mesh**. The type of vehicle will decide how complicated an art setup you will need, and special consideration 
+may need to be given to the suspension. A tank does not need to have a crazy suspension setup, because it will almost always go unseen, but a dune buggy, such as the one in 
+**Vehicle Game**, will require more joints to make all the exposed components move in a believable way.
 
 ## Basics
 
@@ -39,7 +41,8 @@ You will need to measure the radius of your wheels, in centimeters, for use with
 
 ![](wheelRadius.png)
 
-In the image above, the Distance Measurement tool in Maya is being used to get the diameter of the wheels based on the distance between 2 vertices opposite each other on the wheel. The radius of the wheel is half of this value.
+In the image above, the Distance Measurement tool in Maya is being used to get the diameter of the wheels based on the distance between 2 vertices opposite each other on the wheel. 
+The radius of the wheel is half of this value.
 
 A similar tool exists in the Helpers section of 3D Studio Max.
 
@@ -53,11 +56,13 @@ This will ensure they roll on the Y axis and steer on the Z axis.
 
 All other joints can be arranged as you see fit, but it should be noted that things like **Look At** nodes for the **Animation Blueprint** assume X forward.
 
-To prevent visual oddities, the joints for the wheels should be accurately centered. The visual mesh will not be used for collision detection, however, if the wheel mesh is off center, then it will look like the wheel is broken, and will be really noticeable due to motion blur.
+To prevent visual oddities, the joints for the wheels should be accurately centered. The visual mesh will not be used for collision detection, however, if the wheel mesh is off center,
+then it will look like the wheel is broken, and will be really noticeable due to motion blur.
 
 ## Binding
 
-Standard smooth bind for Maya or Skin Modifier for 3DS Max. Wheels should only have weighting to one joint, so they can spin free with no odd deformation. For shocks and struts, you can get away with some fancy skinning, but it will require more thought on the Unreal Editor side. 
+Standard smooth bind for Maya or Skin Modifier for 3DS Max. Wheels should only have weighting to one joint, so they can spin free with no odd deformation. For shocks and struts, you 
+can get away with some fancy skinning, but it will require more thought on the Unreal Editor side. 
 
 
 ## Export
@@ -73,34 +78,40 @@ Standard FBX import for a Skeletal Mesh. You will want to have the importer crea
 
 ![](PAicon.png)
 
-The **Physics Asset** is incredibly important for vehicles, and cannot be overlooked or skipped. When you first have Unreal Engine 4 create a **Physics Asset** for you, it will probably end up looking much like this:
+The Physics Asset is incredibly important for vehicles, and cannot be overlooked or skipped. When you first have Unreal Engine 4 create a Physics Asset for you, it will 
+probably end up looking much like this:
 
 ![](badPA.png)
 
-That is because the **Physics Asset Tool** (**PhAT**) attempts to wrap the vertices skinned to a joint the best it can. And honestly, you are going to delete and re-create them all. Why? The **PhAT** does not currently have a way to gracefully handle the re-creation of the constraints that hold all the Physics Bodies together. So if you delete an intermediate Physics Body, it does not create a new constraint up the hierarchy. However, if you delete all the Physics Bodies and start building them from the root joint out, all your constraints will get created correctly.
+That is because the **Physics Asset Tool** (**PhAT**) attempts to wrap the vertices skinned to a joint the best it can. And honestly, you are going to delete and re-create them all. 
+Why? The PhAT does not currently have a way to gracefully handle the re-creation of the constraints that hold all the Physics Bodies together. So if you delete an intermediate 
+Physics Body, it does not create a new constraint up the hierarchy. However, if you delete all the Physics Bodies and start building them from the root joint out, all your constraints 
+will get created correctly.
 
-In the **Hierarchy** pane, **Shift** Select everything, and then press the **Delete** key. This will remove all Physics Bodies from the asset.
+In the **Hierarchy** pane, **Shift Select** everything, and then **press** the **Delete** key. This will remove all Physics Bodies from the asset.
 
 ![](selectAllDelete.png)
 
-Now, starting with your root joint, start creating physics bodies on your vehicles joints. Keeping in mind, you only need a physics body on a joint that either needs to be physically simulated or affects the bounds of the vehicle.
+Now, starting with your root joint, start creating physics bodies on your vehicles joints. Keeping in mind, you only need a physics body on a joint that either needs to be physically 
+simulated or affects the bounds of the vehicle.
 
-In the case of our buggy here, a box for the root and spheres for the wheels will be just fine to start, but it will require a number of other Physics Bodies to get the desired behavior.
+In the case of our buggy here, a box for the root and spheres for the wheels will be just fine to start, but it will require a number of other Physics Bodies to get the desired 
+behavior.
 
 ### Create the Body Collision
 
 To create the box for the root of the vehicle:
 
 1. **Right-click** the root joint in the **Hierarchy** pane.
-1. Select **Add Body**
-1. Set the **Collision Geometry** to **Box**.
-1. Click "**OK**"
+1. **Select** **Add Body**
+1. **Set** the **Collision Geometry** to **Box**.
+1. **Click** "**OK**"
 
 	[REGION:raw]
 	![](bodyToRoot.png)(convert:false)
 	[/REGION]
 
-1. You can then move, rotate, and scale the new **Physics Body** to better approximate the shape of your vehicle
+1. You can then **move**, **rotate**, and **scale** the new Physics Body to better approximate the shape of your vehicle
 
 	![](newBodyBoxResultScaled.png)
 
@@ -109,24 +120,25 @@ To create the box for the root of the vehicle:
 
 To create the spheres for the wheels:
 
-1. **Right-click** a wheel joint in the **Hierarchy** pane.
-1. Select **Add Body**.
+1. **Right-click** a wheel joint in the Hierarchy pane.
+1. **Select Add Body**.
 
 	![](bodyToWheel.png)
 
-1. Set the **Collision Geometry** to **Sphere**.
+1. **Set** the **Collision Geometry** to **Sphere**.
 1. Click "**OK**"
 
 	![](spheres.png)
 
-1. In the **Details** pane, find the **Physics Type** property and set it to **Kinematic**.
+1. In the Details pane, find the **Physics Type** property and **set it to Kinematic**.
 
 	![](makeTheWheelsKin.png)
 
-	This is necessary for the wheels to affect the bounds of the vehicle so shadows and culling occur correctly, and also, to prevent them from just rolling off the vehicle once your game begins.
+	This is necessary for the wheels to affect the bounds of the vehicle so shadows and culling occur correctly, and also, to prevent them from just rolling off the vehicle once your 
+	game begins.
 
 	[REGION:note]
-	The wheel **Physics Bodies** are never actually used for collision. Currently, the wheels use a ray casting for interacting with the world.
+	The wheel Physics Bodies are never actually used for collision. Currently, the wheels use a ray casting for interacting with the world.
 	[/REGION]
 
 
@@ -136,7 +148,8 @@ This is just a simple setup. If you examine the Physics Asset for the Vehicle in
 
 ![](goodPA.png)
 
-Outside of the wheels, all of those collision bodies exist on the root joint of the vehicle. They prevent the wheels from getting jammed on certain game elements, and also prevents the meshes of the wheels from clipping through walls and railings.
+Outside of the wheels, all of those collision bodies exist on the root joint of the vehicle. They prevent the wheels from getting jammed on certain game elements, and also prevents 
+the meshes of the wheels from clipping through walls and railings.
 
 
 
