#UE4doc #docs: Updated tutorial to follow v4.8

diff --git a/Engine/Documentation/Source/Gameplay/HowTo/CharacterMovement/Blueprints/Setup_1/Setup_1.INT.udn b/Engine/Documentation/Source/Gameplay/HowTo/CharacterMovement/Blueprints/Setup_1/Setup_1.INT.udn
--- a/Engine/Documentation/Source/Gameplay/HowTo/CharacterMovement/Blueprints/Setup_1/Setup_1.INT.udn
+++ b/Engine/Documentation/Source/Gameplay/HowTo/CharacterMovement/Blueprints/Setup_1/Setup_1.INT.udn
@@ -34,7 +34,7 @@ For this project, we are using the **Blank** Project Template in order to start
 
 	![](Characters_1b.png)
 
-1.	In the **Target Skeleton** window, select the **HeroTPP_Skeleton** asset. 
+1.	In the **Target Skeleton** window, select the **UE4_Mannequin_Skeleton** asset. 
 
 	![](Characters_1c.png)
 
@@ -48,7 +48,7 @@ For this project, we are using the **Blank** Project Template in order to start
 
 	![](ClickMesh.png)
 
-1.	In the **Details** panel, click the **None** box next to **Skeletal Mesh** and assign the **HeroTPP** Skeletal Mesh. 
+1.	In the **Details** panel, click the **None** box next to **Skeletal Mesh** and assign the **UE4_Mannequin_Skeleton** Skeletal Mesh. 
 
 	![](Characters_2.png)
 
