Updated for Jira UEDOC-1824
Also split off the page into two pages instead of one page.

diff --git a/Engine/Documentation/Source/Engine/Blueprints/BP_HowTo/BlueprintHowTo.INT.udn b/Engine/Documentation/Source/Engine/Blueprints/BP_HowTo/BlueprintHowTo.INT.udn
--- a/Engine/Documentation/Source/Engine/Blueprints/BP_HowTo/BlueprintHowTo.INT.udn
+++ b/Engine/Documentation/Source/Engine/Blueprints/BP_HowTo/BlueprintHowTo.INT.udn
@@ -44,7 +44,7 @@ Refer to the tables below for more information and click each link to view its c
 | [](Engine\Blueprints\BP_HowTo\Interfaces) | %Engine\Blueprints\BP_HowTo\Interfaces:description% |
 | [](Engine\Blueprints\BP_HowTo\EventDispatcher) | %Engine\Blueprints\BP_HowTo\EventDispatcher:description% |
 | [](Engine/Blueprints/BP_HowTo/MakingMacros) | %Engine/Blueprints/BP_HowTo/MakingMacros:description% |
-
+| [](Engine\Blueprints\BP_HowTo\MacroLibraries) | %Engine\Blueprints\BP_HowTo\MacroLibraries:description% |
 
 <!-- To remove since information is redundant. 
 
diff --git a/Engine/Documentation/Source/Engine/Blueprints/BP_HowTo/MacroLibraries/MacroLibraries.INT.udn b/Engine/Documentation/Source/Engine/Blueprints/BP_HowTo/MacroLibraries/MacroLibraries.INT.udn
new file mode 100644
--- /dev/null
+++ b/Engine/Documentation/Source/Engine/Blueprints/BP_HowTo/MacroLibraries/MacroLibraries.INT.udn
@@ -0,0 +1,104 @@
+Availability: Public
+Title:Using Macro Libraries
+Crumbs: %ROOT%, Engine, Engine\Blueprints, Engine\Blueprints\BP_HowTo
+Description:Macros inside a Macro Library are used to increase the health and scale of an Actor. 
+Related: Engine/Blueprints/UserGuide/Macros
+Related: Engine/Blueprints/UserGuide/Types/MacroLibrary
+Related: Engine/Blueprints
+
+A **Macro Library** is a container that holds a collection of Macros or self-contained graphs that can be placed as nodes in other Blueprints. These can be time-savers as they can store commonly used sequences of nodes complete with inputs and outputs for both execution and data transfer.
+
+## Creating a Macro Library
+
+In this example we create a Macro Library consisting of two macros, one of which is used to add to a "Health" variable and the other is used to adjust the size of an Actor that is specified. 
+
+[REGION:note]
+_For this example, we are using the **Blueprint Third Person Project** with **Starter Content** enabled._
+[/REGION]
+
+
+1.	**Right-click** in an empty space in the **Content Browser** then in the context menu under **Blueprints**, select [Blueprint Macro Library](Engine/Blueprints/UserGuide/Types/MacroLibrary). 
+
+	![](CreateLib1.png)(w:640)
+
+1.	In the **Pick Parent Class** window that appears, select your [Parent Class](Engine/Blueprints/UserGuide/Types/ClassBlueprint) (for this example, select **Actor**). 
+
+	![](CreateLib2.png)
+
+1.	Enter a name for your Macro Library, then **Double-click** on it to open it up. 
+
+	![](CreateLib3.png)
+
+	This will display the **Blueprint Macro** interface. 
+
+	![](CreateLib4.png)(w:640)
+
+1.	In the **MyBluprint** window in the upper-right, rename the default macro to **ScaleUp** by pressing **F2** on it. 
+
+1.	In the **Details** panel for the **ScaleUp** macro, add the **Inputs** and **Outputs** shown below. 
+
+	![](CreateLib5.png)
+
+	Above we have created two inputs, one called _In_ set to the **Exec** type and the other called _InActor_ set to the **Actor** type. The output added is called _Out_ and is set to the **Exec** type. We will use these to determine the Actor that we want to affect through this macro. 
+
+1.	In the graph for the **ScaleUp** macro, recreate the node network indicated below. 
+
+	![](CreateLib6.png)
+
+	Above, when this macro is called it will get whatever Actor is provided as the **InActor** and get its current scale to which it will multiply it by 1.25 before setting it as the new scale 3D for the **Target** (which is the Actor specified as the _InActor_.) We can now use this macro with any Actor and affect their scale through the use of this macro. 
+
+1.	Add another macro from the **MyBlueprint** window called **AddHealth**. 
+
+1.	In the **Details** panel for **AddHealth**, add the **Inputs** and **Outputs** shown below. 
+
+	![](CreateLib7.png)
+
+	Above we are taking in a float value called _GetHealth_ which we will add to before we pass it through to the _SetHealth_ output node. 
+
+1.	In the graph for the **AddHealth** macro, recreate the node network indicated below. 
+
+	![](CreateLib8.png)
+
+	Above we are adding **50** to whatever float value is provided as the _GetHealth_ value before outputting the resulting value to _SetHealth_. 
+
+1.	**Save** and close the Blueprint Macro Library. 
+
+1.	Inside the **Content/ThirdPersonBP/Blueprints** folder, open the **ThirdPersonCharacter** Blueprint. 
+
+1.	In the **MyBlueprint** window, add a **Float** variable to represent **Health** (we left the default value to 0). 
+
+	![](CreateLib9.png)
+
+1.	In the Event Graph, recreate the node network indicated below. 
+
+	![](CreateLib10.png)
+
+	Above we are calling the **AddHealth** macro whenever **Q** is pressed which takes in the _Health_ variable (performs the macro script of increasing it) and updates it with the out pin _SetHealth_ before printing it to the screen. We then use **E** and check if _Health_ is greater than 100 before we allow the **ScaleUp** macro to be called on the Actor **Self** which is the **ThirdPersonCharacter**. 
+
+1.	**Compile** and **Play** in the editor. 
+
+	[OBJECT:EmbeddedVideo]
+	[PARAMLITERAL:width]
+	640
+	[/PARAMLITERAL]
+	[PARAMLITERAL:height]
+	360
+	[/PARAMLITERAL]
+	[PARAMLITERAL:videoid]
+	HB2w84EHowQ
+	[/PARAMLITERAL]
+	[/OBJECT]
+
+	_To better see the text printed to the screen, watch the video above in fullscreen._
+
+	In the video above, we included the Blueprint Script from step 13 so that we can see when the script is being fired. When we press **E** to attempt to call the **ScaleUp** macro, we first check the _Health_ value which by default we set to 0. We then press **Q** to call the **AddHealth** macro to increase the _Health_ variable each time it is pressed by 50. After a few presses, we then try the **ScaleUp** macro again and now we can see that it gets the call because our _Health_ has been increased and is greater than 100. 
+
+	These macros can be called from any other Blueprints provided the inputs for Health and the Target Actor are supplied. 
+
+	
+
+
+
+
+
+
diff --git a/Engine/Documentation/Source/Engine/Blueprints/BP_HowTo/MakingMacros/MakingMacros.INT.udn b/Engine/Documentation/Source/Engine/Blueprints/BP_HowTo/MakingMacros/MakingMacros.INT.udn
--- a/Engine/Documentation/Source/Engine/Blueprints/BP_HowTo/MakingMacros/MakingMacros.INT.udn
+++ b/Engine/Documentation/Source/Engine/Blueprints/BP_HowTo/MakingMacros/MakingMacros.INT.udn
@@ -1,27 +1,20 @@
 Availability: Public
 Title:Making Macros
 Crumbs: %ROOT%, Engine, Engine\Blueprints, Engine\Blueprints\BP_HowTo
-Description:Examples of how to create Macros and use Macro Libraries. 
+Description:A Macro is used to check if a player has enough energy to jump. 
 Related: Engine/Blueprints/UserGuide/Macros
 Related: Engine/Blueprints/UserGuide/Types/MacroLibrary
 Related: Engine/Blueprints
 
-
-[TOC(start:2)]
-
-## Overview
-
 **Macros** are essentially the same as collapsed graphs of nodes. They have an entry point and exit point designated by tunnel nodes. Each tunnel can have any number of execution or data pins which are visible on the macro node when used in other Blueprints and graphs.
 
-A **Macro Library** is a container that holds a collection of Macros or self-contained graphs that can be placed as nodes in other Blueprints. These can be time-savers as they can store commonly used sequences of nodes complete with inputs and outputs for both execution and data transfer.
-
-Below you will find examples on both creating Macros and creating a Macro Library. 
-
 ## Creating Macros
 
 In this example, we will create a **Macro** that will be used to check if a player character has enough energy to jump or not. If they have enough energy to jump, we will remove some energy from the player and print their current value to the screen then allow them to jump. If they do not have enough energy, we will print to the screen that they are out of energy and not allow them to jump. 
 
+[REGION:note]
 _For this example, we are using the **Blueprint Third Person Project** with **Starter Content** enabled._
+[/REGION]
 
 1.  Inside the **Content/ThirdPersonBP/Blueprints** folder, open the **ThirdPersonCharacter** Blueprint. 
 
@@ -101,170 +94,21 @@ _For this example, we are using the **Blueprint Third Person Project** with **St
 
 1.	Click the **Play** button from the Main Toolbar to play in the Editor. 
 
-	[![](ExampleImage.png)](MacroExample.mp4)
-
-	You should now see that when you press the **Spacebar** to jump, in the upper left corner of the screen the value of **Energy** is printed to the screen. When **Energy** is 0, you should no longer be able to jump. 
+	[OBJECT:EmbeddedVideo]
+	[PARAMLITERAL:width]
+	640
+	[/PARAMLITERAL]
+	[PARAMLITERAL:height]
+	360
+	[/PARAMLITERAL]
+	[PARAMLITERAL:videoid]
+	nnKHWiQba0s
+	[/PARAMLITERAL]
+	[/OBJECT]
+
+	When you press the **Spacebar** to jump, in the upper left corner of the screen the value of **Energy** is printed to the screen. When **Energy** is 0, you should no longer be able to jump. 
 	
 	This is basic example of using a Macro to execute and consolidate script into a single node, improving the readability of our Event Graph and main character script. Additionally, you could call this Macro in other instances - say you have some other action that reduces the player's energy and you want to check if they have enough energy to perform the action (attacking for example). You could run this Macro to check if the player has enough energy to Attack after pressing your Attack key then execute an attack off the **HasEnergy** exec pin. 
 
 
-
-
-## Creating a Macro Library
-
-In this example, we create a **Macro Library** which will contain two Macros in it that will be shared by multiple Blueprints. 
-
-_For this example, we are using the **Blueprint Third Person Project** with **Starter Content** enabled._
-
-1.	Inside your project, **Right-click** in the **Content Browser** and under **Blueprints**, select **Blueprint Macro Library**. 
-
-	![](MacroHT12.png)
-
-1.	In the **Pick Parent Class** window, select **Actor**. 
-
-	![](MacroHT13.png)
-
-	Pick Parent Class allows macros in this library to access variables and functions declared in the parent class as if they were in scope.  For example, macros in a library based on Pawn can call **GetController** without being passed a Pawn, and they will only show up as options in the context menu for Blueprints based on Pawn or a more derived class like Character.
-
-1.	Give the Macro Library a name, then **Double-click** on it to open it up. 
-
-	![](MacroHT14.png)
-
-	When opening the Macro Library, you may notice that it looks slightly different than a normal Blueprint. 
-
-1.	Select the macro inside the **My Blueprint** window and press **F2** to rename it to **RotateMacro**. 
-
-	![](RotateMacroRenamed.png)
-
-1.	Click the **New** button under Inputs twice and under Outputs twice to create two pins for each. 
-
-	![](MacroHT16.png)
-
-1.	For **Inputs**, set one pin to an **Exec** called **In** and the other to an **Actor** called **Target**. 
-
-1.	For **Outputs**, set one pin to an **Exec** called **Out** and the other to a **Rotator** called **Rotation**. 
-
-	![](MacroHT17.png)
-
-	This script will be used to rotate the specified Actor around by a specified amount. 
-
-1.	Drag off the **Target** pin in the graph and add the **Get Actor Rotation** node and connect it to a **Break Rot** node. 
-
-	![](MacroHT18.png)
-
-1.	**Right-click** in the graph and add a **Make Rot** node. 
-
-1.	Drag off the **Yaw** of the **Break Rot** node and add the **+** node, with the value set to **5**. 
-
-1.	Connect the out of the **+** node to the **Yaw** on the **Make Rot** node. 
-
-1.	Connect the **Pitch** and **Roll** from the **Break** node to the **Make** node. 
-
-1.	Connect the remaining nodes as shown below. 
-
-	![](MacroHT19.png)
-
-	With the script above we are adding + 5 to the Actor's current Rotation for the Yaw, and setting that as the **Rotation** value. 
-
-1.	Click the **Add Macro** button to create another Macro and call it **HideMacro** and add two Input pins. 
-
-1.	Set one of the Input pins to an **Exec** called **In** and the other to an **Actor** called **Target**. 
-
-	![](MacroHT20.png)
-
-	You do not always need to have Outputs, but more than likely you will have some Inputs to enter your Macro. 
-
-1.	Off the **In** pin in the graph, add a **FlipFlop** node. 
-
-	![](MacroHT21.png)
-
-	Above, each time this Macro is entered it will cycle between the A pin and the B pin (flip-flopping between the two). 
-
-1.	Drag off the **Target** pin and search for and add the **Set Actor Hidden in Game** node. 
-
-1.	Connect the **A** pin to the in pin on the **Set Actor Hidden in Game** node and check the **New Hidden** box. 
-
-1.	Repeat the previous two steps for the **B** pin, but leave **New Hidden** unchecked. 
-
-	![](MacroHT22.png)
-
-	The Macro will now set the Target Actor to Hidden the first time it is executed, then un-hide it the next time. 
-
-1.	Go to the **Content Browser** under **Content/StarterContent/Blueprints** and open the **Blueprint_CeilingLight** Blueprint. 
-
-	![](MacroHT23.png)
-
-1.	Click the **Event Graph** tab to jump to the Event Graph. 
-
-	![](OnTheEventGraph.png)
-
-1.	**Right-click** in the graph and add the **RotateMacro** and **HideMacro**. 
-
-	![](MacroHT24.png)
-
-1.	**Right-click** in the graph and add an **Event Tick** node and connect it to the in pin on the **RotateMacro**. 
-
-1.	Drag off the **Target** and add the **Get Reference to Self** node. 
-
-1.	Drag off the **Self** node and add the **Set Actor Rotation** node. 
-
-1.	Connect the **Rotation** from **RotateMacro** to the **New Rotation** on the **Set Actor Rotation** node. 
-
-	![](MacroHT25.png)
-
-	After connecting the remaining wires, you script should look like above. 
-
-1.	Off the out pin of the **Set Actor Rotation**, search for and add a **Delay** node set to 1 second. 
-
-1.	Off the **Completed** pin, add the **HideMacro** node and drag off **Target** and **Get Reference to Self**. 
-
-	![](MacroHT26.png)
-
-	The full script should look similar to above. 
-
-1.	**Left-click** and drag a selection box around the script then hit **Control+C** to copy it. 
-
-1.	**Compile** and **Save** then close the Blueprint. 
-
-1.	In the **Content Browser** under **Content/StarterContent/Blueprints**, open the **Blueprint_WallSconce** Blueprint. 
-
-1.	In its graph, press **Control+V** to paste the script. 
-
-1.	**Right-click** in the graph and add/connect a new **Event Tick** node. 
-
-	![](MacroHT26.png)
-
-	Your copied script should look similar to above. 
-
-1.	**Compile** and **Save** then close the Blueprint. 
-
-1.	Drag the **Blueprint_CeilingLight** and **Blueprint_WallSconce** into the level. 
-
-	![](MacroHT28.png)
-
-1.	Click the **Play** button from the Main Toolbar to play in the Editor. 
-
-	![](ExampleGIF.gif)(convert:false)
-
-	When playing in the editor, you should see that both Blueprint Actors continually rotate around and toggle between hidden/visible. By creating a **Macro Library** and **Macros** that provide the functionality we wanted our Actors to exhibit, we were able to quickly provide that functionality to multiple Actors without having to recreate the scripts for each one. Keep this in mind the next time you are scripting something that other Blueprints may use as creating a Macro or Macro Library may save you some time. 
-
-
-
-	
-
-
-
-
-
-
-
-
-
-
-
-
-
-
-
-
 
