#UE4doc #docs: Removed extraneous bolding as according to the instructional guide.

diff --git a/Engine/Documentation/Source/Gameplay/HowTo/CharacterMovement/Blueprints/BlendSpace/BlendSpaces.INT.udn b/Engine/Documentation/Source/Gameplay/HowTo/CharacterMovement/Blueprints/BlendSpace/BlendSpaces.INT.udn
--- a/Engine/Documentation/Source/Gameplay/HowTo/CharacterMovement/Blueprints/BlendSpace/BlendSpaces.INT.udn
+++ b/Engine/Documentation/Source/Gameplay/HowTo/CharacterMovement/Blueprints/BlendSpace/BlendSpaces.INT.udn
@@ -24,7 +24,7 @@ Description:In this step we are going to set up Blend Spaces for our movement st
 The next thing we will do to create movement for our character is to create **Blend Spaces** for each of our movement states (Walking, Jogging, and Crouch Walking). Blend Spaces are special assets that allow for blending of animations based on the values of two inputs. In our case, we will blend between forwards/backwards and left/right movement based on the character's movement Speed or Direction. 
 
 [REGION:note]
-For more information on **Blend Spaces**, refer to the [](Engine/Animation/Blendspaces) documentation. 
+For more information on Blend Spaces, refer to the [](Engine/Animation/Blendspaces) documentation. 
 [/REGION]
 
 The steps below will show you how to set up the movement Blend Spaces: 
@@ -33,7 +33,7 @@ The steps below will show you how to set up the movement Blend Spaces:
 
 	![](Blendspace_1.png)
 
-1.	In the **Pick Skeleton** window, select your Skeleton (**HeroTPP_Skeleton**), then name the asset **Walk_BS**. 
+1.	In the **Pick Skeleton** window, select your Skeleton (**UE4_Mannequin_Skeleton**), then name the asset **Walk_BS**. 
 
 1.	Open the **Walk_BS** Blend Space, then under **Parameters**, set the **X Axis Label** to **Direction** and **Y Axis Label** to **Speed**. 
 
@@ -87,7 +87,7 @@ The steps below will show you how to set up the movement Blend Spaces:
 
 	![](Blendspace_8.png)
 
-	The 375 value refers to the speed of our Jog that we set inside the **MyCharacter** Blueprint. 
+	The 375 value refers to the speed of our Jog that we set inside the MyCharacter Blueprint. 
 
 1.	The grid will update to use the max speed value, however the points on the grid remain in their previous position. 
 
