Jira UEDOC-1816
Updated GIF to Youtube Video, minor content fixes as well.

diff --git a/Engine/Documentation/Source/Engine/Blueprints/BP_HowTo/Casting/Casting.INT.udn b/Engine/Documentation/Source/Engine/Blueprints/BP_HowTo/Casting/Casting.INT.udn
--- a/Engine/Documentation/Source/Engine/Blueprints/BP_HowTo/Casting/Casting.INT.udn
+++ b/Engine/Documentation/Source/Engine/Blueprints/BP_HowTo/Casting/Casting.INT.udn
@@ -1,14 +1,12 @@
 Availability: Public
 Title:Blueprint Casting
 Crumbs: %ROOT%, Engine, Engine\Blueprints, Engine\Blueprints\BP_HowTo
-Description:How Blueprints can work together and pass information between each other through Blueprint Casting. 
+Description:Using Blueprint Casting to communicate from one Blueprint to multiple Blueprints. 
 Related: Engine/Blueprints/UserGuide/CastNodes
 Related: Engine/Blueprints/UserGuide/BlueprintCommsUsage
 Related: Engine/Blueprints/UserGuide
 
-This page covers how to directly communicate between two Blueprints and pass information between them. 
-
-For this example we will allow the player to pass through a box and communicate with and activate a particle effect. 
+This page covers how to use Blueprint Casting to communicate with different Blueprints and share information. In this example we allow the player to shoot an object, making it smaller each time it is hit. Each time the object is hit, it also notifies the Character Blueprint and increases the player's movement speed. 
 
 [REGION:note]
 For this guide we are using the **Blueprint First Person** template with **Starter Content** enabled. 
@@ -18,14 +16,14 @@ For this guide we are using the **Blueprint First Person** template with **Start
 
 	![](Casting1.png)
 
+1.  At the confirmation window that appears, click **Ok** to use the default name and create the Blueprint. 
+
 1.  In the **Content Browser** under **Content/FirstPersonBP/Blueprints**, open the **FirstPeronProjectile** Blueprint. 
 
 	![](Casting2.png)
 
 	This is the Blueprint for the projectile that is launched when the player fires the weapon during gameplay. 
 
-1.  At the confirmation window that appears, click **Ok** to use the default name and create the Blueprint. 
-
 1.  Locate the **Event Hit** node and drag off the **Other** pin then search for and add the **Cast To LargeCube_Blueprint** node. 
 
 	![](Casting3.png)(w:800)
@@ -84,9 +82,19 @@ For this guide we are using the **Blueprint First Person** template with **Start
 
 1.	Click **Compile** and **Save**, then the **Play** button to play in the editor. 
 
-	![](CastingGIF.gif)(convert:false)
-
-	Each time you shoot the Large Cube Blueprint, it will reduce in size and your character will move faster each time. 
+	[OBJECT:EmbeddedVideo]
+	[PARAMLITERAL:width]
+	640
+	[/PARAMLITERAL]
+	[PARAMLITERAL:height]
+	360
+	[/PARAMLITERAL]
+	[PARAMLITERAL:videoid]
+	u8aJQ3PBXNk
+	[/PARAMLITERAL]
+	[/OBJECT]
+
+	Each time you shoot the Large Cube Blueprint, it will reduce in size and your character will move faster. 
 
 	For more information on **Casting** and usage cases, refer to the **Related Topics** section below. 
 	
