UEDOC-1818
Updated GIF to Youtube video, minor content fix.

diff --git a/Engine/Documentation/Source/Engine/Blueprints/BP_HowTo/Debugging/Debugging.INT.udn b/Engine/Documentation/Source/Engine/Blueprints/BP_HowTo/Debugging/Debugging.INT.udn
--- a/Engine/Documentation/Source/Engine/Blueprints/BP_HowTo/Debugging/Debugging.INT.udn
+++ b/Engine/Documentation/Source/Engine/Blueprints/BP_HowTo/Debugging/Debugging.INT.udn
@@ -32,11 +32,21 @@ Once you have an object selected, you can jump to that object in your level by c
 
 When you play in the editor, with your Blueprint up in another window, you should see the pulsating "Active Wires" as your script executes. 
 
-![](Example.gif)(convert:false)
+[OBJECT:EmbeddedVideo]
+[PARAMLITERAL:width]
+640
+[/PARAMLITERAL]
+[PARAMLITERAL:height]
+360
+[/PARAMLITERAL]
+[PARAMLITERAL:videoid]
+UtS3wgltYgg
+[/PARAMLITERAL]
+[/OBJECT]
 
 In the example above, there are two fire effect Blueprints in the level each set to different times to deactivate themselves. On **Event Begin Play** (when the game is launched), a **Delay** node is used to **Deactivate** the fire and its associated sound effect after a variable specified time. 
 
-You can see in the right window that the script executes when the game is launched (as the pulsating wire leaves the red Event Begin Play node and enters the Delay node). After the specified time, the pulse is shown entering the third node which is used for deactivating one of the fire effects. 
+You can see in the left window that the script executes when the game is launched (as the pulsating wire leaves the red Event Begin Play node and enters the Delay node). After the specified time, the pulse is shown entering the third node which is used for deactivating one of the fire effects. 
 
 ## Breakpoints
 
