#UEDoc #docs - Updated doc images, spelling, and bold lettering.

diff --git a/Engine/Documentation/Source/Engine/Physics/Vehicles/VehicleUserGuide/VehicleIntro.INT.udn b/Engine/Documentation/Source/Engine/Physics/Vehicles/VehicleUserGuide/VehicleIntro.INT.udn
--- a/Engine/Documentation/Source/Engine/Physics/Vehicles/VehicleUserGuide/VehicleIntro.INT.udn
+++ b/Engine/Documentation/Source/Engine/Physics/Vehicles/VehicleUserGuide/VehicleIntro.INT.udn
@@ -26,35 +26,38 @@ Navigation:topic
 
 [TOC(start:2 end:3)]
 
-## Overview
 
-This document will cover the creation of a vehicle using the **Wheeled Vehicle** Blueprint class, in a blank project. The process will take you through the creation of all the necessary in-editor items, but assumes you have a **Skeletal Mesh** and a **Physics Asset** for your vehicle ready to go.
 
-If you need information on getting a vehicle **Skeletal Mesh** and **Physics Asset** into the engine, please see the [**Vehicle Content Guide**](Engine/Physics/Vehicles/VehcileContentCreation). Alternatively, you can work out of the **Vehicle Template** or **Vehicle Game** to have access to a prebuilt **Skeletal Mesh** and **Physics Asset**
+This document will cover the creation of a vehicle using the **Wheeled Vehicle** Blueprint class, in a blank project. The process will take you through the creation of all the 
+necessary in-editor items, but assumes you have a **Skeletal Mesh** and a **Physics Asset** for your vehicle ready to go.
+
+If you need information on getting a vehicle Skeletal Mesh and Physics Asset into the engine, please see the [**Vehicle Content Guide**](Engine/Physics/Vehicles/VehcileContentCreation). 
+Alternatively, you can work out of the **Vehicle Template** or **Vehicle Game** to have access to a prebuilt Skeletal Mesh and Physics Asset.
 
 
 ## Content
 
 A vehicle in Unreal Engine 4 is composed of a number of assets:
 
-* A **Skeletal Mesh**.
-* A **Physics Asset**.
-* An **Animation Blueprint**.
-* A Vehicle Blueprint.
+* A Skeletal Mesh
+* A Physics Asset
+* An Animation Blueprint
+* A Vehicle Blueprint
 * One or more Wheel Blueprints.
-* A **TireType Data Asset**.
+* A TireType Data Asset
 
 [REGION:raw]
 ![](overview.png)(convert:false)
 [/REGION]
 
 [REGION:note]
-This document will cover the Blueprint setup required to make a functioning vehicle in Unreal Engine 4, and assumes you already have a vehicle asset to work with. See the [**Vehicle Content Guide**](Engine/Physics/Vehicles/VehcileContentCreation) for more information on setting up the **Skeletal Mesh** and **Physics Asset**.
+This document will cover the Blueprint setup required to make a functioning vehicle in Unreal Engine 4, and assumes you already have a vehicle asset to work with. See the 
+[**Vehicle Content Guide**](Engine/Physics/Vehicles/VehcileContentCreation) for more information on setting up the Skeletal Mesh and Physics Asset.
 [/REGION]
 
 ## Creating the Vehicle Blueprint
 
-To create a new **Vehicle Blueprint**: in the **Content Browser** click the **New** button -> **Blueprint** -> search for "vehicle" and click **WheeledVehicle** from the list.
+To create a new **Vehicle Blueprint**: in the Content Browser **click** **Add New** button -> **Blueprint** -> search for "vehicle" and **click** **WheeledVehicle** from the list.
 
 [REGION:raw]
 ![](newVehicle.png)(convert:false)
@@ -62,32 +65,36 @@ To create a new **Vehicle Blueprint**: in the **Content Browser** click the **Ne
 
 ## Creating the Wheel Blueprints
 
-To create a new **Wheel Blueprint**: in the **Content Browser** click the **New** button -> **Blueprint** -> search for "wheel" and click **VehicleWheel** from the list. Repeat this at least once so you have a front and rear wheel type.
+To create a new **Wheel Blueprint**: in the Content Browser **click** **Add New** button -> **Blueprint** -> search for "wheel" and **click** **VehicleWheel** from the list. Repeat 
+this at least once so you have a front and rear wheel type.
 
 [REGION:raw]
 ![](newWheels.png)(convert:false)
 [/REGION]
 
-In most cases, you will have at least two wheel types: A wheel that is affected by steering, and one that is not. Also you can set differing radii, mass, width, handbrake effect, suspension, and many other properties to give your vehicle the handling you desire.
+In most cases, you will have at least two wheel types: A wheel that is affected by steering, and one that is not. Also you can set differing radii, mass, width, handbrake effect, 
+suspension, and many other properties to give your vehicle the handling you desire.
 
 
 ## Creating the TireType Data Asset
 
-To create a new **TireType Data Asset**: in the **Content Browser** click the **New** button -> **Misc** -> **Data Asset** -> choose **TireType** -> click the **Select** button.
+To create a new **TireType Data Asset**: in the Content Browser **click** **Add New** button -> **Misc** -> **Data Asset** -> choose **TireType** -> **click** **Select** button.
 
 [REGION:raw]
 ![](newDataAsset.png)(convert:false)
 [/REGION]
 
 [EXCERPT:frictionScale]
-The **TireType Data Asset** only has a single value: **Friction Scale**. This value not only affects the raw friction of the wheel, but also scales the values for how difficult (or easy) it is for a wheel to slide while in a hard turn.
+The **TireType Data Asset** only has a single value: **Friction Scale**. This value not only affects the raw friction of the wheel, but also scales the values for how difficult 
+(or easy) it is for a wheel to slide while in a hard turn.
 [/EXCERPT:frictionScale]
 
 There is a property slot on the **VehicleWheel** Blueprint type for the **TireType Data Asset**.
 
 ## Creating the Animation Blueprint
 
-To create a new **Animation Blueprint**: in the **Content Browser** click the **New** button -> **Animation** -> **Animation Blueprint** -> choose the **Skeleton** from the list that is for your vehicle  -> click the **Ok** button.
+To create a new **Animation Blueprint**: in the Content Browser **click** **Add New** button -> **Animation** -> **Animation Blueprint** -> choose the **Skeleton** from the list 
+that is for your vehicle  -> click the **Ok** button.
 
 [REGION:raw]
 ![](newAnimBP.png)(convert:false)
@@ -96,7 +103,8 @@ To create a new **Animation Blueprint**: in the **Content Browser** click the **
 
 ## Editing the Animation Blueprint
 
-**Double-click** the **Animation Blueprint** from the **Content Browser** to edit the vehicle Animation Blueprint in Persona. The **Anim Graph** for a vehicle has been greatly simplified for the 4.2 release:
+**Double-click** the **Animation Blueprint** from the Content Browser to edit the vehicle Animation Blueprint in Persona. The **Anim Graph** for a vehicle has been greatly 
+simplified for the 4.2 release:
 
 1. **Right-click** in the **Anim Graph**, from the context menu search for the **Mesh Space Ref Pose** node, and select it to create it.
 
@@ -124,17 +132,21 @@ That is it! The **Component to Local** conversion node will be created automatic
 
 ![](wheelHandler02.png)
 
-The **Wheel Handler** node handles all the animation needs of your wheels: spinning, steering, handbrake, and suspension. There is no additional setup, it just grabs all the necessary information from the wheels (How fast is it spinning? Can it steer? Is it affected by the Handbrake? What are the suspension settings for this wheel?) and translates that to animation on the bone the wheel is associated with.
+The **Wheel Handler** node handles all the animation needs of your wheels: spinning, steering, handbrake, and suspension. There is no additional setup, it just grabs all the 
+necessary information from the wheels (How fast is it spinning? Can it steer? Is it affected by the Handbrake? What are the suspension settings for this wheel?) and translates that 
+to animation on the bone the wheel is associated with.
 
-Now, if you have additional struts and other suspension needs, you will need additional nodes in your Anim Graph to handle the joints that affect those polygons. The Vehicle from Vehicle Game has a few extra joints that control the axel connections to the wheels. These are driven by simple **Look At** nodes:
+Now, if you have additional struts and other suspension needs, you will need additional nodes in your Anim Graph to handle the joints that affect those polygons. The Vehicle from 
+Vehicle Game has a few extra joints that control the axel connections to the wheels. These are driven by simple **Look At** nodes:
 
 ![](fullVehGameAnimBP.png)
 
-The additional nodes in the Vehicle Game's Vehicle Anim Graph simply cause the strut joints point at the wheel joints. Given the wheel joints will be driven by the Wheel Handler Node, and the **Look At** nodes will make sure the suspension stays attached to the wheels.
+The additional nodes in the Vehicle Game's Vehicle Anim Graph simply cause the strut joints point at the wheel joints. Given the wheel joints will be driven by the Wheel Handler 
+Node, and the **Look At** nodes will make sure the suspension stays attached to the wheels.
 
 ## Editing the Tires Data Asset
 
-**Double-click** the **TireType Data Asset** from the **Content Browser**.
+**Double-click** the **TireType Data Asset** from the Content Browser.
 
 ![](tires02.png)
 
@@ -143,7 +155,7 @@ The additional nodes in the Vehicle Game's Vehicle Anim Graph simply cause the s
 
 ## Editing the Wheels Blueprints
 
-**Double-click** the **Wheel Blueprint** from the **Content Browser** to bring up the **Blueprint Editor**.
+**Double-click** the **Wheel Blueprint** from the Content Browser to bring up the Blueprint Editor.
 
 
 1. Switch the Blueprint Editor to Defaults Editing mode:
@@ -168,19 +180,20 @@ The additional nodes in the Vehicle Game's Vehicle Anim Graph simply cause the s
 	[/REGION]
 
 
-These are the 5 things you will need to change for each wheel. The rest of the properties change the way your vehicle will perform, and can be tweaked as you start testing your vehicle.
+These are the 5 things you will need to change for each wheel. The rest of the properties change the way your vehicle will perform, and can be tweaked as you start testing your 
+vehicle.
 
 ## Editing the Vehicle Blueprint
 
 This is really where the vehicle starts to come together. If your art is ready, and you have setup your wheels, then you will be able to fashion a working vehicle at this point.
 
-**Double-click** the **Vehicle Blueprint** from the **Content Browser** to bring up the **Blueprint Editor**.
+**Double-click** the **Vehicle Blueprint** from the Content Browser to bring up the Blueprint Editor.
 
 1. In **Components Mode**, on the **Skeletal Mesh Component** called "Mesh," set the **Skeletal Mesh** property to your vehicle's Skeletal Mesh asset.
 
 	![](componentsMode.png)
 
-1. Set the **Anim Blueprint Generated Class** to your vehicles Animation Blueprint.
+1. **Set** the **Anim Blueprint Generated Class** to your vehicles Animation Blueprint.
 
 	[REGION:raw]
 	![](AnimBPandSkelMesh.png)
@@ -192,7 +205,7 @@ This is really where the vehicle starts to come together. If your art is ready,
 	![](addCamera.png)
 	[/REGION]
 
-1. Position this Camera Component where you like. In the image below, it is position behind the vehicle.
+1. **Position** this Camera Component where you like. In the image below, it is positioned behind the vehicle.
 
 	![](cameraPlacement.png)
 
@@ -200,13 +213,13 @@ This is really where the vehicle starts to come together. If your art is ready,
 
 	![](cameraControllerOff.png)
 
-1. In **Defaults Mode**, in the **Vehicle Setup** category, expand "Wheel Setups."
+1. In **Defaults Mode**, in the **Vehicle Setup** category, **expand** "Wheel Setups."
 
 	![](defaultsMode.png)
 
 1. For each wheel (0-3):
-	1. Set the **Wheel Class** to one of the wheel Blueprints you created.
-	1. Set the **Bone Name** to the name of the joint that should be controlled by the wheel.
+	1. **Set** the **Wheel Class** to one of the wheel Blueprints you created.
+	1. **Set** the **Bone Name** to the name of the joint that should be controlled by the wheel.
 
 		[REGION:raw]
 		![](wheelsSetup.png)
@@ -234,25 +247,25 @@ There are a few ways to setup the vehicle for testing, and some steps will be de
 	[/REGION]
 
 	1. **Project Settings** -> **Input**
-	1. Add an **Axis Mapping**, by clicking the "**+**" sign next to the **Axis Mappings** property
-	1. Expand the Axis Mappings
-	1. Rename the Axis Mapping from "None" to **Forward**
-	1. Click the "**+**" next to the **Forward** Axis Mapping
-	1. Expand the **Forward** Axis Mapping
-	1. Set the first "None" to "**W**"
-	1. Set the second "None" to "**S**"
-	1. Set the Scale of the "**S**" mapping to "-1.0"
-	1. Add another **Axis Mapping**, by clicking the "**+**" sign next to the **Axis Mappings** property
-	1. Rename the new Axis Mapping from "None" to **Right** 
-	1. Expand the **Right** Axis Mapping
-	1. Set the first "None" to "**D**"
-	1. Set the second "None" to "**A**"
-	1. Set the Scale of the "**A**" mapping to "-1.0"
-	1. Then add an **Action Mapping**, by clicking the "**+**" sign next to the Action Mappings property, for:
-	1. Expand the **Action Mappings** property
-	1. Rename the new Action Mapping from "None" to "**Handbrake**"
-	1. Expand the **Handbrake** Action Mapping
-	1. Set the "None" to "**Space Bar**"
+	1. **Add** an **Axis Mapping**, by clicking the "**+**" sign next to the Axis Mappings property
+	1. **Expand** the Axis Mappings
+	1. **Rename** the Axis Mapping from "None" to **Forward**
+	1. **Click** the "**+**" next to the **Forward** Axis Mapping
+	1. **Expand** the **Forward** Axis Mapping
+	1. **Set** the first "None" to "**W**"
+	1. **Set** the second "None" to "**S**"
+	1. **Set** the Scale of the "**S**" mapping to "-1.0"
+	1. **Add** another **Axis Mapping**, by clicking the "**+**" sign next to the Axis Mappings property
+	1. **Rename** the new Axis Mapping from "None" to **Right** 
+	1. **Expand** the **Right** Axis Mapping
+	1. **Set** the first "None" to "**D**"
+	1. **Set** the second "None" to "**A**"
+	1. **Set** the Scale of the "**A**" mapping to "-1.0"
+	1. Then **add** an **Action Mapping**, by clicking the "**+**" sign next to the Action Mappings property, for:
+	1. **Expand** the **Action Mappings** property
+	1. **Rename** the new Action Mapping from "None" to "**Handbrake**"
+	1. **Expand** the **Handbrake** Action Mapping
+	1. **Set** the "None" to "**Space Bar**"
 1. New Game Mode
 	1. Create new Game Mode Blueprint.
 
@@ -265,7 +278,7 @@ There are a few ways to setup the vehicle for testing, and some steps will be de
 
 		![](defaultsMode.png)
 
-	1. **Game Mode** Category -> **Default Pawn Class** -> change "None" to your Vehicle Blueprint.
+	1. **Game Mode** Category -> **Default Pawn Class** -> **change** "None" to your **Vehicle Blueprint**.
 
 		[REGION:raw]
 		![](defaultPawnClass.png)
@@ -278,7 +291,7 @@ There are a few ways to setup the vehicle for testing, and some steps will be de
 	[/REGION]
 
 1. In your vehicle's Blueprint, you can then call on the Axis and Action events you created in the **Project Settings Input** category.
-	1. Edit your Vehicle Blueprint, and switch to **Graph Mode**
+	1. **Edit** your Vehicle Blueprint, and **switch** to **Graph Mode**
 
 		![](graphMode.png)
 
@@ -295,7 +308,8 @@ There are a few ways to setup the vehicle for testing, and some steps will be de
 		![](handbrakeBP.png)
 
 1. Player Start
-	* If you do not have a **Player Start** Actor in the world, you can either add one to your map using the **Modes Panel** -> **Placement Mode** -> **Basic** -> then clicking and dragging a **Player Start** into the world.
+	* If you do not have a **Player Start** Actor in the world, you can either add one to your map using the **Modes Panel** -> **Placement Mode** -> **Basic** -> then **clicking and 
+	dragging** a **Player Start** into the world.
 
 	[REGION:raw]
 	![](playerStart.png)
@@ -308,7 +322,8 @@ There are a few ways to setup the vehicle for testing, and some steps will be de
 	![](Play.png)
 	[/REGION]
 
-	Given the Axis and Action Bindings you setup above, "**W**" will increase your throttle, "**S**" will apply the brakes, "**A**" will turn left, "**D**" will turn right, and "**Space Bar**" will apply the Handbrake.
+	Given the Axis and Action Bindings you setup above, "**W**" will increase your throttle, "**S**" will apply the brakes, "**A**" will turn left, "**D**" will turn right, and 
+	"**Space Bar**" will apply the Handbrake.
 
 
 
