#UE4doc #docs: Removed extraneous bolding as according to the instructional guide. Altered images to match the instructional guide: green -> yellow. Altered images which contained the blue character to make him match the new, silver character.

diff --git a/Engine/Documentation/Source/Gameplay/HowTo/CharacterMovement/Blueprints/AnimBlueprint_Walk/Walk_ABP.INT.udn b/Engine/Documentation/Source/Gameplay/HowTo/CharacterMovement/Blueprints/AnimBlueprint_Walk/Walk_ABP.INT.udn
--- a/Engine/Documentation/Source/Gameplay/HowTo/CharacterMovement/Blueprints/AnimBlueprint_Walk/Walk_ABP.INT.udn
+++ b/Engine/Documentation/Source/Gameplay/HowTo/CharacterMovement/Blueprints/AnimBlueprint_Walk/Walk_ABP.INT.udn
@@ -24,7 +24,7 @@ Description:Here in this step, we'll create the logic for Idle and Walking state
 In this section, we will set up the **Animation Blueprint** that will determine what animations to play based on the actions that are taken in game by the player. We will then set up a **State Machine** and create an **Idle** and **Walk** state and set up the transition between each state. 
 
 [REGION:note]
-For more information on **Animation Blueprints**, refer to the [](Engine/Animation/AnimBlueprints) documentation. 
+For more information on Animation Blueprints, refer to the [](Engine/Animation/AnimBlueprints) documentation. 
 [/REGION]
 
 The steps below will show you how to set up the Animation Blueprint: 
