#UEDoc #docs - doc bold lettering

diff --git a/Engine/Documentation/Source/Programming/UnrealArchitecture/Reference/Classes/GameplayClasses.INT.udn b/Engine/Documentation/Source/Programming/UnrealArchitecture/Reference/Classes/GameplayClasses.INT.udn
--- a/Engine/Documentation/Source/Programming/UnrealArchitecture/Reference/Classes/GameplayClasses.INT.udn
+++ b/Engine/Documentation/Source/Programming/UnrealArchitecture/Reference/Classes/GameplayClasses.INT.udn
@@ -23,7 +23,7 @@ by looking at the first letter, or prefix. The prefixes for gameplay classes are
 
 ## Adding Classes
 
-The [](Programming\Development\ManagingGameCode\CppClassWizard) sets up the header file and source file you need for your new class, and also updates your game module accordingly. 
+The **[](Programming\Development\ManagingGameCode\CppClassWizard)** sets up the header file and source file you need for your new class, and also updates your game module accordingly. 
 The header file and source file automatically include the class declaration and class constructor, as well as Unreal Engine-specific code like the `UCLASS()` macro.
 
 ## Class Headers
@@ -40,7 +40,7 @@ At the top of each gameplay class header file, the generated header file for the
 ### Class Declaration
 
 The class declaration defines the name of the class, what class it inherits from and, thus, any functions and variables it inherits, and other engine and editor specific behavior 
-that may be desired via [class specifiers](#ClassSpecifiers) and metadata. The syntax for declaring a class is as follows:
+that may be desired via **[class specifiers](#ClassSpecifiers)** and metadata. The syntax for declaring a class is as follows:
 
 	UCLASS([specifier, specifier, ...], [meta(key=value, key=value, ...)])
 	class ClassName : ParentName
