TypoCorrection

diff --git a/Engine/Documentation/Source/Programming/Introduction/IntroToProgramming.INT.udn b/Engine/Documentation/Source/Programming/Introduction/IntroToProgramming.INT.udn
--- a/Engine/Documentation/Source/Programming/Introduction/IntroToProgramming.INT.udn
+++ b/Engine/Documentation/Source/Programming/Introduction/IntroToProgramming.INT.udn
@@ -1,4 +1,4 @@
-Title:Introduction to C++ Programming in UE4
+﻿Title:Introduction to C++ Programming in UE4
 Crumbs:
 Description:Introductory guide for C++ programmers new to Unreal Engine
 Availability:Public
@@ -14,27 +14,27 @@ Before we go on, it's really important that you're already familiar with C++ or
 
 If you're coming in with no programming experience at all, we've got you covered also!  Check out our [Blueprint Visual Scripting guide](Engine/Blueprints) and you'll be on your way.  You can create entire games using Blueprint scripting!
 
-It is possible to write "plain old C++ code" in Unreal Engine, but you'll be most successful after reading through this guide and learning the basics about the Unreal programming model.  We’ll talk more about that we go along.
+It is possible to write "plain old C++ code" in Unreal Engine, but you'll be most successful after reading through this guide and learning the basics about the Unreal programming model.  We'll talk more about that we go along.
 
 ## C++ and Blueprints
 
-Unreal Engine provides two methods, C++ and Blueprints Visual Scripting, to create new gameplay elements. Using C++, programmers add the base gameplay systems that designers can then build upon or with to create the custom gameplay for a level or the game. In these cases, the C++ programmer works in their favorite IDE (usually Microsoft Visual Studio, or Apple’s Xcode) and the designer works in the Unreal Editor’s Blueprint Editor.
+Unreal Engine provides two methods, C++ and Blueprints Visual Scripting, to create new gameplay elements. Using C++, programmers add the base gameplay systems that designers can then build upon or with to create the custom gameplay for a level or the game. In these cases, the C++ programmer works in their favorite IDE (usually Microsoft Visual Studio, or Apple's Xcode) and the designer works in the Unreal Editor's Blueprint Editor.
 
 The gameplay API and framework classes are available to both of these systems, which can be used separately, but show their true power when used in conjunction to compliment each other. What does that really mean, though? It means that the engine works best when programmers are creating gameplay building blocks in C++ and designers take those blocks and make interesting gameplay.
 
-With that said, let’s take a look at a typical workflow for the C++ programmer that is creating building blocks for the designer. In this case, we’re going to create a class that is later extended via Blueprints by a designer or programmer. In this class, we’re going to create some properties that the designer can set and we’re going to derive new values from those properties. The whole process is very easy to do using the tools and C++ macros we provide for you.
+With that said, let's take a look at a typical workflow for the C++ programmer that is creating building blocks for the designer. In this case, we're going to create a class that is later extended via Blueprints by a designer or programmer. In this class, we're going to create some properties that the designer can set and we're going to derive new values from those properties. The whole process is very easy to do using the tools and C++ macros we provide for you.
 
 ### Class Wizard
 
-First thing we’re going to do is use the class wizard within the Unreal Editor to generate the basic C++ class that will be extended by Blueprints later. The image below shows the wizard’s first step where we are creating a new Actor.
+First thing we're going to do is use the class wizard within the Unreal Editor to generate the basic C++ class that will be extended by Blueprints later. The image below shows the wizard's first step where we are creating a new Actor.
 
 ![image alt text](image_1.png)
 
-The second step in the process tells the wizard the name of the class you want generated. Here’s the second step with the default name used.
+The second step in the process tells the wizard the name of the class you want generated. Here's the second step with the default name used.
 
 ![image alt text](image_2.png)
 
-Once you choose to create the class, the wizard will generate the files and open your development environment so that you can start editing it. Here’s the class definition that is generated for you. For more information on the Class Wizard, follow this [link.](Programming/Development/ManagingGameCode/CppClassWizard)
+Once you choose to create the class, the wizard will generate the files and open your development environment so that you can start editing it. Here's the class definition that is generated for you. For more information on the Class Wizard, follow this [link.](Programming/Development/ManagingGameCode/CppClassWizard)
 
 	#include "GameFramework/Actor.h"
 	#include "MyActor.generated.h"
@@ -54,7 +54,7 @@ Once you choose to create the class, the wizard will generate the files and open
 		virtual void Tick( float DeltaSeconds ) override;
 	};
 
-The class wizard generates your class with BeginPlay() and Tick() specified as overloads. BeginPlay() is an event that lets you know the Actor has entered the game in a playable state. This is a good place to initiate gameplay logic for your class. Tick() is called once per frame with the amount of elapsed time since the last call passed in. There you can do any recurring logic. However if you don’t need that functionality, it is best to remove it to save yourself a small amount of performance. If you remove it, make sure to remove the line in the constructor that indicated ticking should occur. The constructor below contains the line in question.
+The class wizard generates your class with BeginPlay() and Tick() specified as overloads. BeginPlay() is an event that lets you know the Actor has entered the game in a playable state. This is a good place to initiate gameplay logic for your class. Tick() is called once per frame with the amount of elapsed time since the last call passed in. There you can do any recurring logic. However if you don't need that functionality, it is best to remove it to save yourself a small amount of performance. If you remove it, make sure to remove the line in the constructor that indicated ticking should occur. The constructor below contains the line in question.
 
 	AMyActor::AMyActor()
 
@@ -68,7 +68,7 @@ The class wizard generates your class with BeginPlay() and Tick() specified as o
 
 ### Making a property show up in the editor
 
-We have our class created, so now let’s create some properties that can be set by designers in the Unreal Editor. Exposing a property to the editor is quite easy using our special macro, UPROPERTY(). All you have to do is use the UPROPERTY(EditAnywhere) macro before your property declaration as seen in the class below.
+We have our class created, so now let's create some properties that can be set by designers in the Unreal Editor. Exposing a property to the editor is quite easy using our special macro, UPROPERTY(). All you have to do is use the UPROPERTY(EditAnywhere) macro before your property declaration as seen in the class below.
 
 	UCLASS()
 	class AMyActor : public AActor
@@ -81,21 +81,21 @@ We have our class created, so now let’s create some properties that can be set
 		...
 	};
 
-That’s all you need to do to be able to edit that value in the editor. There are more ways to control how and where it is edited. This is done by passing more information into the UPROPERTY() macro. For instance, if you want the TotalDamage property to appear in a section with related properties, you can use the categorization feature. The property declaration below shows this.
+That's all you need to do to be able to edit that value in the editor. There are more ways to control how and where it is edited. This is done by passing more information into the UPROPERTY() macro. For instance, if you want the TotalDamage property to appear in a section with related properties, you can use the categorization feature. The property declaration below shows this.
 
 	UPROPERTY(EditAnywhere, Category="Damage")
 	int32 TotalDamage;
 
 When the user looks to edit this property, it now appears under the Damage heading along with any other properties that you have marked with this category name. This is a great way to place commonly used settings together for editing by designers.
 
-Now let’s expose that same property to Blueprints.
+Now let's expose that same property to Blueprints.
 
 	UPROPERTY(EditAnywhere, BlueprintReadWrite, Category="Damage")
 	int32 TotalDamage;
 
-As you can see, there is a Blueprint specific parameter to make a property available for reading and writing. There’s a separate option, BlueprintReadOnly, you can use if you want the property to be treated as *const* in Blueprints. There are quite a few options available for controlling how a property is exposed to the engine. To see more options, follow this [link.](Programming/UnrealArchitecture/Reference/Properties/Specifiers)
+As you can see, there is a Blueprint specific parameter to make a property available for reading and writing. There's a separate option, BlueprintReadOnly, you can use if you want the property to be treated as *const* in Blueprints. There are quite a few options available for controlling how a property is exposed to the engine. To see more options, follow this [link.](Programming/UnrealArchitecture/Reference/Properties/Specifiers)
 
-Before continuing to the section below, let’s add a couple of properties to this sample class. There’s already a property to control the total amount of damage this actor will deal out, but let’s take that further and make that damage happen over time. The code below adds one designer settable property and one that is visible to the designer but not changeable by them.
+Before continuing to the section below, let's add a couple of properties to this sample class. There's already a property to control the total amount of damage this actor will deal out, but let's take that further and make that damage happen over time. The code below adds one designer settable property and one that is visible to the designer but not changeable by them.
 
 	UCLASS()
 	class AMyActor : public AActor
@@ -114,7 +114,7 @@ Before continuing to the section below, let’s add a couple of properties to th
 		...
 	};
 
-DamageTimeInSeconds is a property the designer can modify. The DamagePerSecond property is a calculated value using the designer’s settings (see the next section). The VisibleAnywhere flag marks that property as viewable, but not editable in the Unreal Editor. The Transient flag means that it won’t be saved or loaded from disk; it is meant to be a derived, non-persistent value. The image below shows the properties as part of the class defaults.
+DamageTimeInSeconds is a property the designer can modify. The DamagePerSecond property is a calculated value using the designer's settings (see the next section). The VisibleAnywhere flag marks that property as viewable, but not editable in the Unreal Editor. The Transient flag means that it won't be saved or loaded from disk; it is meant to be a derived, non-persistent value. The image below shows the properties as part of the class defaults.
 
 ![image alt text](image_3.png)
 
@@ -138,7 +138,7 @@ Here is the same view of the properties after adding default values in the const
 
 ![image alt text](image_4.png)
 
-In order to support per instance designer set properties, values are also loaded from the instance data for a given object. This data is applied after the constructor. You can create default values based off of designer set values by hooking into the PostInitProperties() call chain. Here’s an example of that process where TotalDamage and DamageTimeInSeconds are designer specified values. Even though these are designer specified, you can still provide sensible default values for them, as we did in the example above. NOTE: if you don’t provide a default value for a property, the engine will automatically set that property to zero or nullptr in the case of pointer types.
+In order to support per instance designer set properties, values are also loaded from the instance data for a given object. This data is applied after the constructor. You can create default values based off of designer set values by hooking into the PostInitProperties() call chain. Here's an example of that process where TotalDamage and DamageTimeInSeconds are designer specified values. Even though these are designer specified, you can still provide sensible default values for them, as we did in the example above. NOTE: if you don't provide a default value for a property, the engine will automatically set that property to zero or nullptr in the case of pointer types.
 
 	void AMyActor::PostInitProperties()
 	{
@@ -146,7 +146,7 @@ In order to support per instance designer set properties, values are also loaded
 		DamagePerSecond = TotalDamage / DamageTimeInSeconds;
 	}
 
-Here again is the same view of the properties after we’ve added the PostInitProperties() code that you see above.
+Here again is the same view of the properties after we've added the PostInitProperties() code that you see above.
 
 ![image alt text](image_5.png)
 
@@ -158,9 +158,9 @@ Here is a cool feature of Unreal that you might be surprised about if you are us
 
 	![image alt text](image_6.png)
 
-	(Note that if you’re attached with the debugger, you’ll need to detach first so that Visual Studio will allow you to Build.)
+	(Note that if you're attached with the debugger, you'll need to detach first so that Visual Studio will allow you to Build.)
 
-2. Or, simply click the **Compile **button on the editor’s main toolbar!
+2. Or, simply click the **Compile **button on the editor's main toolbar!
 
 	![image alt text](image_7.png)
 
@@ -168,9 +168,9 @@ You can use this feature in the sections below as we advance through the tutoria
 
 ### Extending a C++ Class via Blueprints
 
-So far, we’ve created a simple gameplay class with the C++ Class Wizard and added some properties for the designer to set. Let’s now take a look at how a designer can start creating unique classes from our humble beginnings here.
+So far, we've created a simple gameplay class with the C++ Class Wizard and added some properties for the designer to set. Let's now take a look at how a designer can start creating unique classes from our humble beginnings here.
 
-First thing we’re going to do is create a new Blueprint class from our AMyActor class. Notice in the image below that the name of the base class selected shows up as MyActor instead of AMyActor. This is intentional and hides the naming conventions used by our tools from the designer, making the name friendlier to them.
+First thing we're going to do is create a new Blueprint class from our AMyActor class. Notice in the image below that the name of the base class selected shows up as MyActor instead of AMyActor. This is intentional and hides the naming conventions used by our tools from the designer, making the name friendlier to them.
 
 ![image alt text](image_8.png)
 
@@ -178,11 +178,11 @@ Once you choose Select, a new, default named Blueprint class is created for you.
 
 ![image alt text](image_9.png)
 
-This is the first class that we’re going to customize with our designer hats on. First thing we’re going to do is change the default values for our damage properties. In this case, the designer changed the TotalDamage to 300 and the time it takes to deliver that damage to 2 seconds. This is how the properties now appear.
+This is the first class that we're going to customize with our designer hats on. First thing we're going to do is change the default values for our damage properties. In this case, the designer changed the TotalDamage to 300 and the time it takes to deliver that damage to 2 seconds. This is how the properties now appear.
 
 ![image alt text](image_10.png)
 
-Wait a second...our calculated value doesn’t match what we’d expect. It should be 150 but it is still at the default value of 200. The reason for this is that we are only calculating our damage per second value after the properties have been initialized from the loading process. Runtime changes in the Unreal Editor aren’t accounted for. There’s a simple solution to this problem because the engine notifies the target object when it has been changed in the editor. The code below shows the added hooks needed to calculate the derived value as it changes in the editor.
+Wait a second...our calculated value doesn't match what we'd expect. It should be 150 but it is still at the default value of 200. The reason for this is that we are only calculating our damage per second value after the properties have been initialized from the loading process. Runtime changes in the Unreal Editor aren't accounted for. There's a simple solution to this problem because the engine notifies the target object when it has been changed in the editor. The code below shows the added hooks needed to calculate the derived value as it changes in the editor.
 
 	void AMyActor::PostInitProperties()
 	{
@@ -205,13 +205,13 @@ Wait a second...our calculated value doesn’t match what we’d expect. It shou
 	}
 	#endif
 
-One thing to note is that the PostEditChangeProperty() method is inside an editor specific #ifdef. This is so that building your game only the code that you need for the game, removing any extra code that might increase the size of your executable unnecessarily. Now that we have that code compiled in, the DamagePerSecond value matches what we’d expect it to be as seen in the image below.
+One thing to note is that the PostEditChangeProperty() method is inside an editor specific #ifdef. This is so that building your game only the code that you need for the game, removing any extra code that might increase the size of your executable unnecessarily. Now that we have that code compiled in, the DamagePerSecond value matches what we'd expect it to be as seen in the image below.
 
 ![image alt text](image_11.png)
 
 ### Calling Functions across the C++ and Blueprint Boundary
 
-So far, we’ve shown how to expose properties to Blueprints, but there’s one last introductory topic that we should cover before you dive deeper into the engine. During the creation of the gameplay systems, designers will need to be able to call functions created by a C++ programmer as well as the gameplay programmer calling functions implemented in Blueprints from C++ code. Let’s start by first making the CalculateValues() function callable from Blueprints. Exposing a function to Blueprints is just as simple as exposing a property. It takes only one macro placed before the function declaration! The code snippet below show what is needed for this.
+So far, we've shown how to expose properties to Blueprints, but there's one last introductory topic that we should cover before you dive deeper into the engine. During the creation of the gameplay systems, designers will need to be able to call functions created by a C++ programmer as well as the gameplay programmer calling functions implemented in Blueprints from C++ code. Let's start by first making the CalculateValues() function callable from Blueprints. Exposing a function to Blueprints is just as simple as exposing a property. It takes only one macro placed before the function declaration! The code snippet below show what is needed for this.
 
 	UFUNCTION(BlueprintCallable, Category="Damage")
 	void CalculateValues();
@@ -226,34 +226,34 @@ As you can see, the function is selectable from the Damage category. The Bluepri
 
 This uses the same function that we added earlier to calculate our dependent property. Much of the engine is exposed to Blueprints via the UFUNCTION() macro, so that people can build games without writing C++ code. However, the best approach is to use C++ for building base gameplay systems and performance critical code with Blueprints used to customize behavior or create composite behaviors from C++ building blocks.
 
-Now that our designers can call our C++ code, let’s explore one more powerful way to cross the C++/Blueprint boundary. This approach allows C++ code to call functions that are defined in Blueprints. We often use the approach to notify the designer of an event that they can respond to as they see fit. Often that includes the spawning of effects or other visual impact, such as hiding or unhiding an actor. The code snippet below shows a function that is implemented by Blueprints.
+Now that our designers can call our C++ code, let's explore one more powerful way to cross the C++/Blueprint boundary. This approach allows C++ code to call functions that are defined in Blueprints. We often use the approach to notify the designer of an event that they can respond to as they see fit. Often that includes the spawning of effects or other visual impact, such as hiding or unhiding an actor. The code snippet below shows a function that is implemented by Blueprints.
 
 	UFUNCTION(BlueprintImplementableEvent, Category="Damage")
 	void CalledFromCpp();
 
-This function is called like any other C++ function. Under the covers, the Unreal Engine generates a base C++ function implementation that understands how to call into the Blueprint VM. This is commonly referred to as a Thunk. If the Blueprint in question doesn’t provide a function body for this method, then the function behaves just like a C++ function with no body behaves: it does nothing. What if you want to provide a C++ default implementation while still allowing a Blueprint to override the method? The UFUNCTION() macro has an option for that too. The code snippet below shows the changes needed in the header to achieve this.
+This function is called like any other C++ function. Under the covers, the Unreal Engine generates a base C++ function implementation that understands how to call into the Blueprint VM. This is commonly referred to as a Thunk. If the Blueprint in question doesn't provide a function body for this method, then the function behaves just like a C++ function with no body behaves: it does nothing. What if you want to provide a C++ default implementation while still allowing a Blueprint to override the method? The UFUNCTION() macro has an option for that too. The code snippet below shows the changes needed in the header to achieve this.
 
 	UFUNCTION(BlueprintNativeEvent, Category="Damage")
 	void CalledFromCpp();
 
-This version still generates the thunking method to call into the Blueprint VM. So how do you provide the default implementation? The tools also generate a new function declaration that looks like <function name>_Implementation(). You must provide this version of the function or your project will fail to link. Here’s the implementation code for the declaration above.
+This version still generates the thunking method to call into the Blueprint VM. So how do you provide the default implementation? The tools also generate a new function declaration that looks like <function name>_Implementation(). You must provide this version of the function or your project will fail to link. Here's the implementation code for the declaration above.
 
 	void AMyActor::CalledFromCpp_Implementation()
 	{
 		// Do something cool here
 	}
 
-Now this version of the function is called when the Blueprint in question doesn’t override the method. One thing to note, is that in future versions of the build tools the auto generated _Implementation() declaration will go away and you’ll be expected to explicitly add that to the header. As of version 4.7, the auto generation of that declaration still occurs.
+Now this version of the function is called when the Blueprint in question doesn't override the method. One thing to note, is that in future versions of the build tools the auto generated _Implementation() declaration will go away and you'll be expected to explicitly add that to the header. As of version 4.7, the auto generation of that declaration still occurs.
 
 Now that we have walked through the common gameplay programmer workflow and methods to work with designers to build out gameplay features, it is time for you to choose your own adventure. You can either continue with this document to read more about how we use C++ in the engine or you can jump right into one of our samples that we include in the launcher to get a more hands on experience. 
 
 ## Diving Deeper
 
-I see you’re still with me on this adventure. Excellent. The next topics of discussion revolve around what our gameplay class hierarchy looks like. In this section, we’ll start with the base building blocks and talk through how they relate to each other. This is where we’ll look at how the Unreal Engine uses both inheritance and composition to build custom gameplay features.
+I see you're still with me on this adventure. Excellent. The next topics of discussion revolve around what our gameplay class hierarchy looks like. In this section, we'll start with the base building blocks and talk through how they relate to each other. This is where we'll look at how the Unreal Engine uses both inheritance and composition to build custom gameplay features.
 
 ### Gameplay Classes: Objects, Actors, and Components
 
-There are 4 main class types that you derive from for the majority of gameplay classes.  They are **UObject**, **AActor**, **UActorComponent**, and **UStruct**. Each of these building blocks are described in the following sections. Of course, you can create types that don’t derive from any of these classes, but they will not participate in the features that are built into the engine. Typical use of classes that are created outside of the **UObject** hierarchy are: integrating 3rd party libraries; wrapping of OS specific features; etc.
+There are 4 main class types that you derive from for the majority of gameplay classes.  They are **UObject**, **AActor**, **UActorComponent**, and **UStruct**. Each of these building blocks are described in the following sections. Of course, you can create types that don't derive from any of these classes, but they will not participate in the features that are built into the engine. Typical use of classes that are created outside of the **UObject** hierarchy are: integrating 3rd party libraries; wrapping of OS specific features; etc.
 
 #### Unreal Objects (UObject)
 
@@ -266,13 +266,13 @@ The base building block in the Unreal Engine is called **UObject**. This class,
 * Configurable values for properties
 * Networking support for properties and methods
 
-Each class that derives from **UObject** has a singleton **UClass** created for it that contains all of the meta data about the class instance. **UObject** and **UClass** together are at the root of everything that a gameplay object does during its lifetime. The best way to think of the difference between a **UClass** and a **UObject** is that the **UClass** *describes* what an instance of a **UObject** will look like, what properties are available for serialization, networking, etc. Most gameplay development doesn’t involve directly deriving from **UObjects**, but instead from **AActor** and **UActorComponent**. You don’t need to know the details of how **UClass**/**UObject** works in order to write gameplay code, but it is good to know that these systems exist.
+Each class that derives from **UObject** has a singleton **UClass** created for it that contains all of the meta data about the class instance. **UObject** and **UClass** together are at the root of everything that a gameplay object does during its lifetime. The best way to think of the difference between a **UClass** and a **UObject** is that the **UClass** *describes* what an instance of a **UObject** will look like, what properties are available for serialization, networking, etc. Most gameplay development doesn't involve directly deriving from **UObjects**, but instead from **AActor** and **UActorComponent**. You don't need to know the details of how **UClass**/**UObject** works in order to write gameplay code, but it is good to know that these systems exist.
 
 #### AActor
 
-An **AActor** is an object that is meant to be part of the gameplay experience. **AActors** are either placed in a level by a designer or created at runtime via gameplay systems. All objects that can be placed into a level extend from this class. Examples include **AStaticMeshActor**, **ACameraActor**, and **APointLight** actors. **AActor** derives from **UObject**, so enjoys all of the standard features listed in the previous section. **AActors** can be explicitly destroyed via gameplay code (C++ or Blueprints) or via the standard garbage collection mechanism when the owning level is unloaded from memory. **AActors** are responsible for the high-level behaviors of your game’s objects. **AActors** are also the base type that can be replicated during networking. During network replication, **AActors** can also distribute information for any **UActorComponents** owned by that **AActor** that require network support.
+An **AActor** is an object that is meant to be part of the gameplay experience. **AActors** are either placed in a level by a designer or created at runtime via gameplay systems. All objects that can be placed into a level extend from this class. Examples include **AStaticMeshActor**, **ACameraActor**, and **APointLight** actors. **AActor** derives from **UObject**, so enjoys all of the standard features listed in the previous section. **AActors** can be explicitly destroyed via gameplay code (C++ or Blueprints) or via the standard garbage collection mechanism when the owning level is unloaded from memory. **AActors** are responsible for the high-level behaviors of your game's objects. **AActors** are also the base type that can be replicated during networking. During network replication, **AActors** can also distribute information for any **UActorComponents** owned by that **AActor** that require network support.
 
-**AActors** have their own behaviors (specialization through inheritance), but they also act as containers for a hierarchy of **UActorComponents** (specialization through composition). This is done through the **AActor’s** RootComponent member, which contains a single **UActorComponent** that, in turn, can contain many others. Before an **AActor** can be placed in a level, that **AActor** must contain at least a **USceneComponent** which contains the translation, rotation, and scale for that **AActor**. 
+**AActors** have their own behaviors (specialization through inheritance), but they also act as containers for a hierarchy of **UActorComponents** (specialization through composition). This is done through the **AActor's** RootComponent member, which contains a single **UActorComponent** that, in turn, can contain many others. Before an **AActor** can be placed in a level, that **AActor** must contain at least a **USceneComponent** which contains the translation, rotation, and scale for that **AActor**. 
 
 **AActors** have a series of events that are called during the lifecycle of the **AActor**. The list below is a simplified set of the events that illustrate the lifecycle.
 
@@ -284,7 +284,7 @@ See [](Programming/UnrealArchitecture/Actors) for a more detailed discussion on
 
 ##### Runtime Lifecycle
 
-Just above we discussed a subset of an **AActor’s** lifecycle. For actors that are placed in a level, understanding the lifecycle is pretty easy to imagine: actors are loaded and come into existence and eventually the level is unloaded and the actors are destroyed. What is the process for runtime creation and destruction? Unreal Engine calls the creation of an **AActor** at runtime *spawning*. Spawning an actor is a bit more complicated than creating a normal object in the game. The reason is that an **AActor** needs to be registered with a variety of runtime systems in order to serve all of its needs. The initial location and rotation for the actor need to be set. Physics may need to know about it. The manager responsible for telling an actor to tick needs to know. And so on. Because of this, we have a method devoted to the spawning of an actor, **UWorld::SpawnActor()**. Once that actor is spawned successfully, its **BeginPlay()** method is called, followed by **Tick()** the next frame.
+Just above we discussed a subset of an **AActor's** lifecycle. For actors that are placed in a level, understanding the lifecycle is pretty easy to imagine: actors are loaded and come into existence and eventually the level is unloaded and the actors are destroyed. What is the process for runtime creation and destruction? Unreal Engine calls the creation of an **AActor** at runtime *spawning*. Spawning an actor is a bit more complicated than creating a normal object in the game. The reason is that an **AActor** needs to be registered with a variety of runtime systems in order to serve all of its needs. The initial location and rotation for the actor need to be set. Physics may need to know about it. The manager responsible for telling an actor to tick needs to know. And so on. Because of this, we have a method devoted to the spawning of an actor, **UWorld::SpawnActor()**. Once that actor is spawned successfully, its **BeginPlay()** method is called, followed by **Tick()** the next frame.
 
 Once an actor has lived out its lifetime, you can get rid of it by calling **Destroy()**. During that process **EndPlay()** will be called where you can do any custom logic for destruction. Another option for controlling how long an actor exists is to use the Lifespan member. You can set a timespan in the constructor of the object or with other code at runtime. Once that amount of time has expired, the actor will automatically have Destroy() called on it.
 
@@ -294,14 +294,14 @@ To learn more about spawning actors see the [](Programming/UnrealArchitecture/Ac
 
 **UActorComponents** have their own behaviors and are usually responsible for functionality that is shared across many types of **AActors**, e.g. providing visual meshes, particle effects, camera perspectives, and physics interactions. While **AActors** are often given high-level goals related to their overall roles your game, **UActorComponents** usually perform the individual tasks that support those higher-level objectives. Components can also be attached to other Components, or can be the root Component of an Actor. A Component can only attach to one parent Component or Actor, but it may have many child Components attached to itself. Picture a tree of Components. Child Components have location, rotation, and scaling relative to their parent Component or Actor.
 
-While there are many ways to use Actors and Components, one way to think of the Actors-Component relationship is that Actors might answer the question "what is this thing?" while Components might answer “what is this thing made of?”
+While there are many ways to use Actors and Components, one way to think of the Actors-Component relationship is that Actors might answer the question "what is this thing?" while Components might answer "what is this thing made of?"
 
-* RootComponent - this is the member of AActor that holds the top level Component in the AActor’s tree of Components
-* Ticking - Components are ticked as part of the owning AActor’s Tick()
+* RootComponent - this is the member of AActor that holds the top level Component in the AActor's tree of Components
+* Ticking - Components are ticked as part of the owning AActor's Tick()
 
 ##### Dissecting the First Person Character
 
-Over the last few sections we’ve done a lot of talking and not a lot of showing. In order to illustrate the relationship of an **AActor** and its **UActorComponents**, let’s dig into the Blueprint that is created when you generate a new project based off of the First Person Template. The image below is the Component tree for the FirstPersonCharacter Actor. The RootComponent is the CapsuleComponent. Attached to the CapsuleComponent is the ArrowComponent, the Mesh component, and the FirstPersonCameraComponent. The leaf most component is the Mesh1P component which is parented to the FirstPersonCameraComponent, meaning that the first person mesh is relative to the first person camera. 
+Over the last few sections we've done a lot of talking and not a lot of showing. In order to illustrate the relationship of an **AActor** and its **UActorComponents**, let's dig into the Blueprint that is created when you generate a new project based off of the First Person Template. The image below is the Component tree for the FirstPersonCharacter Actor. The RootComponent is the CapsuleComponent. Attached to the CapsuleComponent is the ArrowComponent, the Mesh component, and the FirstPersonCameraComponent. The leaf most component is the Mesh1P component which is parented to the FirstPersonCameraComponent, meaning that the first person mesh is relative to the first person camera. 
 
 ![image alt text](image_14.png)
 
@@ -313,19 +313,19 @@ This tree of components is attached to the one actor class. As you can see from
 
 #### UStruct
 
-To use a **UStruct**, you don’t have to extend from any particular class, you just have mark the struct with **USTRUCT()** and our build tools will do the base work for you. Unlike a **UObject**, **UStructs** are *not* garbage collected. If you create dynamic instances of them, you must manage their lifecycle yourself. **UStructs** are meant to be plain old data types that have the **UObject** reflection support for editing within the Unreal Editor, Blueprint manipulation, serialization, networking, etc.
+To use a **UStruct**, you don't have to extend from any particular class, you just have mark the struct with **USTRUCT()** and our build tools will do the base work for you. Unlike a **UObject**, **UStructs** are *not* garbage collected. If you create dynamic instances of them, you must manage their lifecycle yourself. **UStructs** are meant to be plain old data types that have the **UObject** reflection support for editing within the Unreal Editor, Blueprint manipulation, serialization, networking, etc.
 
-Now that we’ve talked about the basic hierarchy used in our gameplay class construction, it is time to choose your path again. You can read about our gameplay classes [here](Programming/UnrealArchitecture/Reference/Classes), head out to our samples in the launcher armed with more information, or continue digging deeper into our C++ features for building games.
+Now that we've talked about the basic hierarchy used in our gameplay class construction, it is time to choose your path again. You can read about our gameplay classes [here](Programming/UnrealArchitecture/Reference/Classes), head out to our samples in the launcher armed with more information, or continue digging deeper into our C++ features for building games.
 
 ## Diving Deeper Still
 
-Alright, it is clear you want to know more. Let’s keep on going deeper into how the engine works.
+Alright, it is clear you want to know more. Let's keep on going deeper into how the engine works.
 
 ### Unreal Reflection System
 
 [Blog Post: Unreal Property System (Reflection)](https://www.unrealengine.com/blog/unreal-property-system-reflection)
 
-Gameplay classes make use of special markup, so before we go over them, let’s cover some of the basics of the Unreal property system.  UE4 uses its own implementation of reflection that enables dynamic features such as garbage collection, serialization, network replication, and Blueprint/C++ communication.  These features are opt-in, meaning you have to add the correct markup to your types, otherwise Unreal will ignore them and not generate the reflection data for them.  Here is a quick overview of the basic markup:
+Gameplay classes make use of special markup, so before we go over them, let's cover some of the basics of the Unreal property system.  UE4 uses its own implementation of reflection that enables dynamic features such as garbage collection, serialization, network replication, and Blueprint/C++ communication.  These features are opt-in, meaning you have to add the correct markup to your types, otherwise Unreal will ignore them and not generate the reflection data for them.  Here is a quick overview of the basic markup:
 
 * **UCLASS() **- Used to tell Unreal to generate reflection data for a class.  The class must derive from **UObject**. 
 * **USTRUCT() **- Used to tell Unreal to generate reflection data for a struct.
@@ -352,16 +352,16 @@ Here is an example declaration of a **UCLASS**:
 		void ExampleFunction();
 	};
 
-You’ll first notice the inclusion of "MyClass.generated.h".  Unreal will generate all the reflection data and put it into this file.  You must include this file as the last include in the header file that declares your type.
+You'll first notice the inclusion of "MyClass.generated.h".  Unreal will generate all the reflection data and put it into this file.  You must include this file as the last include in the header file that declares your type.
 
-You’ll also have noticed that we can also add additional specifiers to the markup.  I’ve added some of the more common ones for demonstration.  These allow us to specify certain behavior that our types have.
+You'll also have noticed that we can also add additional specifiers to the markup.  I've added some of the more common ones for demonstration.  These allow us to specify certain behavior that our types have.
 
 * **Blueprintable** - This class can be extended by a Blueprint.
 * **BlueprintReadOnly **- This property can only be read from a Blueprint, and not written to.
 * **Category** - Defines what section this property appears under in the Details view of the Editor.  For organizational purposes.
 * **BlueprintCallable **- This function can be called from Blueprints.
 
-There are too many specifiers to list here, so I’ll link to the docs for them:
+There are too many specifiers to list here, so I'll link to the docs for them:
 
 [List of UCLASS Specifiers](Programming/UnrealArchitecture/Reference/Classes/Specifiers)
 
@@ -391,9 +391,9 @@ You can limit the scope of the search by providing a more specific type to the i
 
 Warning: Using object iterators in PIE (Play In Editor) can lead to unexpected results.  Since the editor is loaded, the object iterator will return all **UObjects** created for your game world instance, in addition to those that are just being used by the editor.
 
-Actor iterators work in much the same way as object iterators, but only work for objects that derive from **AActor**.  Actor iterators don’t have the problem noted below, and will only return objects being used by the current game world instance.
+Actor iterators work in much the same way as object iterators, but only work for objects that derive from **AActor**.  Actor iterators don't have the problem noted below, and will only return objects being used by the current game world instance.
 
-When creating an actor iterator, you need to give it a pointer to a **UWorld** instance.  Many **UObject** classes, such as **APlayerController**, provide a **GetWorld** method to help you.  If you’re not sure, you can check the **ImplementsGetWorld** method on a **UObject** to see if it implements the **GetWorld** method.
+When creating an actor iterator, you need to give it a pointer to a **UWorld** instance.  Many **UObject** classes, such as **APlayerController**, provide a **GetWorld** method to help you.  If you're not sure, you can check the **ImplementsGetWorld** method on a **UObject** to see if it implements the **GetWorld** method.
 
 	APlayerController* MyPC = GetMyPlayerControllerFromSomewhere();
 	UWorld* World = MyPC->GetWorld();
@@ -427,14 +427,14 @@ UE4 uses the reflection system to implement a garbage collection system.  With g
 
 In the garbage collector, there is this concept called the *root set*.  This root set is basically a list of objects that the collector knows about will never be garbage collected.  An object will not be garbage collected as long as there is a path of references from an object in the root set to the object in question.  If no such path to the root set exists for an object, it is called *unreachable *and will be collected (deleted) the next time the garbage collector is ran.  The engine runs the garbage collector at certain intervals.
 
-What counts as a "reference"?  Any **UObject** pointer stored in a **UPROPERTY**.  Let’s start with a simple example.
+What counts as a "reference"?  Any **UObject** pointer stored in a **UPROPERTY**.  Let's start with a simple example.
 
 	void CreateDoomedObject()
 	{
 		MyGCType* DoomedObject = NewObject<MyGCType>();
 	}
 
-When we call the above function, we create a new **UObject**, but we don’t store a pointer to it in any **UPROPERTY**, and it isn’t a part of the root set.  Eventually, the garbage collector will detect this object is unreachable, and destroy it.
+When we call the above function, we create a new **UObject**, but we don't store a pointer to it in any **UPROPERTY**, and it isn't a part of the root set.  Eventually, the garbage collector will detect this object is unreachable, and destroy it.
 
 ### Actors and Garbage collection
 
@@ -466,7 +466,7 @@ This is a more common case, where you have actors with **UObject** properties.
 		World->SpawnActor<AMyActor>(Location, Rotation);
 	}
 
-When we call the above function, we spawn an actor into the world.  The actor’s constructor creates two objects.  One gets assigned to a **UPROPERTY**, the other to a bare pointer.  Since actors are automatically a part of the root set, **SafeObject** will not be garbage collected because it can be reached from a root set object.  **DoomedObject**, however, will not fare so well.  We didn’t mark it with **UPROPERTY**, so the collector doesn’t know its being referenced, and will eventually destroy it.
+When we call the above function, we spawn an actor into the world.  The actor's constructor creates two objects.  One gets assigned to a **UPROPERTY**, the other to a bare pointer.  Since actors are automatically a part of the root set, **SafeObject** will not be garbage collected because it can be reached from a root set object.  **DoomedObject**, however, will not fare so well.  We didn't mark it with **UPROPERTY**, so the collector doesn't know its being referenced, and will eventually destroy it.
 
 When a **UObject** is garbage collected, all **UPROPERTY** references to it will be set to **nullptr** for you.  This makes it safe for you to check if an object has been garbage collected or not.
 
@@ -501,11 +501,11 @@ Normal, non-**UObjects** can also have the ability to add a reference to an obje
 		}
 	};
 
-We use the **FReferenceCollector** to manually add a hard reference to the **UObject **we need and don’t want garbage collected.  When the object is deleted and its destructor is run, the object will automatically clear all references that it added.
+We use the **FReferenceCollector** to manually add a hard reference to the **UObject **we need and don't want garbage collected.  When the object is deleted and its destructor is run, the object will automatically clear all references that it added.
 
 ### Class Naming Prefixes
 
-Unreal Engine provides tools that generate code for you during the build process. These tools have some class naming expectations and will trigger warnings or errors if the names don’t match the expectations. The list of class prefixes below delineates what the tools are expecting.
+Unreal Engine provides tools that generate code for you during the build process. These tools have some class naming expectations and will trigger warnings or errors if the names don't match the expectations. The list of class prefixes below delineates what the tools are expecting.
 
 * Classes derived from **Actor** prefixed with **A**, e.g. AController.
 * Classes derived from **Object** are prefixed with **U**, e.g. UComponent.
@@ -546,16 +546,16 @@ FString MyStr = TEXT("Hello, Unreal 4!").
 
 **FText** is similar to **FString**, but it is meant for localized text.  To create a new **FText**, use the **NSLOCTEXT** macro.  This macro takes a namespace, key, and a value for the default language:
 
-	FText MyText = NSLOCTEXT("Game UI", “Health Warning Message”, “Low Health!”)
+	FText MyText = NSLOCTEXT("Game UI", "Health Warning Message", "Low Health!")
 
 You could also use the **LOCTEXT** macro, so you only have to define a namespace once per file.  Make sure to undefine it at the bottom of your file.
 
 	// In GameUI.cpp
 	#define LOCTEXT_NAMESPACE "Game UI"
 	
-	//…
-	FText MyText = LOCTEXT("Health Warning Message", “Low Health!”)
-	//…
+	//...
+	FText MyText = LOCTEXT("Health Warning Message", "Low Health!")
+	//...
 	
 	#undef LOCTEXT_NAMESPACE
 	// End of file
@@ -574,7 +574,7 @@ A **FName** stores a commonly recurring string as an identifier in order to save
 
 [Full Topic: Character Encoding](Programming/UnrealArchitecture/StringHandling/CharacterEncoding)
 
-This is needed for some functions, such as **FString::Printf**, where the **‘%s’** string format specifier expects a **TCHAR*** instead of an **FString**.
+This is needed for some functions, such as **FString::Printf**, where the **'%s'** string format specifier expects a **TCHAR*** instead of an **FString**.
 
 	FString Str1 = TEXT("World");
 	int32 Val1 = 123;
@@ -582,8 +582,8 @@ This is needed for some functions, such as **FString::Printf**, where the **‘%
 
 The **FChar** type provides a set of static utility functions for working with individual **TCHARs**.
 
-	TCHAR Upper(‘A’);
-	TCHAR Lower = FChar::ToLower(Upper); // ‘a’
+	TCHAR Upper('A');
+	TCHAR Lower = FChar::ToLower(Upper); // 'a'
 
 [REGION:note]
 The **FChar** type is defined as **TChar<TCHAR>** (as it is listed in the API).
@@ -599,7 +599,7 @@ Containers are classes whose primary function is to store collections of data.
 
 #### TArray
 
-Of these three containers the primary container you’ll use in Unreal Engine 4 is **TArray**, it functions much like **std::vector **does, but offers a lot more functionality.  Here are some common operations:
+Of these three containers the primary container you'll use in Unreal Engine 4 is **TArray**, it functions much like **std::vector **does, but offers a lot more functionality.  Here are some common operations:
 
 	TArray<AActor*> ActorArray = GetActorArrayFromSomewhere();
 
@@ -616,16 +616,16 @@ Of these three containers the primary container you’ll use in Unreal Engine 4
 	ActorArray.Add(NewActor);
 
 	// Adds an element to the end of the array only if it is not already in the array
-	ActorArray.AddUnique(NewActor); // Won’t change the array because NewActor was already added
+	ActorArray.AddUnique(NewActor); // Won't change the array because NewActor was already added
 
-	// Removes all instances of ‘NewActor’ from the array
+	// Removes all instances of 'NewActor' from the array
 	ActorArray.Remove(NewActor);
 
 	// Removes the element at the specified index
 	// Elements above the index will be shifted down by one to fill the empty space
 	ActorArray.RemoveAt(Index);
 
-	// More efficient version of ‘RemoveAt’, but does not maintain order of the elements
+	// More efficient version of 'RemoveAt', but does not maintain order of the elements
 	ActorArray.RemoveAtSwap(Index);
 
 	// Removes all elements in the array
@@ -638,13 +638,13 @@ Of these three containers the primary container you’ll use in Unreal Engine 4
 	{
 		GENERATED_BODY();
 
-		// …
+		// ...
 
 		UPROPERTY()
 		TArray<AActor*> GarbageCollectedArray;
 	};
 
-We’ll cover the garbage collection in depth in a later section.
+We'll cover the garbage collection in depth in a later section.
 
 [Full Topic: TArrays](Programming/UnrealArchitecture/TArrays)
 
@@ -654,7 +654,7 @@ We’ll cover the garbage collection in depth in a later section.
 
 A **TMap** is a collection of key-value pairs, similar to **std::map**.  **TMap** has quick methods for finding, adding, and removing elements based on their key.  You can use any type for the key, as long as it has a **GetTypeHash** function defined for it, which we will go over later.
 
-Let’s say you were creating a grid-based board game and needed to store and query what piece is on each square.  A **TMap **would provide you with an easy way to do that.  If your board size is small and is always the same size, there are obviously more efficient ways at going about this, but let’s roll with it for example’s sake!
+Let's say you were creating a grid-based board game and needed to store and query what piece is on each square.  A **TMap **would provide you with an easy way to do that.  If your board size is small and is always the same size, there are obviously more efficient ways at going about this, but let's roll with it for example's sake!
 
 	enum class EPieceType
 	{
@@ -770,7 +770,7 @@ Using iterators, you can loop through each element of a container.  Here is an e
 			AEnemy* Enemy = *EnemyIterator;
 			if (Enemy.Health == 0)
 			{
-				// ‘RemoveCurrent’ is supported by TSets and TMaps
+				// 'RemoveCurrent' is supported by TSets and TMaps
 				EnemyIterator.RemoveCurrent();
 			}
 		}
@@ -819,7 +819,7 @@ Iterators are nice, but can be a bit cumbersome if you just want to loop through
 		// ...
 	}
 
-Remember that the **auto** keyword doesn’t automatically specify a pointer/reference for you, you need to add that yourself!
+Remember that the **auto** keyword doesn't automatically specify a pointer/reference for you, you need to add that yourself!
 
 #### Using your own types with TSet/TMap (Hash Functions)
 
