Page redone to address UEDOC-1941.
Different example provided, shortened the page and updated all images/content.

diff --git a/Engine/Documentation/Source/Gameplay/HowTo/UseOnHit/Blueprints/UseOnHit_BP.INT.udn b/Engine/Documentation/Source/Gameplay/HowTo/UseOnHit/Blueprints/UseOnHit_BP.INT.udn
--- a/Engine/Documentation/Source/Gameplay/HowTo/UseOnHit/Blueprints/UseOnHit_BP.INT.udn
+++ b/Engine/Documentation/Source/Gameplay/HowTo/UseOnHit/Blueprints/UseOnHit_BP.INT.udn
@@ -1,331 +1,108 @@
 Availability: Public
 Title:Using the OnHit Event in Blueprints
 Crumbs: %ROOT%, Gameplay, Gameplay\HowTo\UseOnHit
-Description:In this HowTo, the OnHit Event is used to affect Actors in different ways. 
+Description:Using the OnHit Events to fracture a destructible mesh as well as applying impulse to move it. 
 Related: Engine/Blueprints
 Related: Engine/Blueprints/UserGuide/Events
 Related: Engine/Blueprints/UserGuide/CastNodes
 
 
-[TOC (start:2 end:2)]
+On this page we will use the **OnComponentHit** and **Event Hit** [](Engine/Blueprints/UserGuide/Events) to apply damage to and fracture a Destructible Mesh as well as apply an impulse at the hit location to push the destructible, simulating the effects of being hit by a projectile and applying force at the hit location. 
 
-On this page, we will use the Event Hit Event to shoot out a light in a light fixture, which causes an emitter to spawn. We will also apply an impulse to a physics object, causing it to be launched in the direction of the impact point. Finally, we will apply damage to an object and destroy it after its "Health" has been depleted. 
-
-## Project Setup
-
-In order to generate Hit Events, we need an object to shoot and things to shoot at which we will do in the Setup sections below. 
- 
 [REGION:note]
-For this example, we are using the **Blueprint First Person Template with Starter Content** enabled.
-[/REGION:note]
-
-### Setup (Bomb Blueprint)
-
-In this section, we will create a "Bomb" Blueprint that our player character can shoot/throw to try and hit things with. 
-
-1.	In the **Content Browser**, click the **Create** Button and select **Blueprint**, then select **Actor** as the Parent Class. 
-
-1.	Name the Blueprint created **BP_Bomb**, then **Double-click** on it to open it up. 
-
-1.	Click the **Add Component** button and search for **projectile** then select **Projectile Movement**. 
-
-	[REGION:raw]
-	![](OnHit_1.png)
-	[/REGION]
+_For this How To, we are using the **Blueprint First Person Template** project with **Starter Content** enabled._
+[/REGION]
 
-1.	Click **Add Component** and add a **Sphere Collision** and a **Static Mesh**. 
+1.	In your project in the **Content Browser**, go to the **Content/FirstPersonBP/Meshes** folder. 
 
-1.	Click the **Sphere** component in the **Components** window, then in the **Details** rename it to **Collision**. 
+1.	**Right-click** on **FirstPersonCube_Rounded** Static Mesh and select **Create Destructible Mesh**. 
 
-	[REGION:raw]
-	![](OnHit_2.png)
-	[/REGION]
+	![](OnHit1.png)
 
-1.	Click the **Static Mesh** and rename it to **Bomb** and set the **Z** value for Location to **-10** and the Scale for XYZ to **.25** each. 
+	We are creating a Destructible Mesh out of this Static Mesh so that when it is shot, we can fracture it causing it to break apart. 
 
-1.	Under **Static Mesh**, click the dropdown menu and search for and add the **Shape_Sphere**. 
+1.	Inside the Destructible Mesh Editor, click the **Fracture Mesh** button then **Save** and close the window. 
 
-	[REGION:raw]
-	![](OnHit_3.png)
-	[/REGION]
+	![](OnHit2.png)
 
-	[REGION:note]
-	If you do not see the **Shape_Sphere**, make sure that your project has **Starter Content** enabled. 
-	[/REGION]
+1.	**Right-click** on the newly created Destructible Mesh asset, then under **Asset Actions**, select **Create Blueprint Using This...** 
 
-1.	Scroll down and find the **Collision** section, and set the collision to **NoCollision**. 
+	![](OnHit3.png)
 
-	[REGION:raw]
-	![](OnHit_4.png)
-	[/REGION]
+	Be sure to choose a save location and name for the new Blueprint created. 
 
-	Collision will be detected using the Root Sphere Component instead of the Static Mesh itself. 
+1.	Inside the Blueprint Editor for the new Blueprint, in the **Components** window select the **Destructible** component. 
 
-1.	Click the **Collision** Sphere Component, then in the **Details** panel, set the **Sphere Radius** to **20**. 
+1.	In the **Details** panel, under **Physics**, check the **Simulate Physics** checkbox. 
 
-1.	Under **Collision**, set **Simulation Generates Hit Events** to **True** and **Block All** for the **Collision Presets**. 
+	![](OnHit4.png)
 
-	[REGION:raw]
-	![](OnHit_5.png)
-	[/REGION]
+	This Blueprint will now simulate physics. 
 
-1.	Click the **Projectile Movement** component, then in the **Details** panel, set the **Initial** and **Max Speed** values to **1500**. 
+1.	Also in the **Details** panel, at the bottom under **Events**, click the **OnComponentHit** Event. 
 
-1.	Also set **Should Bounce** to **True**.
+	![](OnHit5.png)
 
-	[REGION:raw]
-	![](OnHit_6.png)
-	[/REGION]
+	This will add a new node to and open the Event Graph. 
 
-	The previous two steps will determine how the projectile moves, you can change these values as you wish. 
+	![](OnHit6.png)
 
-1.	Click the **Class Defaults** button on the tool bar to open the Blueprint Defaults in the Details tab, then set **Initial Life Span** to **3**. 
+1.	**Left-click** and drag off the **Other Actor** pin and search for and add the **Cast To FirstPersonProjectile** node. 
 
-	[REGION:raw]
-	![](OnHit_7.png)
-	[/REGION]
+	![](OnHit7.png)
 
-	This determines how long the projectile will remain in the world once spawned, again you can set this as you wish. 
+	Here we are [](Engine/Blueprints/UserGuide/CastNodes) to another Blueprint called the **FirstPersonProjectile** Blueprint and making sure that when we hit the Destructible Mesh, that it is infact the FirstPersonProjectile Blueprint that hit it. If it is, then we can apply additional script to fracture the mesh, breaking it apart. If it is not, we will not do anything. 
 
-1.	**Compile** and **Save** then close the Blueprint. 
+1.	Drag off the **Hit** pin and add a **Break Hit Result** node. 
 
+	![](OnHit8.png)
 
-### Setup (Light Blueprint)
+	The **Break Hit Result** node gives us access to all sorts of information about the resulting hit and its information. 
 
-In this section, we will modify the Ceiling Light Blueprint so that it will spawn a spark effect when hit. 
+1.	Drag off the **Hit Actor** pin from the **Break Hit Result** and add the **Apply Point Damage** node. 
 
-1.	In the **Game/StarterContent/Blueprints** folder, open the **Blueprint_CeilingLight** Blueprint. 
+1.	On the **Apply Point Damage** node, set the **Base Damage** to **100** and set the **Damage Type Class** to **Damage Type**. 
 
-1.	Click the **Component** tab, then click **Add Component** and search for and add a **Particle System**. 
+	![](OnHit9.png)
 
-	[REGION:raw]
-	![](OnHit_8.png)	
-	[/REGION]
+	The Apply Point Damage node allows us to not only specifies the amount of damage to cause, but the location at which the damage was inflicted. 
 
-1.	In the **Details** panel for the particle system, set **Z** for Location to **-140**, use the **P_Sparks** effect, and un-check **Auto Activate**. 
+1.	On the **Apply Point Damage** node, connect the remaining wires as shown below. 
 
-	![](OnHit_9.png)
+	![](OnHit10.png)
 
-	This particle system will be activated when the player hits the lamp, causing sparks to appear. 
+	With all wires connected, your graph should look similar to above. 
 
-1.	**Compile** and **Save** then close the Blueprint. 
+1.	**Compile** and **Save** then, close the Blueprint. 
 
-### Setup (Shootable Physics Box)
+1.	Drag the Blueprint into the level from the **Content Browser**. 
 
-Here we will create a simple box that reacts to physics so that our character can shoot it. 
+	![](OnHit11.png)
 
-1.	In the **Content Browser**, create another Blueprint of the Actor type, then open it up. 
+1.	Click **Play** to play in the editor and use the **Left-mouse** button to fire a projectile at the cube. 
 
-1.	First, add a **Scene** component (which is a dummy component that we will use as our Root component). 
+	[OBJECT:EmbeddedVideo]
+	[PARAMLITERAL:width]
+	640
+	[/PARAMLITERAL]
+	[PARAMLITERAL:height]
+	360
+	[/PARAMLITERAL]
+	[PARAMLITERAL:videoid]
+	WRlsw9q-bVA
+	[/PARAMLITERAL]
+	[/OBJECT]
 
-	[REGION:raw]
-	![](OnHit_11.png)
-	[/REGION]
+	When you play in the editor, you will see that when you hit the cube with the projectile that is fired it causes the cube to take damage and break apart as well as applies an impulse at the location in which it was hit causing it to fly in that direction. The amount of force that is applied is defined inside the **FirstPersonProjectile** Blueprint which uses the **Event Hit** node to determine when the projectile actually hits something. 
 
-1.	Then using the **Components** tab, add a **Static Mesh** component. 
+1.	In the **Content/FirstPersonBP/Blueprints** folder, open the **FirstPersonProjectile** Blueprint. 
 
-	[REGION:raw]
-	![](OnHit_10.png)
+	[REGION:lightbox]
+	[![](OnHit12.png)](OnHit12.png)
 	[/REGION]
+	_Click the image for a larger view._
 
-1.	Click the **Static Mesh**, then in the **Details** panel, use the **Shape_Cube** and set **Simulate Physics** to **True**. 
-
-	![](OnHit_13.png)
-
-	Also make sure that **PhysicsActor** is the selected **Collision Preset** used. 
-
-1.	**Compile** and **Save** then close the Blueprint. 
-
-### Setup (Player Character)
-
-The final setup section will have us modify the player character Blueprint so they can press a button to shoot our Bomb. 
-
-1.	Inside the **Content Browser** under **Game/Blueprints**, open the **FirstPersonCharacter** Blueprint. 
-
-1.	On the **Graph** tab, **Right-click** in the graph and add the **Left Mouse Button** event. 
-
-1.	Drag off the **Pressed** pin and add the **SpawnActorFromClass** node, and set **Class** to **BP_Bomb**. 
-
-	![](OnHit_14.png)
-
-1.	In the **MyBlueprint** window, click the **Add Variable** button to create a new variable. 
-
-1.	In the **Details** panel for the variable, rename it to **SpawnOffset** and make it of the **Vector** type. 
-
-1.	**Compile** the Blueprint, then set the **Default Value** for **SpawnOffset** to **100, 33, 10**. 
-
-	![](OnHit_15.png)
-
-	This variable will be used to add an offset to where our Bomb projectile is spawned relative to the player. 
-
-1.	Hold **Control** and drag the **Spawn Offset** variable into the graph near the **Left Mouse Button** event. 
-
-1.	**Right-click** in the graph and add a **Get Control Rotation** node. 
-
-1.	Off its **Return Value**, add the **Rotate Vector** node, then connect the **Spawn Offset** to the **A** pin. 
-
-	![](OnHit_16.png)
-
-1.	**Right-click** in the graph and add the **Get Actor Location** node, then off the **Return** add a **Vector + Vector** node. 
-
-1.	Connect the **Return Value** from the **Rotate Vector** to the other in pin of the **Vector + Vector** node. 
-
-	![](OnHit_17.png)
-
-	Here we are using the control rotation and a vector variable to produce a new vector value to add to the player character's location. This gives us an offset in which we can spawn our bomb from, slightly in front of the player so it does not spawn inside the player character. 
-
-1.	**Right-click** in the graph and add a **Make Transform** node, then connect the nodes as shown below. 
-
-	![](OnHit_18.png)
-
-## Event Hit Example (Shooting Out Lights)
-
-With our project setup, we can now use the **Event Hit** node to shoot out our Ceiling Light and turn on the spark effect. 
-
-1.	In the **Game/StarterContent/Blueprints** folder, open the **Blueprint_CeilingLight** Blueprint. 
-
-1.	On the **Graph** tab, **Right-click** and add the **Event Hit** node. 
-
-1.	Drag off the **Other** pin of the **Event Hit** node and add the **Cast To BP_Bomb** node. 
-
-	![](OnHit_19.png)
-
-	The **Other** pin refers to the other Actor involved in the collision. By using the **Cast To BP_Bomb** node here, we are asking is the other Actor involved that we hit **BP_Bomb** and if it is, we can continue the script. If it is not, the script will not continue. This is a quicker method of running a check to see what hit the light. 
-
-	[REGION:note]
-	For details on each pin type, refer to the [Event Hit](Engine/Blueprints/UserGuide/Events#eventhit) documentation.
-	[/REGION]
-
-1.	Hold **Control** and drag the **Sparks** component into the graph from the **MyBlueprint** window. 
-
-1.	Drag off **Sparks** and add the **Activate** node and connect as shown below. 
-
-	![](OnHit_20.png)
-
-1.	Hold **Control** and drag the **Point Light 1** component into the graph from the **MyBlueprint** window. 
-
-1.	Drag off **Point Light 1** and add the **Set Visibility** node, leaving **New Visibility** unchecked, and connect as shown. 
-
-	![](OnHit_21.png)
-
-1.	Following **Set Visibility**, add a **Delay** node (set to 1.5), then drag off **Sparks** and use the **Deactivate** node and connect as shown. 
-
-	![](OnHit_22.png)
-
-1.	**Compile** and **Save**, then close the Blueprint. 
-
-1.	In the **Content Browser**, drag-and-drop the **Blueprint_CeilingLight** into a dark area in your level. 
-
-	![](OnHit_23.png)
-
-1.	Click **Play** to play in the editor, then use the **Left Mouse Button** to shoot and hit the light. 
-
-	![](OnHit_24.gif)(convert:false)
-
-	You should now see that when you hit the light with the bomb, the light goes off and the spark appears then disappears. 
-
-By using the Cast to BP_Bomb node, we are ensuring that only the Bomb Blueprint can register a Hit on the Light. If we wanted to have other objects hit the light and register hits, we could remove that aspect of the script and use just the Event Hit node alone (or possibly add a Branch following the Event Hit to determine what happens when different objects hit it). 
-
-
-## Event Hit Example (Applying Force and Damage)
-
-In this example, we are going to use the Event Hit node and apply force to the object we hit at the location we hit it. We will also apply damage to the object we hit, where once we hit it enough times, the object is destroyed. 
-
-1.	In the **Content Browser**, locate your **Cube** Blueprint and drag-and-drop it into the level. 
-
-	![](OnHit_25.png)
-
-	If you play in the editor now, even though the object is a Physics Actor, it does not appear to be affected when the player shoots it. If you move the player character into the cube you can push the cube around however. There is not enough force behind the projectile when it collides with the cube, so it is not being moved. We can fix this by adding an impulse force to push the object when a collision with the projectile occurs. 
-
-1.	Open the **BP_Bomb** Blueprint, **Right-click** in the graph and add the **Event Hit** node. 
-
-1.	Drag off the **Other Comp** pin and add the **Add Impulse at Location** node. 
-
-	![](OnHit_26.png)
-
-	The **Other Comp** pin refers to the component on the other Actor involved in the collision to which we will apply an impulse. 
-
-1.	**Right-click** in the graph and add a **Get Velocity** node, then off its **Return**, add a **Vector X Float** pin (set to 100). 
-
-1.	Connect the out pin on the **Vector X Float** node to the **Impulse** pin on the **Add Impulse at Location** node. 
-
-1.	Connect the **Hit Location** pin on the **Event Hit** to the **Location** pin on the **Add Impulse at Location** node. 
-
-	![](OnHit_27.png)
-
-	The last few steps should result in the graph depicted above. Here we are getting the current velocity of the Actor (our bomb) and multiplying it by 100 and using that as the impulse to apply to the component of the Actor we hit at the location of the hit.
-
-1.	Drag off the **Other** pin of the **Event Hit** and add the **Apply Damage** node. 
-
-1.	Set the **Base Damage** value to 10 and connect the nodes as shown. 
-
-	![](OnHit_28.png)
-
-	We are now using the **Other** pin to get the Actor itself and applying 10 damage to it using the **Apply Damage** node. 
-
-1.	**Compile** and **Save**, then close the Blueprint. 
-
-1.	In the **Content Browser**, open the **Cube** Blueprint. 
-
-1.	On the **Graph** tab, in the **MyBlueprint** window, click the **Add Variable** button. 
-
-1.	Name the variable **CubeHealth**, set it to a **Float** type, then **Compile** and set the **Default Value** to 50. 
-
-	![](OnHit_29.png)
-
-1.	**Right-click** in the graph and add the **Event Any Damage** node and **Control** drag in the **CubeHealth** variable. 
-
-	![](OnHit_30.png)
-
-1.	Drag off **CubeHealth** and add a **Float - Float** node, then plug **Damage** into the other in pin of the **Float - Float** node. 
-
-1.	Hold **Alt** and drag in the **CubeHealth** variable to create a **Set** node and connect the nodes as shown below. 
-
-	![](OnHit_31.png)
-
-1.	Hold the **B** key, and click in the graph to create a **Branch** node. 
-
-1.	Drag off the **CubeHealth** variable and add a **<=** node, then plug its out pin into the **Condition** on the **Branch**. 
-
-1.	Off the **True** add a **Destroy Actor** node and off the **False**, add a **Print String** node. 
-
-1.	Drag off **CubeHealth** and plug it into the **InString** on the **Print String** node (a conversion node will automatically be created). 
-
-	![](OnHit_32.png)
-
-	The graph setup should look similar to above. Here we are saying that when the Cube Blueprint is damaged, we are taking the incoming damage (set to 10 when we applied it) and subtracting it from our cube's overall health (CubeHealth), then setting the value as the new overall health value for the cube. 
-
-	If the cube's health is less than or equal to 0, we will then destroy the cube. If it is not, we will print to the screen the current value of CubeHealth so that we can see how much health the cube has left. 
-
-1.	**Compile** and **Save**, then close the Blueprint. 
-
-1.	Click the **Play** button to play in the editor. 
-
-	[![](OnHit_32b.png)](OnHit_32.mp4)
-
-	When playing in the editor, you should now see that the cube responds to the impact from the bomb projectile and apply force at the location of impact. Each time you hit the cube, health is taken away from the cube's overall health value and printed to the upper left hand corner of the screen. After shooting it 5 times, the cube is destroyed. 
-
-These are two examples of how the Event Hit can be used for gameplay purposes, depending upon your needs, you may perform additional checks when a Hit occurs and based on the Hit, perform different scripts that do different things. 
-	
-	
-
-
-
-
-
-
-
-
-
-
-
-
-
-
-
-
-
-
+	The script in this Blueprint, checks to see if the object that is hit is simulating physics (which we set to true on our cube Blueprint). If it is, it then applies an Impulse at the location it hits (the amount is defined inside the green box which we have reduced to 50). You can adjust this value to increase/decrease the amount of impulse applied when a hit occurs. 
 
 
 
