#UE4doc #docs: Removed extraneous bolding as according to the instructional guide. ##Overview removed as according to the instructional guide. Changed title of How To to match instructional guide.

diff --git a/Engine/Documentation/Source/Gameplay/HowTo/FindingActors/FindingActorsOverview.INT.udn b/Engine/Documentation/Source/Gameplay/HowTo/FindingActors/FindingActorsOverview.INT.udn
--- a/Engine/Documentation/Source/Gameplay/HowTo/FindingActors/FindingActorsOverview.INT.udn
+++ b/Engine/Documentation/Source/Gameplay/HowTo/FindingActors/FindingActorsOverview.INT.udn
@@ -1,11 +1,8 @@
 Availability: Public
-Title:How to Find Actors
+Title:Finding Actors
 Crumbs: %ROOT%, Gameplay
 Description:A How To Guide for Finding Actors in your Scenes in Unreal Engine 4. 
 
-
-##Overview
-
 Once you start constructing your game and its levels, you may find that locating an individual Actor in your scene may be difficult. It is not uncommon to have hundreds of Actors in your scene ranging from level set pieces or geometry, to NPCs or enemies, to interactable objects or pick-ups and locating one Actor out of all those is like the old saying of "looking for a needle in a haystack". 
 
 You can however use the **World Outliner** to help narrow your search for an Actor in your scene (pictured below). 
@@ -20,12 +17,12 @@ With an Actor selected in the **World Outliner**, and while navigating inside th
 
 Refer to the [](Engine/UI/LevelEditor/SceneOutliner) documentation for more information on locating an Actor in your scenes. 
 
-While finding an Actor in the Editor can be accomplished with the help of the **World Outliner**, how do you go about finding an Actor through **Blueprints** so that you may modify it in some way? This "How To" series will show you some ways in which you can search your scene for an Actor based on varying criteria at which point you can then access and modify it if you wish. 
+While finding an Actor in the Editor can be accomplished with the help of the **World Outliner**, how do you go about finding an Actor through Blueprints so that you may modify it in some way? This "How To" series will show you some ways in which you can search your scene for an Actor based on varying criteria at which point you can then access and modify it if you wish. 
 
 
 ## Setup
 
-To learn more about **Finding Actors** using Blueprints, use the link below:
+To learn more about Finding Actors using Blueprints, use the link below:
 
 [REGION:buttonlist]
 * [](Gameplay\HowTo\FindingActors\Blueprints)
