#UE4doc #docs: Removed extraneous bolding as according to the instructional guide. ##Overview removed as according to the instructional guide. Changed title of How To to match instructional guide. Altered images to match the instructional guide: green -> yellow.

diff --git a/Engine/Documentation/Source/Gameplay/HowTo/FindingActors/Blueprints/FindingActors_BP.INT.udn b/Engine/Documentation/Source/Gameplay/HowTo/FindingActors/Blueprints/FindingActors_BP.INT.udn
--- a/Engine/Documentation/Source/Gameplay/HowTo/FindingActors/Blueprints/FindingActors_BP.INT.udn
+++ b/Engine/Documentation/Source/Gameplay/HowTo/FindingActors/Blueprints/FindingActors_BP.INT.udn
@@ -1,13 +1,11 @@
 Availability: Public
-Title:How to Find Actors in Blueprints
+Title:Finding Actors in Blueprints
 Crumbs: %ROOT%, Gameplay, Gameplay\HowTo\FindingActors
 Description:A How To Guide for Finding Actors in Blueprints. 
 
 [TOC (start:2 end:2)]
 
-## Overview
-
-This page covers examples on how to find Actors in a level by using the **Get All Actors of Class** node. This node when called will retrieve all Actors in your level of the specified Class and place them in an **Array** (collection of variables) and from that Array, you can pull out all Actors or individual Actors based on criteria you wish to filter by. You can then access the properties of the Actor(s) or modify them in some way based on the type of Actor it is and what you wish to achieve. 
+This page covers examples on how to find Actors in a level by using the Get All Actors of Class node. This node when called will retrieve all Actors in your level of the specified Class and place them in an Array (collection of variables) and from that Array, you can pull out all Actors or individual Actors based on criteria you wish to filter by. You can then access the properties of the Actor(s) or modify them in some way based on the type of Actor it is and what you wish to achieve. 
  
 
 ## Get All Actors of Class
@@ -67,7 +65,7 @@ Using this method, we were able to find all Actors of a particular Class and aff
 
 ## Get Specific Actors
 
-In the previous section, we learned that using the **Get All Actors of Class** node produces an **Array** of Actors of the specified Class. This section will show you how you can filter the Array results based on different criteria to get specific Actors or a single Actor from the Array. 
+In the previous section, we learned that using the **Get All Actors of Class** node produces an Array of Actors of the specified Class. This section will show you how you can filter the Array results based on different criteria to get specific Actors or a single Actor from the Array. 
 
 One method you can take is by using **Tags** on the Actor, for example: 
 
