#UE4doc #docs: ##Overview removed as according to the instructional guide.

diff --git a/Engine/Documentation/Source/GettingStarted/DownloadingUnrealEngine/Perforce/AutoMerge/AutoMerge.INT.udn b/Engine/Documentation/Source/GettingStarted/DownloadingUnrealEngine/Perforce/AutoMerge/AutoMerge.INT.udn
--- a/Engine/Documentation/Source/GettingStarted/DownloadingUnrealEngine/Perforce/AutoMerge/AutoMerge.INT.udn
+++ b/Engine/Documentation/Source/GettingStarted/DownloadingUnrealEngine/Perforce/AutoMerge/AutoMerge.INT.udn
@@ -5,10 +5,6 @@ Description:Script that automates the process of performing an integration using
 
 [TOC(start:2)]
 
-
-
-## Overview
-
 The Perforce AutoMerge script exists to automate the tedious and/or error-prone steps involved in performing a Perforce integration (aka merge). It incorporates several best-practices for merging as described in the [Perforce Integration](GettingStarted/DownloadingUnrealEngine/Perforce/Integration) page, so is the recommended way to merge in the Unreal Engine Perforce database.
 
 
