#UE4doc #docs:
Updated engine version
Added references to Community edition being supported
Changed font settings for multiple terms
Fixed step about finding a certain option
Replaced some images to have the correct yellow colored borders for highlighting.

diff --git a/Engine/Documentation/Source/Programming/Development/VisualStudioSetup/UnrealVS/UnrealVS.INT.udn b/Engine/Documentation/Source/Programming/Development/VisualStudioSetup/UnrealVS/UnrealVS.INT.udn
--- a/Engine/Documentation/Source/Programming/Development/VisualStudioSetup/UnrealVS/UnrealVS.INT.udn
+++ b/Engine/Documentation/Source/Programming/Development/VisualStudioSetup/UnrealVS/UnrealVS.INT.udn
@@ -2,10 +2,10 @@ Availability:Public
 Title:UnrealVS Extension
 Crumbs: %ROOT%, GettingStarted, Programming, Programming/Development/VisualStudioSetup
 Description:Guide to installing, setting up, and using the UnrealVS extension for Visual Studio.
-Version: 4.5
+Version: 4.9
 
 [EXCERPT:Overview]
-The UnrealVS extension for Visual Studio provides easy access to common actions when developing with Unreal Engine.
+The **UnrealVS** extension for Visual Studio provides easy access to common actions when developing with Unreal Engine 4.
 
 ![](unrealvs_toolbar_cmd.png)
 
@@ -18,7 +18,7 @@ Features include:
 * Quick Build project menu.
 
 [REGION:note]
-The UnrealVS extension **does not work** with the **Visual Studio Express** editions. It is only compatible with Visual Studio Professional.
+The UnrealVS extension does not work with the Visual Studio Express editions. It is only compatible with Visual Studio Community and Professional editions.
 [/REGION]
 [/EXCERPT:Overview]
 
@@ -47,8 +47,9 @@ The UnrealVS extension **does not work** with the **Visual Studio Express** edit
 	[REGION:note]
 	If Visual Studio is already running, you will need to restart before the extension will be loaded and displayed.
 	[/REGION]
-
-1. The extension is also shown in the Installed Products list in the **Help->About Microsoft Visual Studio** dialog.
+	[REGION:tip]
+	The extension is also shown in the Installed Products list in the **Help->About Microsoft Visual Studio** dialog.
+	[/REGION]
 
 1. Go to **View->Toolbars** or **right-click** the Visual Studio toolbar area and select **UnrealVS** to display the extension's toolbar.
 
@@ -58,16 +59,16 @@ The UnrealVS extension **does not work** with the **Visual Studio Express** edit
 
 	![UnrealVS Customize Toolbar](toolbar_customize.png)
 	
-1. In the Customize	dialog, choose **Add Command...** and select **UnrealVS** from the Categories list to see a list of UnrealVS commands that can be added to the toolbar.
+1. In the **Customize**	dialog, choose **Add Command...** and select **UnrealVS** from the **Categories** list to see a list of UnrealVS commands that can be added to the toolbar.
 
 	![UnrealVS Customize Commands](toolbar_customize_commands.png)
 
 
 ## Set the Startup Project
 
-The **Startup Project** dropdown is a means of quickly setting or switching between _startup projects_. 
+The **Startup Project** dropdown is a means of quickly setting or switching between startup projects. 
 It automatically lists all available projects that build executables in the solution.
-Selecting a project from the list sets that as the current _startup project_.
+Selecting a project from the list sets that as the current startup project.
 
 ![](startup_project_menu.png)
 
@@ -77,12 +78,12 @@ You can change which projects are listed in the UnrealVS options. To show only G
 
 ### Build Startup Project
 
-The extension also contains a command for building the current _startup project_ which can be bound to a 
+The extension also contains a command for building the current startup project which can be bound to a 
 hotkey or other means of running custom commands.
 
-**To bind the command to a hotkey:**
+####To bind the command to a hotkey:
 
-1. Go to **Tools->Options...** and select **Keyboard** in the **Options** dialog.
+1. Go to **Tools->Options...** and select **Keyboard** under the **Environment** tree.
 
 1. Select the **_UnrealVS.BuildStartupProject_** command in the commands list. 
 
@@ -107,7 +108,7 @@ hotkey or other means of running custom commands.
 
 	![Assigned Keys](build_startup_assigned.png)
 
-1. Press the ![OK](button_ok.png) button to save the changes. Now, when you use the shortcut, the project set as the _startup project_ will automatically build.
+1. Press the ![OK](button_ok.png) button to save the changes. Now, when you use the shortcut, the project set as the startup project will automatically build.
 
 
 ## Command Line Arguments
@@ -177,7 +178,7 @@ file and run it manually.
 		* ![Succeeded Build Job](buildjob_green.png) - complete, succeeded
 		* ![Failed Build Job](buildjob_red.png) - failed
 		* ![Cancelled Build Job](buildjob_orange.png) - cancelled
-		* The current **executing build job is shown in bold**
+		* The current executing build job is shown in bold
 	* Busy animation and stop button are shown when a batch build is running.
 	
 1. You can view the output from individual **Build Jobs** by **double-clicking** an item in the **Output tab** or **right-clicking** and selecting **Show Output** from the context menu.
