#UE4doc #docs: Modified spacing between phrases using colons (spoke with Mina about edit. Don't believe it has been specifically implemented in the intructional guide yet). Hooked up unhooked hyperlink.

diff --git a/Engine/Documentation/Source/GettingStarted/DownloadingUnrealEngine/Perforce/CommonMergeProblems/CommonMergeProblems.INT.udn b/Engine/Documentation/Source/GettingStarted/DownloadingUnrealEngine/Perforce/CommonMergeProblems/CommonMergeProblems.INT.udn
--- a/Engine/Documentation/Source/GettingStarted/DownloadingUnrealEngine/Perforce/CommonMergeProblems/CommonMergeProblems.INT.udn
+++ b/Engine/Documentation/Source/GettingStarted/DownloadingUnrealEngine/Perforce/CommonMergeProblems/CommonMergeProblems.INT.udn
@@ -105,7 +105,7 @@ A good way to handle these issues is to utilize the integration functionality in
 1. Use Perforce's Safe Automatic Resolve to automatically resolve files that have not been changed by you locally.  (e.g. copyright update touches basically all files. Many of which will not have been changed locally. )
 
 
-Reference: http://www.perforce.com/perforce/doc.062/manuals/p4wsad/topics/resolving.html
+Reference: [http://www.perforce.com/perforce/doc.062/manuals/p4wsad/topics/resolving.html](http://www.perforce.com/perforce/doc.062/manuals/p4wsad/topics/resolving.html)
 
 Safe Automatic Resolve: If either your file, or their file (but not both files) are different from the base file, accept the differing version of the file as the head revision. If both files are different from base file, do not perform a resolve on this file.
 
@@ -125,15 +125,12 @@ Once branched, it helps to keep up on changes made in branches, often requiring
 ### "Not a class or namespace name" Error
 
 **Symptom:**
-
 A "not a class or namespace name" error on a class in a FooClass.h file's DECLARE_NATIVE_TYPE() macro. The class in question is declared above it in the FooClasses.h file, but presumably because of the #defines that are active when this file gets included, it is not seeing that declaration.
 
 **Cause:**
-
 An autogenerated header included in the LaunchEngineLoop.cpp (with NAMES_ONLY defined) was not included in the precompiled header (without NAMES_ONLY defined), causing a compile error.
 
 **Fix/Workaround:**
-
 Include the autogenerated header in the precompiled header with NAMES_ONLY defined.
 
 
@@ -143,17 +140,14 @@ Reference: [https://udn.epicgames.com/lists/showpost.php?list=unprog3&id=10516](
 ### Enum Clashes
 
 **Symptom:**
-
 Licensees and Epic add an enum with the same value.
 
 **Cause:**
-
 Enums are serialized by value, not by name.
 
 **Fix/Workaround:**
-
-One workaround is to create a commandlet that patches up content, but this is a very error prone and time consuming task.  A better short term workaround is for licensees to pad enums.  This is easily done in C++ by specifying an offset in the enum, however in Blueprint, padding entries will have to be added manually.  <br>  
-
+One workaround is to create a commandlet that patches up content, but this is a very error prone and time consuming task.  A better short term workaround is for licensees to pad enums.  This is easily done in C++ by specifying an offset in the enum, however in Blueprint, padding entries will have to be added manually. 
+ 
 The solution to this problem is an engine level change that serializes all enums as FNames, rather than by value.  This is an outstanding task: 41337.  This change would make renaming enums difficult, but that is a much less common case.
 
 
