#UE4doc #docs:
Updated Engine version
Changed font settings for multiple terms
Removed line about 'a few notes' that preceeded a single note.

diff --git a/Engine/Documentation/Source/Programming/Modules/Gameplay/GameplayModules.INT.udn b/Engine/Documentation/Source/Programming/Modules/Gameplay/GameplayModules.INT.udn
--- a/Engine/Documentation/Source/Programming/Modules/Gameplay/GameplayModules.INT.udn
+++ b/Engine/Documentation/Source/Programming/Modules/Gameplay/GameplayModules.INT.udn
@@ -2,7 +2,7 @@ Availability:Public
 Title:Gameplay Modules
 Crumbs:%ROOT%, Programming
 Description:Collections of gameplay classes belonging to a game project compiled into DLLs.
-Version: 4.5
+Version: 4.9
 
 [TOC(start:2)]
 
@@ -16,13 +16,13 @@ If your game project is not already created, see [Creating a Game Project](Engin
 
 ## Module Creation
 
-Gameplay modules must contain, at the very least, a header file (`.h`), C++ file (`.cpp`), and a build file (`*.Build.cs`). The header file must be located in the `Public` folder of the module's directory, i.e. `[GameName]\Source\[ModuleName]\Public`. This file includes any header files - including the module's auto-generated headers - necessary to compile the classes contained within the module.
+Gameplay modules must contain, at the very least, a header file (`.h`), C++ file (`.cpp`), and a build file (`*.Build.cs`). The header file must be located in the Public folder of the module's directory, i.e. [GameName]\Source\[ModuleName]\Public. This file includes any header files - including the module's auto-generated headers - necessary to compile the classes contained within the module.
 
 	#include "Engine.h"
 	#include "EnginePrivate.h"
 	#include "<ModuleName>Classes.h"
 
-The C++ file is placed in the `Private` folder of the module's directory, i.e. `[GameName]\Source\[ModuleName]\Private`, registers and implements the module.
+The C++ file is placed in the Private folder of the module's directory, i.e. [GameName]\Source\[ModuleName]\Private, registers and implements the module.
 
 [REGION:note]At least one module in your game must be registered using `IMPLEMENT_PRIMARY_GAME_MODULE`. Additional modules can use the alternative `IMPLEMENT_GAME_MODULE` method. See the [Multiple Gameplay Modules](#MultipleGameplayModules) section for more details on using more than one gameplay module in a game project.[/REGION]
 
@@ -32,7 +32,7 @@ The C++ file is placed in the `Private` folder of the module's directory, i.e. `
 	// Designate the module as primary
 	IMPLEMENT_PRIMARY_GAME_MODULE( <ModuleName>, "<GameName>" );
 
-The build file is placed in the root directory of the gameplay module, i.e. `[GameName]\Source\[ModuleName]`, and it defines certain pieces of information used by the UnrealBuildTool to compile the module.
+The build file is placed in the root directory of the gameplay module, i.e. [GameName]\Source\[ModuleName], and it defines certain pieces of information used by the UnrealBuildTool to compile the module.
 
 	using UnrealBuildTool;
 
@@ -47,7 +47,7 @@ The build file is placed in the root directory of the gameplay module, i.e. `[Ga
 
 ### INI File Setup
 
-Because the new gameplay modules will have `UObject` code, some configuration is necessary.
+Because the new gameplay modules will have UObject code, some configuration is necessary.
 
 1. The module needs to be added in several areas of the `DefaultEngine.ini` file:
 
@@ -75,7 +75,5 @@ You can create a primary game module, then any number of additional game-specifi
 
 ## Limitations
 
-Below are a few other notes about setting up gameplay modules:
-
-* We do support creating modules that are cross-dependent (both export and import functions and data from each other -- e.g., Engine and UnrealEd modules), but this is not ideal for compile-times and may sometimes cause problems with static initialization of variables. Gameplay modules that are not cross-dependent are harder to design and maintain, but the code may be cleaner for it.
+ We do support creating modules that are cross-dependent (both export and import functions and data from each other -- e.g., Engine and UnrealEd modules), but this is not ideal for compile-times and may sometimes cause problems with static initialization of variables. Gameplay modules that are not cross-dependent are harder to design and maintain, but the code may be cleaner for it.
 
