#UE4doc #docs:
Updated Engine version

diff --git a/Engine/Documentation/Source/Programming/Online/OnlineSubsystem.INT.udn b/Engine/Documentation/Source/Programming/Online/OnlineSubsystem.INT.udn
--- a/Engine/Documentation/Source/Programming/Online/OnlineSubsystem.INT.udn
+++ b/Engine/Documentation/Source/Programming/Online/OnlineSubsystem.INT.udn
@@ -2,7 +2,7 @@ Availability:Public
 Title:Online Subsystem Overview
 Crumbs:%ROOT%, Programming
 Description:Overview of the various systems related to the online platform.
-Version: 4.5
+Version: 4.9
 
 
 [VAR:Topic]
