#UE4doc #docs: fixed spelling error: mege -> merge. ##Overview removed as according to the instructional guide. Unbolded hyperlink as according to the insctructional guide.

diff --git a/Engine/Documentation/Source/GettingStarted/DownloadingUnrealEngine/Perforce/Perforce.INT.udn b/Engine/Documentation/Source/GettingStarted/DownloadingUnrealEngine/Perforce/Perforce.INT.udn
--- a/Engine/Documentation/Source/GettingStarted/DownloadingUnrealEngine/Perforce/Perforce.INT.udn
+++ b/Engine/Documentation/Source/GettingStarted/DownloadingUnrealEngine/Perforce/Perforce.INT.udn
@@ -5,10 +5,6 @@ Description:Guide to using Perforce source control with Unreal Engine.
 
 [TOC(start:2)]
 
-
-
-## Overview
-
 Unreal Engine is provided through a Perforce depot that licensees must connect to in order to download the engine. Since Unreal Engine is updated regularly, as a licensee working on a project, you may choose to update your version of the engine to a newer version at one or more times during the development cycle. Setting up Perforce properly from the outset to sync and integrate builds of the engine as easily as possible is important.
 
 Epic has set aside the **2nd and 4th Sunday of the month from 8am-Noon EST** for performing scheduled maintenance. Please keep this in mind when planning as it is possible that any or all of the support outlets, especially network/server resources, could be affected during these times.
@@ -22,7 +18,7 @@ The first step to being able to connect to Epic's public Perforce depot is to se
 **Web-based Connection**
 
 
-1. Go to **[https://vpn.epicgames.com](https://vpn.epicgames.com)** in your web browser (IE recommended for stability).
+1. Go to [https://vpn.epicgames.com](https://vpn.epicgames.com) in your web browser (IE recommended for stability).
 1. Log in using your VPN credentials provided by your technical contact.
 1. Follow any on-screen instructions and prompts.
 1. The VPN client is installed and the connection is established.
@@ -121,7 +117,7 @@ You can see an example of this workflow here:
 
 ![Perforce_34.jpg](Perforce_34.jpg)(w:665 h:443)
 
-The more frequently you update, the less likely you are to run into serious conflicts when merging. It can be extremely useful to use mege programs such as Araxis to help resolve conflicts when they do arise.
+The more frequently you update, the less likely you are to run into serious conflicts when merging. It can be extremely useful to use merge programs such as Araxis to help resolve conflicts when they do arise.
 
 Here are some tips for preparing to merge a new QA build:
 
