#UE4doc #docs: ##Overview removed as according to the instructional guide. Fixed typo: WareGame -> WarGame

diff --git a/Engine/Documentation/Source/GettingStarted/DownloadingUnrealEngine/Perforce/Integration/Integration.INT.udn b/Engine/Documentation/Source/GettingStarted/DownloadingUnrealEngine/Perforce/Integration/Integration.INT.udn
--- a/Engine/Documentation/Source/GettingStarted/DownloadingUnrealEngine/Perforce/Integration/Integration.INT.udn
+++ b/Engine/Documentation/Source/GettingStarted/DownloadingUnrealEngine/Perforce/Integration/Integration.INT.udn
@@ -5,10 +5,6 @@ Description:Commonly aksed questions about Perforce terminology and usage detail
 
 [TOC(start:2)]
 
-
-
-## Overview
-
 This document provides an understanding of many of the common uses of Perforce with Unreal Engine 4, as well as related terminology.
 
 _[Note: To help simplify the merge process, an auto-merge tool has been created that can be integrated into P4V and P4Win. It automates several of the tedious steps and best practices described in this doc. For more information, see the [Perforce AutoMerge](GettingStarted/DownloadingUnrealEngine/Perforce/AutoMerge) page for more information.]_
@@ -183,7 +179,7 @@ When performing a NULL merge, you should mention the changelist number in your c
 
 Branchspecs should be set up so the normal flow of change is the default direction. Generally this is done when it is time to merge the branch and deprecate it.
 
-Perforce supports "reversing" a branchspec during an integration:
+Perforce supports "reversing" a branchspec during integration:
 
 * In P4V, you can simply click on the green arrow in the _Select branch specification_ section. See the [screenshot](#IntegrationShot) in the [integration item](#IntegrateSection). 
     * There is even text there in case you forget. :)
@@ -203,7 +199,7 @@ Typically when this happens, Perforce will do either one of two things:
 		1. It is error-prone. Fixing the branchspec when the problem is found will make sure you or anyone after you will not accidentally (re)add those files. 
 		1. It does not play nice with the Perforce integration database because Perforce will always think that there are still files in that changelist that need to be integrated.
 
-For example, The Delta-To-Main branch (Gears PC) had a WarGame folder that only Delta needed, and WareGame was removed from Main. When trying to integrate a change that included a file in WarGame, Perforce would complain that it required the -d or -dt option to integrate "around deleted revisions." To fix this, WarGame was removed from the Delta-To-Main branchspec so Perforce would no longer try to integrate those files.
+For example, The Delta-To-Main branch (Gears PC) had a WarGame folder that only Delta needed, and WarGame was removed from Main. When trying to integrate a change that included a file in WarGame, Perforce would complain that it required the -d or -dt option to integrate "around deleted revisions." To fix this, WarGame was removed from the Delta-To-Main branchspec so Perforce would no longer try to integrate those files.
 
  
 
