#UE4doc #docs: ##Overview removed as according to the instructional guide. Fixed typo: recognizeable -> recognizable

diff --git a/Engine/Documentation/Source/GettingStarted/DownloadingUnrealEngine/Perforce/P4V/P4V.INT.udn b/Engine/Documentation/Source/GettingStarted/DownloadingUnrealEngine/Perforce/P4V/P4V.INT.udn
--- a/Engine/Documentation/Source/GettingStarted/DownloadingUnrealEngine/Perforce/P4V/P4V.INT.udn
+++ b/Engine/Documentation/Source/GettingStarted/DownloadingUnrealEngine/Perforce/P4V/P4V.INT.udn
@@ -5,10 +5,6 @@ Description:Guide to installing, setting up, and syncing builds with the Perforc
 
 [TOC(start:2)]
 
-
-
-## Overview
-
 This document covers the steps in setting up the P4V client to connect to the Unreal Engine 4 depot and get a build of the engine.
 
 [REGION:note]
@@ -91,7 +87,7 @@ Press the ![connect_ok_button.jpg](connect_ok_button.jpg) to open the connection
 Workspaces are created on a per-machine basis, because they keep track of the exact revision of each file that has been stored on that machine's hard drive in the selected directory. Thus, if you ever need to start syncing from a different PC, you will need to repeat this process to create a new workspace on that machine.
 
 [REGION:note]
-If you are an anonymous user, you will want to create a workspace that is named: UE4_Anonymous_#  where the # matches your Anonymous userid (e.g. UE4_Anonymous_8). Additionally, you will not want to put any irl names in any of the fields. And you may not want to have a machine name that is recognizeable. All of those fields can be seen by other users if they search the users or workspace listings.
+If you are an anonymous user, you will want to create a workspace that is named: UE4_Anonymous_#  where the # matches your Anonymous userid (e.g. UE4_Anonymous_8). Additionally, you will not want to put any url names in any of the fields. And you may not want to have a machine name that is recognizable. All of those fields can be seen by other users if they search the users or workspace listings.
 [/REGION]
 
 #### Troubleshooting
