#UE4doc #docs: Fixed grammar error: your are -> you are. ##Overview removed as according to the instructional guide.

diff --git a/Engine/Documentation/Source/GettingStarted/DownloadingUnrealEngine/Perforce/Setup/Setup.INT.udn b/Engine/Documentation/Source/GettingStarted/DownloadingUnrealEngine/Perforce/Setup/Setup.INT.udn
--- a/Engine/Documentation/Source/GettingStarted/DownloadingUnrealEngine/Perforce/Setup/Setup.INT.udn
+++ b/Engine/Documentation/Source/GettingStarted/DownloadingUnrealEngine/Perforce/Setup/Setup.INT.udn
@@ -5,10 +5,6 @@ Description:Guide to connecting to Epic's Perforce server to obtain Unreal Engin
 
 [TOC(start:2)]
 
-
-
-## Overview
-
 Epic Games makes QA-approved builds of Unreal Engine 4, as well as other specialized code drops, available to 
 licensees through a Perforce depot that can be accessed externally. This is the method licensees use to first 
 obtain the engine as well as update to new versions as they are released and it is deemed appropriate by the 
@@ -107,7 +103,7 @@ data from your local proxy instead of Epic.
 
 
 In order to use a proxy configuration, you must request it by making a request on the [Unreal Developer Network](http://udn.unrealengine.com). 
-When this access is granted, your are only granted two VPN connections and you only have access to the primary depot, 
+When this access is granted, you are only granted two VPN connections and you only have access to the primary depot, 
 you will lose your access to the licensee proxy. 
 
 You can obtain the proxy application (P4P) from Perforce at http://perforce.com/perforce/downloads/. During 
