#UE4doc #docs: ##Overview removed as according to the instructional guide.

diff --git a/Engine/Documentation/Source/GettingStarted/DownloadingUnrealEngine/Perforce/Syncing/Syncing.INT.udn b/Engine/Documentation/Source/GettingStarted/DownloadingUnrealEngine/Perforce/Syncing/Syncing.INT.udn
--- a/Engine/Documentation/Source/GettingStarted/DownloadingUnrealEngine/Perforce/Syncing/Syncing.INT.udn
+++ b/Engine/Documentation/Source/GettingStarted/DownloadingUnrealEngine/Perforce/Syncing/Syncing.INT.udn
@@ -5,10 +5,6 @@ Description:Guide to obtaining builds of Unreal Engine 4 using Perforce.
 
 [TOC(start:2)]
 
-
-
-## Overview
-
 The two main ways sets of file revisions are identified are through labels and changelists. On a high level, a label identifies all the changes that make up a specific build of the engine, e.g. a QA-approved build, while a changelist identifies a specific set of changes. The actual differences between these two concepts is explained in much more detail on the [Codelines and Branching](http://www.perforce.com/perforce/doc.current/manuals/p4guide/06_codemgmt.html) page of the Perforce website. For the purposes of syncing to Epic's Perforce depot, here is what you need to know:
 
 
