#UE4doc #docs: ##Overview removed as according to the instructional guide.

diff --git a/Engine/Documentation/Source/GettingStarted/DownloadingUnrealEngine/VPNSetup/VPNSetup.INT.udn b/Engine/Documentation/Source/GettingStarted/DownloadingUnrealEngine/VPNSetup/VPNSetup.INT.udn
--- a/Engine/Documentation/Source/GettingStarted/DownloadingUnrealEngine/VPNSetup/VPNSetup.INT.udn
+++ b/Engine/Documentation/Source/GettingStarted/DownloadingUnrealEngine/VPNSetup/VPNSetup.INT.udn
@@ -5,10 +5,6 @@ Description:Guide to connecting to Epic's virtual private network (VPN) in order
 
 [TOC(start:2)]
 
-
-
-## Overview
-
 This document will guide you through the steps in setting up a secure Virtual Private Network (VPN) Connection to Epic Games so that you will be able to connect to the Perforce depot. Once you have a VPN Connection set up, you can [set up Perforce](GettingStarted/DownloadingUnrealEngine/Perforce/Setup).
 
 Epic has set aside the **2nd and 4th Sunday of the month from 8am-Noon EST** for performing scheduled maintenance. Please keep this in mind when planning as it is possible that any or all of the support outlets, especially network/server resources, could be affected during these times.
