#UE4doc #docs:
Removed Overview header
Updated Engine version

diff --git a/Engine/Documentation/Source/Programming/Online/Interfaces/Session/Session.INT.udn b/Engine/Documentation/Source/Programming/Online/Interfaces/Session/Session.INT.udn
--- a/Engine/Documentation/Source/Programming/Online/Interfaces/Session/Session.INT.udn
+++ b/Engine/Documentation/Source/Programming/Online/Interfaces/Session/Session.INT.udn
@@ -2,11 +2,10 @@ Availability:Public
 Title: Sessions and Matchmaking
 Crumbs:%ROOT%, Programming, Programming/Online
 Description:System for setting up advertisement of matches and handling matchmaking of players
-Version: 4.5
+Version: 4.9
 
 [TOC(start:2)]
 
-## Overview
 
 Matchmaking is the process of matching _players_ with _sessions_. A _session_ is basically an instance of 
 the game running on the server with a given set of properties which is either _advertised_ so that it can 
