#UE4doc #docs:
Removed Overview header
Updated Engine version
Added a few note regions
Bolded a term

diff --git a/Engine/Documentation/Source/Programming/Online/Steam/OnlineSubsystemSteam.INT.udn b/Engine/Documentation/Source/Programming/Online/Steam/OnlineSubsystemSteam.INT.udn
--- a/Engine/Documentation/Source/Programming/Online/Steam/OnlineSubsystemSteam.INT.udn
+++ b/Engine/Documentation/Source/Programming/Online/Steam/OnlineSubsystemSteam.INT.udn
@@ -2,7 +2,7 @@ Availability:Public
 Title: Online Subsystem Steam
 Crumbs:%ROOT%, Programming
 Description:Information specific to the Steam Online Subsystem 
-Version: 4.5
+Version: 4.9
 
 [TOC(start:2)]
 
@@ -11,9 +11,9 @@ Version: 4.5
 Information contained here pertains specifically to the Steam implementation.   Any additional setup steps, tips & tricks, and workarounds can be found here.
 [/EXCERPT:Intro]
 
-## Overview
 
-The Steam module implements most of the interfaces exposed through the online subsystem and supports most of what Valve offers through the Steamworks SDK. 
+
+The **Steam** module implements most of the interfaces exposed through the online subsystem and supports most of what Valve offers through the Steamworks SDK. 
 
 Steam Interfaces
 * Matchmaking (Lobbies / GameServer APIs)
@@ -55,15 +55,20 @@ Using Steam against the precompiled version of the engine should only require co
 	* /YourUnrealEnginePath/EngineOrGameFolder/Binaries/Mac/YourGame.app/Contents/MacOS
 		* libsteam_api.dylib (from /redistributable_bin/osx32 - single dylib has both 32 and 64 bit support)
 		
-*Note:* if you are compiling the entire engine, you will want to modify the following line in OnlineSubsystemSteam\Private\OnlineSubsystemSteamPrivatePCH.h 
+[REGION:note]
+If you are compiling the entire engine, you will want to modify the following line in OnlineSubsystemSteam\Private\OnlineSubsystemSteamPrivatePCH.h 
 
 define STEAM_SDK_VER TEXT("Steamv130")
 
 to make sure it references the new SDK directory location / version
+[/REGION]
 
 ### INI Configuration
 
-Turn on some settings in the game's `DefaultEngine.ini`.  *Note:* The SteamDevAppId of 480 is Valve's test app id, shared by everyone.   You will need your own app id eventually, but most features of Steam should work before then.
+Turn on some settings in the game's `DefaultEngine.ini`.  
+[REGION:note]
+The SteamDevAppId of 480 is Valve's test app id, shared by everyone.   You will need your own app id eventually, but most features of Steam should work before then.
+[/REGION]
 
     [/Script/Engine.GameEngine]
     +NetDriverDefinitions=(DefName="GameNetDriver",DriverClassName="OnlineSubsystemSteam.SteamNetDriver",DriverClassNameFallback="OnlineSubsystemUtils.IpNetDriver")
