#UE4doc #docs: ##Overview removed as according to the instructional guide. Altered images to match the instructional guide: green -> yellow. Also changed multiple spelling and grammar errors: labelled -> labeled, its -> it's, etc.

diff --git a/Engine/Documentation/Source/GettingStarted/FromUnity/UnrealEngineforUnityDevelopers.INT.udn b/Engine/Documentation/Source/GettingStarted/FromUnity/UnrealEngineforUnityDevelopers.INT.udn
--- a/Engine/Documentation/Source/GettingStarted/FromUnity/UnrealEngineforUnityDevelopers.INT.udn
+++ b/Engine/Documentation/Source/GettingStarted/FromUnity/UnrealEngineforUnityDevelopers.INT.udn
@@ -8,14 +8,11 @@ Description: Translate your Unity knowledge into UE4 so you can get up to speed
 
 [TOC(start:2 end:2)]
 
-
-## Overview
-
 This guide provides an overview of UE4 from the viewpoint of a Unity user, and aims to help you translate your existing Unity experience into the world of UE4.
 
 ## The Editor
 
-Below are pictures of the Unity Editor and Unreal Editor, **color-coded** to indicate common functionality. Each block is labelled to show the equivalent UE4 terminology.  Unreal Editor's layout is fully customizable by dragging and dropping tabs around.
+Below are pictures of the Unity Editor and Unreal Editor, **color-coded** to indicate common functionality. Each block is labeled to show the equivalent UE4 terminology.  Unreal Editor's layout is fully customizable by dragging and dropping tabs around.
 
 ![](image_1.png)
 
@@ -249,7 +246,7 @@ Most projects will use a mix of Blueprints and C++.  Many developers prototype g
 
 ### Blueprint Classes can extend C++ Classes
 
-Much of the magic of UE4 game development comes from the interplay between programmers implementing new features in C++, with designers and artists utilizing them in Blueprints and then asking for more!  Here's how a team might structure a UE4-based shooter game implementing pickups using a mix of C++ Classes systems programming, and Blueprint Classes for behaviour and appearance:
+Much of the magic of UE4 game development comes from the interplay between programmers implementing new features in C++, with designers and artists utilizing them in Blueprints and then asking for more!  Here's how a team might structure a UE4-based shooter game implementing pickups using a mix of C++ Classes systems programming, and Blueprint Classes for behavior and appearance:
 
 ![image alt text](image_9.png)
 
@@ -281,7 +278,7 @@ There are several sub-classed Actor types you will quickly encounter in the engi
 
 * **Pawn** - A type of Actor representing a controllable game object, typically the player's avatar.  Pawns are moved by Players and AI alike through an owning Controller.
 
-* **Character** - A more specialized version of Pawn designed for biped avatars that handles a lot of the complexity of those kinds of game objects.
+* **Character** - A more specialized version of Pawn designed for biped avatars that handle a lot of the complexity of those kinds of game objects.
 
 * **Controller** - Possesses and controls a Pawn.  By separating the Pawn from the controller, you can write AI Controllers that can manipulate a pawn using the same interface the player would.
 
@@ -299,13 +296,13 @@ Okay, so this is where things get a little crazy (in sort of an awesome way.)  U
 
 However Unreal has this extra layer on top called the [Gameplay Framework](Gameplay/Framework).  This does not really exist in Unity at all.  You do not have to use it in Unreal, but it is actually super cool!  Basically, if you make use of some basic primitive classes and follow certain established conventions, your game will acquire some awesome features automatically that would be otherwise very difficult to implement or retrofit (such as full multiplayer support!)
 
-Countless awesome games have been designed on top of Unreal's Gameplay Framework, and its worth taking some time to understand how it works.  Yes, you could roll your own version of this.  That is totally fine if you want to!   But hundreds of awesome Unreal developers have made great use of this framework in UE4 today, so it is worth taking some time to learn.
+Countless awesome games have been designed on top of Unreal's Gameplay Framework, and it is worth taking some time to understand how it works.  Yes, you could roll your own version of this.  That is totally fine if you want to!   But hundreds of awesome Unreal developers have made great use of this framework in UE4 today, so it is worth taking some time to learn.
 
 To use the Gameplay Framework, you really just need to learn the customized built-in Actor classes included with Unreal, such as **Pawn**, **Character**, and **Player Controller**, and eventually learn about how Unreal replication and networking features work.  For now, let us get back to the basics though.
 
 ## How to write code in UE4
 
-### So, I'm used to programming in MonoDevelop�
+### So, I'm used to programming in MonoDevelop
 
 For Blueprint scripting, you will only need Unreal Editor -- everything is built in!  To write code in C++, download the [free version of Visual Studio](http://go.microsoft.com/fwlink/?LinkId=525621) on Windows, or [install Xcode](https://itunes.apple.com/us/app/xcode/id497799835?mt=12) on Mac.  When you first create a new project (or add code to an existing project), UE4 will **automatically create Visual Studio project files **for you.  You can open Visual Studio by **double-clicking** on a C++ class inside **Content Browser**, or by clicking the "Open Visual Studio" button in the main File menu.
 
@@ -317,7 +314,7 @@ One important difference in UE4:  You will sometimes have to manually refresh yo
 
 ### Writing Event Functions (Start, Update, etc.)
 
-If you are used to working with MonoBehaviours you are used to such methods as Start, Update, and OnDestroy.  Here is a comparison between a Unity behaviour and its counterpart in UE4 Actors and components.
+If you are used to working with MonoBehaviors you are used to such methods as Start, Update, and OnDestroy.  Here is a comparison between a Unity behavior and its counterpart in UE4 Actors and components.
 
 In Unity, we might have a simple component that looks like this:
 
