#UEDoc #docs - bold lettering where needed

diff --git a/Engine/Documentation/Source/Programming/UnrealArchitecture/Reference/Classes/Specifiers/AdvancedClassDisplay/ClassSpecifier_AdvancedClassDisplay.INT.udn b/Engine/Documentation/Source/Programming/UnrealArchitecture/Reference/Classes/Specifiers/AdvancedClassDisplay/ClassSpecifier_AdvancedClassDisplay.INT.udn
--- a/Engine/Documentation/Source/Programming/UnrealArchitecture/Reference/Classes/Specifiers/AdvancedClassDisplay/ClassSpecifier_AdvancedClassDisplay.INT.udn
+++ b/Engine/Documentation/Source/Programming/UnrealArchitecture/Reference/Classes/Specifiers/AdvancedClassDisplay/ClassSpecifier_AdvancedClassDisplay.INT.udn
@@ -5,7 +5,7 @@ Availability:Public
 Version: 4.5
 
 The **AdvancedClassDisplay** class specifier forces all the properties of the class to only show in 
-the [advanced sections](Engine/UI/LevelEditor/Details#AdvancedProperties) in the **Details** panel, 
+the **[advanced sections](Engine/UI/LevelEditor/Details#AdvancedProperties)** in the **Details** panel, 
 and hidden from view by default.
 
 [REGION:tip]
