#UEDoc #docs - adjusted bold lettering

diff --git a/Engine/Documentation/Source/Programming/UnrealArchitecture/Reference/Classes/Specifiers/ClassGroup/ClassSpecifier_ClassGroup.INT.udn b/Engine/Documentation/Source/Programming/UnrealArchitecture/Reference/Classes/Specifiers/ClassGroup/ClassSpecifier_ClassGroup.INT.udn
--- a/Engine/Documentation/Source/Programming/UnrealArchitecture/Reference/Classes/Specifiers/ClassGroup/ClassSpecifier_ClassGroup.INT.udn
+++ b/Engine/Documentation/Source/Programming/UnrealArchitecture/Reference/Classes/Specifiers/ClassGroup/ClassSpecifier_ClassGroup.INT.udn
@@ -5,6 +5,6 @@ Availability:Public
 Version: 4.5
 
 Indicates that Unreal Editor's **Actor Browser** should include this class and any subclass of this 
-class within the specified *GroupName* when **Group View** is enabled in the **Actor Browser**.
+class within the specified *GroupName* when **Group View** is enabled in the Actor Browser.
 
 ClassGroup=GroupName
