UEDOC-1828
Updated GIF to Youtube video, minor content fix.

diff --git a/Engine/Documentation/Source/Engine/Blueprints/BP_HowTo/WorkingWithArrays/WorkingWithArrays.INT.udn b/Engine/Documentation/Source/Engine/Blueprints/BP_HowTo/WorkingWithArrays/WorkingWithArrays.INT.udn
--- a/Engine/Documentation/Source/Engine/Blueprints/BP_HowTo/WorkingWithArrays/WorkingWithArrays.INT.udn
+++ b/Engine/Documentation/Source/Engine/Blueprints/BP_HowTo/WorkingWithArrays/WorkingWithArrays.INT.udn
@@ -175,7 +175,17 @@ _For this example, we are using the **Blueprint 3rd Person Template** project wi
 
 1.	Close the Blueprint, then click the **Play** button to play in the Editor. 
 
-	![](ExampleGIF.gif)(convert:false)
+	[OBJECT:EmbeddedVideo]
+	[PARAMLITERAL:width]
+	640
+	[/PARAMLITERAL]
+	[PARAMLITERAL:height]
+	360
+	[/PARAMLITERAL]
+	[PARAMLITERAL:videoid]
+	r-zHq9feZr8
+	[/PARAMLITERAL]
+	[/OBJECT]
 
 	When you play in the Editor, each time that you press the **F** key a random Blueprint (from the ones you added to the Array) will be selected and spawned at a random location from the Array of Target Points. Instead of the **F** Key Event to spawn the items, and instead of effects that are spawned, you could use this same logic to set up a random spawn location for Health or Ammo pickups or even loot the player can obtain. 
 
