#UE4doc #docs: Formatting modified to match that of the How To section of the Instructional Guide: Title of page and title of table.

diff --git a/Engine/Documentation/Source/GettingStarted/HowTo/IntroHowToOverview.INT.udn b/Engine/Documentation/Source/GettingStarted/HowTo/IntroHowToOverview.INT.udn
--- a/Engine/Documentation/Source/GettingStarted/HowTo/IntroHowToOverview.INT.udn
+++ b/Engine/Documentation/Source/GettingStarted/HowTo/IntroHowToOverview.INT.udn
@@ -42,7 +42,7 @@ If you are new to the **Unreal Engine 4** Editor, this page aims to get you quic
 
 Refer to the table below for more information and click each link to view its contents. 
 
-|How To... | Description |
+|How To | Description |
 | ------- | ------------ |
 | **[Create/Open Projects](Engine/Basics/Projects/Browser)** | %Engine/Basics/Projects/Browser:description% |
 | **[](Engine\Actors\Placement)** | %Engine\Actors\Placement:description% |
