UEDOC-1836

Updated GIF to Youtube video

diff --git a/Engine/Documentation/Source/Engine/Blueprints/HowTo/BlueprintHowTo.INT.udn b/Engine/Documentation/Source/Engine/Blueprints/HowTo/BlueprintHowTo.INT.udn
--- a/Engine/Documentation/Source/Engine/Blueprints/HowTo/BlueprintHowTo.INT.udn
+++ b/Engine/Documentation/Source/Engine/Blueprints/HowTo/BlueprintHowTo.INT.udn
@@ -45,7 +45,17 @@ Description:A sample project created using different methods of Blueprint Commun
 
 In this **Blueprint Communication Project**, we will create a simple shooting gallery style game. Over the course of the project you will learn how to communicate information from the level to a character as well as how to pass character information to a HUD. At the end of the project, you will have something similar to below:
 
-![](Engine\Blueprints\HowTo\BPHT_6\GameplayGIF.gif)(convert:false)
+[OBJECT:EmbeddedVideo]
+[PARAMLITERAL:width]
+640
+[/PARAMLITERAL]
+[PARAMLITERAL:height]
+360
+[/PARAMLITERAL]
+[PARAMLITERAL:videoid]
+34HO1gCSJYc
+[/PARAMLITERAL]
+[/OBJECT]
 
 In this short game example, we are spawning a destructible Static Mesh from inside our Character Blueprint when the player presses a button. The location where the rock is spawned and launched from is randomly selected from spawn points we placed in the level, so we must communicate the spawn locations from the level to our Character Blueprint. 
 We are also taking information stored in our Character Blueprint (the amount of pulls or rocks left to launch, number of shots left per pull and the player's score) and passing that information to a HUD to display it onscreen. 
