#UEDoc #docs - bold lettering.

diff --git a/Engine/Documentation/Source/Programming/UnrealArchitecture/Reference/Functions/Functions.INT.udn b/Engine/Documentation/Source/Programming/UnrealArchitecture/Reference/Functions/Functions.INT.udn
--- a/Engine/Documentation/Source/Programming/UnrealArchitecture/Reference/Functions/Functions.INT.udn
+++ b/Engine/Documentation/Source/Programming/UnrealArchitecture/Reference/Functions/Functions.INT.udn
@@ -10,14 +10,14 @@ Version: 4.5
 
 Functions can exist in two basic forms: a regular C++ function or a `UFunction`. Both C++ functions and `UFunctions` are declared in the `.h` class header file. C++ functions are declared using the 
 standard C++ syntax for declaring a function. `UFunctions` use a specialized syntax that allows additional information about the function to be specified in the declaration through the use of 
-[Function Specifiers](#FunctionSpecifiers).
+**[Function Specifiers](#FunctionSpecifiers)**.
 
 	UFUNCTION([specifier, specifier, ...], [meta(key=value, key=value, ...)])
 	ReturnType FunctionName([Parameter, Parameter, ...])
 
 `UFunctions` behave the same as C++ functions in that they have a C++ implementation, can be called from within C++ code, and can contain calls to other C++ functions or `UFunctions` within their 
-function bodies. `UFunctions` differ in a few areas, however. For instance, they can be called or overridden from within the [Blueprint Visual Scripting](Engine/Blueprints) system. Any 
-`UFunction` declared using the `BlueprintCallable`, `BlueprintImplementableEvent`, or `BlueprintPure` specifiers (see [Function Specifiers](#FunctionSpecifiers) below for more details) will be exposed 
+function bodies. `UFunctions` differ in a few areas, however. For instance, they can be called or overridden from within the **[Blueprint Visual Scripting](Engine/Blueprints)** system. Any 
+`UFunction` declared using the `BlueprintCallable`, `BlueprintImplementableEvent`, or `BlueprintPure` specifiers (see **[Function Specifiers](#FunctionSpecifiers)** below for more details) will be exposed 
 to Blueprints. `UFunctions` can also be assigned as delegates within the default properties of a class. This technique is commonly used to bind input actions or axes to functions in `Input` classes.
  `UFunctions` are also used as replication callbacks; meaning the `UFunction` will be called when the variable it is associated with changes and is replicated across the network. `UFunctions` are 
  also the only type of function that can be declared as an exec function allowing them to be called by the player from the in-game console during play.
@@ -32,7 +32,10 @@ to Blueprints. `UFunctions` can also be assigned as delegates within the default
 $ Out : 
 Declares the parameter as being passed by reference allowing it to be modified by the function.		
 $ Optional : 
-With the optional keyword, you can make certain function parameters optional, as a convenience to the caller. The values for optional parameters which the caller does not specify depend on the function. For example, the `SpawnActor` function takes an optional location and rotation, which default to the location and rotation of the spawning Actor's root component. The default value of optional arguments can be specified by adding `= [value]`. For example: `function myFunc(optional int x = -1)`. In most cases, the default value for the type of variable, or zero (e.g. 0, false, "", none), is used when no values is passed to an optional parameter.
+With the optional keyword, you can make certain function parameters optional, as a convenience to the caller. The values for optional parameters which the caller does not specify depend on the function. 
+For example, the `SpawnActor` function takes an optional location and rotation, which default to the location and rotation of the spawning Actor's root component. The default value of optional arguments 
+can be specified by adding `= [value]`. For example: `function myFunc(optional int x = -1)`. In most cases, the default value for the type of variable, or zero (e.g. 0, false, "", none), is used when 
+no values is passed to an optional parameter.
 
 
 ## Function Implementations
@@ -43,10 +46,10 @@ Functions declared in the class header file are given definitions in the class s
 
 [INCLUDE:Programming/UnrealArchitecture/Delegates#Overview]
 
-See the [](Programming/UnrealArchitecture/Delegates) page for reference and usage information.
+See the **[](Programming/UnrealArchitecture/Delegates)** page for reference and usage information.
 
 ## Timers
 
 [INCLUDE:Programming/UnrealArchitecture/Timers#Overview]
 
-See the [](Programming\UnrealArchitecture\Timers) page for reference and usage information.
+See the **[](Programming\UnrealArchitecture\Timers)** page for reference and usage information.
