#UE4doc #docs: Removed extraneous bolding as according to the instructional guide. ##Overview removed as according to the instructional guide. Altered images to match the instructional guide: green -> yellow.

diff --git a/Engine/Documentation/Source/GettingStarted/HowTo/DetailsPanel/DetailsPanel.INT.udn b/Engine/Documentation/Source/GettingStarted/HowTo/DetailsPanel/DetailsPanel.INT.udn
--- a/Engine/Documentation/Source/GettingStarted/HowTo/DetailsPanel/DetailsPanel.INT.udn
+++ b/Engine/Documentation/Source/GettingStarted/HowTo/DetailsPanel/DetailsPanel.INT.udn
@@ -17,7 +17,7 @@ For this How To guide, we have created a new project using the **Blueprint Third
 
 ## Editing Properties
 
-Over the course of the guide we will change several properties that will affect an Actor in our level such as scaling the Actor to increase its size, applying a **Material**, turning on **Physics** for the Actor, as well as tweaking **Collision Settings** so that it registers when it hits other Actors during gameplay. 
+Over the course of the guide we will change several properties that will affect an Actor in our level such as scaling the Actor to increase its size, applying a Material, turning on Physics for the Actor, as well as tweaking Collision Settings so that it registers when it hits other Actors during gameplay. 
 
 ### Scaling Properties
 
@@ -63,7 +63,7 @@ Next, let us apply a new **Material** to our Actor instead of its default one.
 
 	![](Details6.png)
 
-	When you select a **Material**, the dropdown box will be updated with your selection and the sphere in the level will be updated. 
+	When you select a Material, the dropdown box will be updated with your selection and the sphere in the level will be updated. 
 
 1.	Click the magnifying glass icon next to the **Select Material** dropdown box. 
 
