UEDOC-1877

Updated GIF to Youtube video

diff --git a/Engine/Documentation/Source/Engine/Actors/PlayerStart/PlayerStart.INT.udn b/Engine/Documentation/Source/Engine/Actors/PlayerStart/PlayerStart.INT.udn
--- a/Engine/Documentation/Source/Engine/Actors/PlayerStart/PlayerStart.INT.udn
+++ b/Engine/Documentation/Source/Engine/Actors/PlayerStart/PlayerStart.INT.udn
@@ -12,7 +12,17 @@ The ability to spawn a player into the world at any location is a very useful fe
 
 Player Start Actors can be found in the Mode panel under the Basic category. To place one into the world, simply drag it from the **Modes** panel into the game world.
 
-![](Place_A_Player_Start.gif)(convert:false)
+[OBJECT:EmbeddedVideo]
+[PARAMLITERAL:width]
+640
+[/PARAMLITERAL]
+[PARAMLITERAL:height]
+360
+[/PARAMLITERAL]
+[PARAMLITERAL:videoid]
+DeMIl0BgRtw
+[/PARAMLITERAL]
+[/OBJECT]
 
 ## Using Player Start Actor
 
@@ -34,4 +44,14 @@ While using Player Starts is quite simple, there are a few tips and tricks about
 
 ** Bad Size:** Sometimes your Player Start might change from an icon of a controller to an icon that says Bad Size. If this happens, all you need to do is move your Player Start around in the world until it is not intersecting anything in the scene.
 
-![](Player_Start_Bad_Size.gif)
+[OBJECT:EmbeddedVideo]
+[PARAMLITERAL:width]
+640
+[/PARAMLITERAL]
+[PARAMLITERAL:height]
+360
+[/PARAMLITERAL]
+[PARAMLITERAL:videoid]
+9wolC00NWzo
+[/PARAMLITERAL]
+[/OBJECT]
