UEDOC-1881

Updated GIF to Youtube video

diff --git a/Engine/Documentation/Source/Engine/Actors/TargetPoint/TargetPoint.INT.udn b/Engine/Documentation/Source/Engine/Actors/TargetPoint/TargetPoint.INT.udn
--- a/Engine/Documentation/Source/Engine/Actors/TargetPoint/TargetPoint.INT.udn
+++ b/Engine/Documentation/Source/Engine/Actors/TargetPoint/TargetPoint.INT.udn
@@ -12,7 +12,17 @@ When creating worlds for games you sometimes need the ability to place and spawn
 ## Placing Target Point Actors
 Target Point Actors can be found in the **Modes** panel under the **All Classes** category. To add one to the world, simply select it from the **Modes** panel and then drag it into the world. 
 
-![](Using_Target_Point_Actor.gif)
+[OBJECT:EmbeddedVideo]
+[PARAMLITERAL:width]
+640
+[/PARAMLITERAL]
+[PARAMLITERAL:height]
+360
+[/PARAMLITERAL]
+[PARAMLITERAL:videoid]
+6M5rxn1pWV8
+[/PARAMLITERAL]
+[/OBJECT]
 
 ## Using Target Points
 Target Point Actors can be used for all kinds of things inside Unreal Engine 4. Here is a list of some of the things that you can use Target Point Actors for.
