#UE4doc #docs: Removed extraneous bolding as according to the instructional guide. Altered images to match the instructional guide: green -> yellow. Made minor spelling and grammar changes.

diff --git a/Engine/Documentation/Source/GettingStarted/HowTo/FindingHelp/FindingHelp.INT.udn b/Engine/Documentation/Source/GettingStarted/HowTo/FindingHelp/FindingHelp.INT.udn
--- a/Engine/Documentation/Source/GettingStarted/HowTo/FindingHelp/FindingHelp.INT.udn
+++ b/Engine/Documentation/Source/GettingStarted/HowTo/FindingHelp/FindingHelp.INT.udn
@@ -15,19 +15,19 @@ In this section, we will take a look at the various methods in which you can acc
 
 ### Rich Tooltips
 
-The first method you can take to learn more about an aspect of the Editor is through the Editor's **Rich Tooltips**. 
+The first method you can take to learn more about an aspect of the Editor is through the Editor's Rich Tooltips. 
 
 Whenever you highlight something inside the Editor's interface, a tooltip will appear along with the option to expand upon that tooltip. 
 
 ![](RTT1.png)
 
-Above we have highlighted the **Play** button from the Main Toolbar and a description of what the button does is displayed in the tooltip. In the lower right corner of the tooltip window, you will see the **Hold (Ctrl + Alt) for more** text, indicating that this tooltip has a **Rich Tooltip** as well. 
+Above we have highlighted the **Play** button from the Main Toolbar and a description of what the button does is displayed in the tooltip. In the lower right corner of the tooltip window, you will see the **Hold (Ctrl + Alt) for more** text, indicating that this tooltip has a Rich Tooltip as well. 
 
 Press **Ctrl+Alt** to expand the tooltip to learn more about the action you highlighted. 
 
 ![](RTT2.png)
 
-You can see now that the tooltip has been expanded for a more detailed explanation of what the **Play** action does and there is a **see full documentation** link at the bottom of the Rich Tooltip. Clicking on the see full documentation text will take you directly to the external [Unreal Engine 4 Documentation](#externalhelp:documentation) page related to and covering what the **Play** (or whatever you highlighted) action does and how to use it. 
+You can see now that the tooltip has been expanded for a more detailed explanation of what the **Play** action does and there is a see full documentation link at the bottom of the Rich Tooltip. Clicking on the see full documentation text will take you directly to the external [Unreal Engine 4 Documentation](#externalhelp:documentation) page related to and covering what the **Play** (or whatever you highlighted) action does and how to use it. 
 
 ### Help Icon
 
@@ -75,15 +75,15 @@ When clicking the **Help** button, several options are available in the dropdown
 
 | Control | Action |
 | ------- | ------ |
-| **Documentation...** | This will take you to the main [Official Documentation](#externalhelp:documentation) page. |
-| **API Reference...** | This will open the **UE4 API Documentation** and is a low level index of Engine classes and functions.  |
+| **Documentation** | This will take you to the main [Official Documentation](#externalhelp:documentation) page. |
+| **API Reference** | This will open the **UE4 API Documentation** and is a low level index of Engine classes and functions.  |
 | **Tutorials** | Selecting this option will open the [In-Editor Tutorials](#in-editortutorials). |
-| **Forums...** | This will take you to the [Official Unreal Engine 4 Forums](#externalhelp:forums) page. |
-| **Answer Hub...** | Selecting this will take you to the main [Answer Hub](#externalhelp:answerhub) page. |
-| **Wiki...** | This option will take you to the [Official Unreal Engine 4 Wiki](#externalhelp:wiki) page. |
-| **Visit UnrealEngine.com...** | Takes you to the main **Unreal Engine** page. |
+| **Forums** | This will take you to the [Official Unreal Engine 4 Forums](#externalhelp:forums) page. |
+| **Answer Hub** | Selecting this will take you to the main [Answer Hub](#externalhelp:answerhub) page. |
+| **Wiki** | This option will take you to the [Official Unreal Engine 4 Wiki](#externalhelp:wiki) page. |
+| **Visit UnrealEngine.com** | Takes you to the main **Unreal Engine** page. |
 | **Credits** | This displays the credits. |
-| **About Unreal Editor...** | This will display information about the Unreal Engine 4 Editor. |
+| **About Unreal Editor** | This will display information about the Unreal Engine 4 Editor. |
 
 ### F1 Help
 
@@ -125,7 +125,7 @@ Inside the search bar, you can enter any topic and as you start to enter text, a
 
 ![](HelpSearch_2.png)
 
-Above we have searched for **Blueprints**, the results show [In-Editor Tutorials](#in-editortutorials) related to Blueprints as well as links to external [Official Documentation](#externalhelp:documentation) related to Blueprints. There is also a link to [Answer Hub](#externalhelp:answerhub) where you can ask a question online and get help from Epic Staff as well as the Unreal Engine community. 
+Above we have searched for Blueprints, the results show [In-Editor Tutorials](#in-editortutorials) related to Blueprints as well as links to external [Official Documentation](#externalhelp:documentation) related to Blueprints. There is also a link to [Answer Hub](#externalhelp:answerhub) where you can ask a question online and get help from Epic Staff as well as the Unreal Engine community. 
 
 Not quite sure what it is you are looking for but have an idea of what you want to achieve? You can enter that into the search bar as well. 
 
@@ -135,7 +135,7 @@ Above we are searching for **How to display Health**. Depending upon your entry,
 
 ## In-Editor Tutorials
 
-**In-Editor** Tutorials provide you with step-by-step instructions on how to navigate and use a particular aspects of Unreal Engine 4. 
+**In-Editor** Tutorials provide you with step-by-step instructions on how to navigate and use particular aspects of Unreal Engine 4. 
 
 [REGION:lightbox]
 [![](Tutorial_1.png)](Tutorial_1.png)
@@ -199,7 +199,7 @@ When you want to learn about a particular feature or aspect of the Editor, the [
 
 You can use the search field to search for the information you want to retrieve, or you can use the **Quick Search Terms** to browse commonly searched topics. Additionally, you can browse the documentation by topic. If you are new to Unreal Engine 4, **Getting Started** is a fine place to start. 
 
-For those that want a more hands-on learning experience, you can take a look at **Quick Start** or **How-To** guides which will guide you step-by-step through the creation of a sample project (for Quick Start) or quick one-off learning experiences regarding a particular topic (for How-Tos) as well as the **Samples & Tutorials** section. 
+For those that want a more hands-on learning experience, you can take a look at **Quick Start** or **How-To** guides which will guide you step-by-step through the creation of a sample project (for Quick Start) or quick one-off learning experiences regarding a particular topic (for How-to's) as well as the **Samples & Tutorials** section. 
 
 ## External Help: Answer Hub
 
@@ -228,7 +228,8 @@ Another resource for help is the [Official Unreal Engine 4 Forums](https://forum
 
 ![](Forums.png)
 
-The forums are broken up into several categories from **Animation**, **Blueprints**, and **C++ Gameplay Programming** to **Android** or **iOS** development or even **Community Content**, **Tools**, and **Tutorials**. The forums are largely community driven but Epic Staff does contribute to and respond to posts in the forums quite a bit. While **Answer Hub** is geared toward one-off questions and answers, the forums is a good place to start a discussion on a particular topic or how to approach creating something in Unreal Engine 4. 
+The forums are broken up into several categories from Animation, Blueprints, and C++ Gameplay Programming to Android or iOS development or even Community Content, Tools, and Tutorials. The forums are largely community driven but Epic Staff does contribute to and respond to posts in the forums quite a bit. While Answer Hub is geared toward one-off questions and answers, the forums is a good place to start a discussion on a particular topic or how to 
+approach creating something in Unreal Engine 4. 
 
 
 ## External Help: WIKI
