#UEDoc #docs - bolded a link

diff --git a/Engine/Documentation/Source/Programming/UnrealArchitecture/Reference/Properties/Properties.INT.udn b/Engine/Documentation/Source/Programming/UnrealArchitecture/Reference/Properties/Properties.INT.udn
--- a/Engine/Documentation/Source/Programming/UnrealArchitecture/Reference/Properties/Properties.INT.udn
+++ b/Engine/Documentation/Source/Programming/UnrealArchitecture/Reference/Properties/Properties.INT.udn
@@ -19,14 +19,14 @@ The convention for integral data types is "int" or "uint" followed by the size i
 
 |Variable Type | Description | 
 | ------------| -------------|
-|uint8	 | 8-bit  unsigned|
-|uint16	| 16-bit unsigned|
-| uint32	| 32-bit unsigned|
-| uint64	| 64-bit unsigned|
-| int8	| 8-bit  signed|
-| int16	| 16-bit signed|
-| int32	| 32-bit signed|
-| int64 | 64-bit signed|
+| **uint8**	 | 8-bit  unsigned|
+| **uint16**	| 16-bit unsigned|
+| **uint32**	| 32-bit unsigned|
+| **uint64**	| 64-bit unsigned|
+| **int8**	| 8-bit  signed|
+| **int16**	| 16-bit signed|
+| **int32**	| 32-bit signed|
+| **int64** | 64-bit signed|
 
 ### Floating Point Types
 Unreal uses the standard C++ floating point types, float and double.
@@ -46,7 +46,7 @@ For most uses, Unreal relies on the TCHAR type for characters.  The TEXT() macro
 
 	MyDogPtr->DogName = FName(TEXT("Samson Aloysius"));
 
-For more on the three string types, when to use each one, and how to work with them, see the [String Handling documentation](Programming/UnrealArchitecture/StringHandling).
+For more on the three string types, when to use each one, and how to work with them, see the **[String Handling documentation](Programming/UnrealArchitecture/StringHandling)**.
 
 
 
