TypoCorrection

diff --git a/Engine/Documentation/Source/Engine/QuickStart/2/LDQS_2.INT.udn b/Engine/Documentation/Source/Engine/QuickStart/2/LDQS_2.INT.udn
--- a/Engine/Documentation/Source/Engine/QuickStart/2/LDQS_2.INT.udn
+++ b/Engine/Documentation/Source/Engine/QuickStart/2/LDQS_2.INT.udn
@@ -1,4 +1,4 @@
-Availability: Public
+﻿Availability: Public
 Title:2 - Navigating the Viewport
 Crumbs:%ROOT%, GettingStarted, Engine/QuickStart
 Description:Learn the different methods you can use to navigate the different viewports inside Unreal Engine 4. 
@@ -57,6 +57,6 @@ To switch to an [Orthographic View](https://wiki.unrealengine.com/Videos/Player?
 You should now have an understanding of the different methods you can use to navigate the Viewport and have found one that you are comfortable with using to move the camera around in the level. 
 You will be moving the camera around quite a bit as you are building your levels, so feel free to switch between the different methods at any point to get the precision that you need. 
 
-In the Next Step, we will create a **New Level** so that we can build or own sample area and start populating it with assets. 
+In the Next Step, we will create a **New Level** so that we can build our own sample area and start populating it with assets. 
 
 %Steps%
