UEDOC-1779

Updated GIF to Youtube video

diff --git a/Engine/Documentation/Source/Platforms/GearVR/QuickStart/2/GearVRQuickStart_2.INT.udn b/Engine/Documentation/Source/Platforms/GearVR/QuickStart/2/GearVRQuickStart_2.INT.udn
--- a/Engine/Documentation/Source/Platforms/GearVR/QuickStart/2/GearVRQuickStart_2.INT.udn
+++ b/Engine/Documentation/Source/Platforms/GearVR/QuickStart/2/GearVRQuickStart_2.INT.udn
@@ -35,7 +35,17 @@ Version: 4.8
 
 1. Once the UE4 Editor has loaded, from the **Main Toolbar**, open the **Windows** tab and select the **Plugins** menu item. 
 
-	![](GVR_Plugins_Menu.gif)(convert:false)
+	[OBJECT:EmbeddedVideo]
+	[PARAMLITERAL:width]
+	640
+	[/PARAMLITERAL]
+	[PARAMLITERAL:height]
+	360
+	[/PARAMLITERAL]
+	[PARAMLITERAL:videoid]
+	eo79kj2i2bc
+	[/PARAMLITERAL]
+	[/OBJECT]
 
 1. From the Plug-In menu go to the **Head Mounted Displays** section and **disable** the **Oculus** option and **enable** the **Gear VR** option. Then click the **Restart Now** button to re-start the editor so the changes can take place. 
 
@@ -43,7 +53,17 @@ Version: 4.8
 
 1. Once the Editor has re-started, select **Edit** from the **Main Toolbar** and then select **Project Settings**  and click the **Configure Now** button for Android support in the **APKPacking** section.
 
-	![](GVR_Android_Setup_01.gif)(convert:false)
+	[OBJECT:EmbeddedVideo]
+	[PARAMLITERAL:width]
+	640
+	[/PARAMLITERAL]
+	[PARAMLITERAL:height]
+	360
+	[/PARAMLITERAL]
+	[PARAMLITERAL:videoid]
+	If_vSzD5h3w
+	[/PARAMLITERAL]
+	[/OBJECT]
 
 1. Next change the **Minimum SDK Version** from **9** to **19**.
 
@@ -72,7 +92,17 @@ Version: 4.8
 	Remember to **Re-Enable** this option when your project is ready for its final publish to ensure that your project will work as intended on the Gear VR Headset. 
 	[/region]
 
-	![](GVR_Disable_FullReBuild.gif)(convert:false)
+	[OBJECT:EmbeddedVideo]
+	[PARAMLITERAL:width]
+	640
+	[/PARAMLITERAL]
+	[PARAMLITERAL:height]
+	360
+	[/PARAMLITERAL]
+	[PARAMLITERAL:videoid]
+	29oB5aowElQ
+	[/PARAMLITERAL]
+	[/OBJECT]
 
 	<!-- This was step six, removed based on Josh Adams advice.
 	1. Next in the **Android Package Name** section give your project a name. For this example we will be using the following name **com.Company.GearVR_QS**. 
