UEDOC-1783

Updated GIF to Youtube video

diff --git a/Engine/Documentation/Source/Engine/Landscape/Editing/Gizmos/Gizmos.INT.udn b/Engine/Documentation/Source/Engine/Landscape/Editing/Gizmos/Gizmos.INT.udn
--- a/Engine/Documentation/Source/Engine/Landscape/Editing/Gizmos/Gizmos.INT.udn
+++ b/Engine/Documentation/Source/Engine/Landscape/Editing/Gizmos/Gizmos.INT.udn
@@ -28,7 +28,17 @@ Description:Guide to the Gizmos available for modifying Landscapes.
 
 Landscape Gizmos are editor-only Actors that are similar to volumes in that they define a specified area. Their purpose is to hold height and layer data for an area of a Landscape so that it can be copied to another location on the Landscape or exported for use in a different Landscape or heightmap generator (World Machine, Terresculptor,...).
 
-![](Using_Gizom_Tool.gif)(convert:false)
+[OBJECT:EmbeddedVideo]
+[PARAMLITERAL:width]
+640
+[/PARAMLITERAL]
+[PARAMLITERAL:height]
+360
+[/PARAMLITERAL]
+[PARAMLITERAL:videoid]
+IdzHGD7iWdg
+[/PARAMLITERAL]
+[/OBJECT]
 
 ## Accessing the Gizmo Tool
 
