#UE4doc #docs: Removed extraneous bolding as according to the instructional guide. Altered images to match the instructional guide: green -> yellow.

diff --git a/Engine/Documentation/Source/GettingStarted/HowTo/PIE/PIE.INT.udn b/Engine/Documentation/Source/GettingStarted/HowTo/PIE/PIE.INT.udn
--- a/Engine/Documentation/Source/GettingStarted/HowTo/PIE/PIE.INT.udn
+++ b/Engine/Documentation/Source/GettingStarted/HowTo/PIE/PIE.INT.udn
@@ -104,7 +104,7 @@ There are different **Play Modes** that you can use when starting a play session
 
 1.	Choose **Selected Viewport**, **Mobile Preview**, **New Editor Window**, or **Standalone Game**. 
 
-The **Selected Viewport** is the default setting and launches the game within the Editor inside the currently selected [Viewport](Engine/UI/LevelEditor/Viewports/Basics). **New Editor Window** also launches the game within the Editor however it creates a new standalone window, leaving the Level Editor Viewport the same. **Mobile Preview** and **Standalone Game** will launch the game in a standalone version outside of the Editor (with the Mobile Preview providing additional tools used to mimic devices as closely as possible). 
+The Selected Viewport is the default setting and launches the game within the Editor inside the currently selected [Viewport](Engine/UI/LevelEditor/Viewports/Basics). New Editor Window also launches the game within the Editor however it creates a new standalone window, leaving the Level Editor Viewport the same. Mobile Preview and Standalone Game will launch the game in a standalone version outside of the Editor (with the Mobile Preview providing additional tools used to mimic devices as closely as possible). 
 
 [REGION:note]
 For the other options in this menu, refer to the **Related Topics** section at the bottom of this page for more information. 
