#UE4doc #docs:
Removed Overhead header
Updated Engine version
Changed body of each definition to be more direct
Removed unnecessary italics
Commented out multiple links to empty pages that will be 404ing after the next publish.

diff --git a/Engine/Documentation/Source/Programming/Online/Types/Types.INT.udn b/Engine/Documentation/Source/Programming/Online/Types/Types.INT.udn
--- a/Engine/Documentation/Source/Programming/Online/Types/Types.INT.udn
+++ b/Engine/Documentation/Source/Programming/Online/Types/Types.INT.udn
@@ -2,18 +2,17 @@ Availability:Public
 Title: Online Subsystem Types
 Crumbs:%ROOT%, Programming
 Description:Description of the major types defined in the online subsystem
-Version: 4.5
+Version: 4.9
 
 [TOC(start:2)]
 
 
-## Overview
 
 This document serves to define the different data structures that are used throughout the online subsystem.
 
 ### IOnlinePlatformData
 
- *The base class for anything meant to be opaque so that the data can be passed around without consideration for the data it contains.*
+ **IOnlinePlatformData** is the base class for anything meant to be opaque so that the data can be passed around without consideration for the data it contains.
  A human-readable version of the data is available via the **ToString()** function; otherwise, nothing but platform code should try to operate directly on the data.
 
  There are often platform-specific representations of various constructs that, while they need to be handled by the game, do not need their details exposed. This allows the game to hold onto the data and replicate it to other machines without understanding the details.
@@ -56,24 +55,24 @@ This document serves to define the different data structures that are used throu
 
 ### FUniqueNetId
 
-*Contains a profile service online ID and can represent anything from an individual player, friend or a session.*
+**FUniqueNetId** contains a profile service online ID and can represent anything from an individual player, friend or a session.
 
-See:
- * [](Programming/Online/Interfaces/Profile)
+<!--See:
+ * [](Programming/Online/Interfaces/Profile)-->
 
 
 ### FSharedContentHandle
 
-*Contains a profile service shared file handle.*  This handle can be passed around and used to access shared content from anywhere, given proper permissions.
+**FSharedContentHandle** contains a profile service shared file handle.  This handle can be passed around and used to access shared content from anywhere, given proper permissions.
 
-See:
- * [](Programming/Online/Interfaces/SharedCloud)
+<!--See:
+ * [](Programming/Online/Interfaces/SharedCloud)-->
 
 
  
 ### FSessionInfo
 
-*Contains a session's platform-specific info.*  In most cases, this would be information such as a session identifier, host address, or other means of identifying and connecting to a created session.
+**FSessionInfo** contains a session's platform-specific info.  In most cases, this would be information such as a session identifier, host address, or other means of identifying and connecting to a created session.
 
 See:
  * [](Programming/Online/Interfaces/Session)
@@ -81,7 +80,7 @@ See:
 
 ### FEmsFile
 
-*Metadata about a given downloadable file.*  This structure is used when enumerating available data from a given service. 
+**FEmsFile** is metadata about a given downloadable file.  This structure is used when enumerating available data from a given service. 
 
 Contains the following elements:
  * Hash - the hash value of the given file contents.
@@ -90,14 +89,14 @@ Contains the following elements:
  * File size - size of the file.  
 
 
-See:
+<!--See:
  * [](Programming/Online/Interfaces/SharedCloud)
- * [](Programming/Online/Interfaces/UserCloud)
+ * [](Programming/Online/Interfaces/UserCloud)-->
 
 
 ### FTitleFile
 
-*Internal structure that holds the data used in downloading a file asynchronously from the online service.*
+**FTitleFile** is the internal structure that holds the data used in downloading a file asynchronously from the online service.
 
 Contains the following elements:
  * Filename - the filename as downloaded.
@@ -105,14 +104,14 @@ Contains the following elements:
  * Data - a buffer containing the file contents.
 
 
-See: 
+<!--See: 
  * [](Programming/Online/Interfaces/SharedCloud)
- * [](Programming/Online/Interfaces/UserCloud)
+ * [](Programming/Online/Interfaces/UserCloud)-->
 
 
 ### EOnlineAsyncTaskState
 
-*Simple enum representing the possible states for an asynchronous operation.* 
+**EOnlineAsyncTaskState** is a simple enum representing the possible states for an asynchronous operation.
 
 
 		/** The task has not been started */
