#UE4doc #docs: Removed extraneous bolding as according to the instructional guide.

diff --git a/Engine/Documentation/Source/GettingStarted/SubEditors/SubEditors.INT.udn b/Engine/Documentation/Source/GettingStarted/SubEditors/SubEditors.INT.udn
--- a/Engine/Documentation/Source/GettingStarted/SubEditors/SubEditors.INT.udn
+++ b/Engine/Documentation/Source/GettingStarted/SubEditors/SubEditors.INT.udn
@@ -64,7 +64,7 @@ Refer to each of the sections below for an overview along with links to more det
 [EXCERPT:leveleditor]
 **Level Editor**
 
-The **Level Editor** is the primary Editor used to construct your gameplay levels. Generally speaking it is where you define the play space for your game by adding different types of [](Engine/Actors), [Blueprints](Engine/Blueprints), [](Engine/Rendering/ParticleSystems) or anything else that you want to add to your level. By default, when you create or open a project it will open up to the Level Editor. 
+The Level Editor is the primary Editor used to construct your gameplay levels. Generally speaking it is where you define the play space for your game by adding different types of [](Engine/Actors), [Blueprints](Engine/Blueprints), [](Engine/Rendering/ParticleSystems) or anything else that you want to add to your level. By default, when you create or open a project it will open up to the Level Editor. 
 
 For more information, see: [](Engine/UI/LevelEditor)
 [/EXCERPT:leveleditor]
@@ -74,7 +74,7 @@ For more information, see: [](Engine/UI/LevelEditor)
 [EXCERPT:materialeditor]
 **Material Editor**
 
-The **Material Editor** is where you can create (or edit existing) **Materials** which are assets that can be applied to a mesh to control its visual look. For example, you can create a "dirt" Material and apply it to floors in your level or terrain to create a dirt looking surface. 
+The Material Editor is where you can create (or edit existing) Materials which are assets that can be applied to a mesh to control its visual look. For example, you can create a "dirt" Material and apply it to floors in your level or terrain to create a dirt looking surface. 
 
 For more information, see: [](Engine/Rendering/Materials/Editor)
 [/EXCERPT:materialeditor]
@@ -84,7 +84,7 @@ For more information, see: [](Engine/Rendering/Materials/Editor)
 [EXCERPT:blueprinteditor]
 **Blueprint Editor**
 
-The **Blueprint Editor** is where you can work with and modify **Blueprints** which are special assets that can be used to create new types of Actors and script level events, without needed to write any form of C++ code. 
+The Blueprint Editor is where you can work with and modify Blueprints which are special assets that can be used to create new types of Actors and script level events, without needed to write any form of C++ code. 
 
 For more information, see: [](Engine/Blueprints/Editor)
 [/EXCERPT:blueprinteditor]
@@ -94,7 +94,7 @@ For more information, see: [](Engine/Blueprints/Editor)
 [EXCERPT:behavioreditor]
 **Behavior Tree Editor**
 
-Inside the **Behavior Tree Editor** you can script Artificial Intelligence through a visual node-based system (similar to **Blueprints**) for Actors in your levels (this could be any number of different behaviors for enemies, NPC characters, vehicles, etc.) 
+Inside the Behavior Tree Editor you can script Artificial Intelligence through a visual node-based system (similar to Blueprints) for Actors in your levels (this could be any number of different behaviors for enemies, NPC characters, vehicles, etc.) 
 
 For more information, see: [](Engine/AI/BehaviorTrees/UserGuide)
 [/EXCERPT:behavioreditor]
@@ -104,7 +104,7 @@ For more information, see: [](Engine/AI/BehaviorTrees/UserGuide)
 [EXCERPT:personaeditor]
 **Persona Editor**
 
-The **Persona Editor** is the animation editing toolset within Unreal Engine 4 and is used for editing [](Engine/Animation/Skeleton), [](Engine/Content/Types/SkeletalMeshes), [](Engine/Animation/AnimBlueprints) and various other animation assets. Most (if not all) of your animation work within Unreal Engine 4 will take place in this Editor. 
+The Persona Editor is the animation editing toolset within Unreal Engine 4 and is used for editing [](Engine/Animation/Skeleton), [](Engine/Content/Types/SkeletalMeshes), [](Engine/Animation/AnimBlueprints) and various other animation assets. Most (if not all) of your animation work within Unreal Engine 4 will take place in this Editor. 
 
 
 For more information, see: [](Engine/Animation/Persona)
@@ -115,7 +115,7 @@ For more information, see: [](Engine/Animation/Persona)
 [EXCERPT:cascadeeditor]
 **Cascade Editor**
 
-[](Engine/Rendering/ParticleSystems) in Unreal Engine 4 are edited inside the **Cascade Editor**, which is a fully integrated and modular particle effects editor. **Cascade** offers real-time feedback and modular effects editing, allowing fast and easy creation of even the most complex effects. 
+[](Engine/Rendering/ParticleSystems) in Unreal Engine 4 are edited inside the Cascade Editor, which is a fully integrated and modular particle effects editor. Cascade offers real-time feedback and modular effects editing, allowing fast and easy creation of even the most complex effects. 
 
 For more information, see: [](Engine/Rendering/ParticleSystems/Cascade)
 [/EXCERPT:cascadeeditor]
@@ -125,7 +125,7 @@ For more information, see: [](Engine/Rendering/ParticleSystems/Cascade)
 [EXCERPT:umguieditor]
 **UMG UI Editor**
 
-The **UMG UI Editor** is a visual UI authoring tool which can be used to create UI elements such as in-game HUDs, menus or other interface related graphics you wish to present to your users. 
+The UMG UI Editor is a visual UI authoring tool which can be used to create UI elements such as in-game HUDs, menus or other interface related graphics you wish to present to your users. 
 
 
 For more information, see: [](Engine/UMG/UserGuide)
@@ -136,7 +136,7 @@ For more information, see: [](Engine/UMG/UserGuide)
 [EXCERPT:matineeeditor]
 **Matinee Editor**
 
-The **Matinee Editor** is based on the use of specialized animation tracks in which you can place keyframes to set the values of certain properties of the Actors in your level. This allows you to create in-game cinematics, dynamic gameplay events or even animate the properties of Actors over time (such as driving the brightness of a light). 
+The Matinee Editor is based on the use of specialized animation tracks in which you can place keyframes to set the values of certain properties of the Actors in your level. This allows you to create in-game cinematics, dynamic gameplay events or even animate the properties of Actors over time (such as driving the brightness of a light). 
 
 For more information, see: [](Engine/Matinee/UserGuide)
 [/EXCERPT:matineeeditor]
@@ -146,7 +146,7 @@ For more information, see: [](Engine/Matinee/UserGuide)
 [EXCERPT:soundcueeditor]
 **Sound Cue Editor**
 
-The behavior of audio playback in Unreal Engine 4 is defined within **Sound Cues** which can be editing using the **Sound Cue Editor**. Inside the **Sound Cue Editor**, you can combine and mix several sound assets to produce a single mixed "output" saved as a **Sound Cue**. 
+The behavior of audio playback in Unreal Engine 4 is defined within Sound Cues which can be editing using the Sound Cue Editor. Inside the Sound Cue Editor, you can combine and mix several sound assets to produce a single mixed "output" saved as a Sound Cue. 
 
 
 For more information, see: [](Engine/Audio/SoundCues/Editor)
@@ -156,7 +156,7 @@ For more information, see: [](Engine/Audio/SoundCues/Editor)
 [EXCERPT:spriteeditor]
 **Paper2D Sprite Editor**
 
-The **Paper2D Sprite Editor** enables setup and editing of individual [](Engine/Paper2D/Sprites) (which are in essence a quick and easy way to draw 2D images in Unreal Engine 4). 
+The Paper2D Sprite Editor enables setup and editing of individual [](Engine/Paper2D/Sprites) (which are in essence a quick and easy way to draw 2D images in Unreal Engine 4). 
 
 
 For more information, see: [](Engine/Paper2D/Sprites/Editor)
@@ -167,7 +167,7 @@ For more information, see: [](Engine/Paper2D/Sprites/Editor)
 [EXCERPT:flipbookeditor]
 **Paper2D Flipbook Editor**
 
-With the **Paper2D Flipbook Editor** you can create 2D animations called **Flipbooks**. By specifying a series of Sprites along certain key frames inside the **Paper2D Flipbook Editor**, those frames are then "flipped" through to create an animation. The best way to think of **Flipbooks** is of the old hand-drawn style of animation. 
+With the Paper2D Flipbook Editor you can create 2D animations called Flipbooks. By specifying a series of Sprites along certain key frames inside the Paper2D Flipbook Editor, those frames are then "flipped" through to create an animation. The best way to think of Flipbooks is of the old hand-drawn style of animation. 
 
 For more information, see: [](Engine/Paper2D/Flipbooks/Editor)
 [/EXCERPT:flipbookeditor]
@@ -177,7 +177,7 @@ For more information, see: [](Engine/Paper2D/Flipbooks/Editor)
 [EXCERPT:phateditor]
 **Physics Asset Tool Editor**
 
-The **Physics Asset Tool** (or **PhAT** for short) is used for creating **Physics Assets** for use with [](Engine/Content/Types/SkeletalMeshes). You can start from nothing and build to a full ragdoll setup or use the automation tools to create a basic set of Physics Bodies and Physics Constraints. 
+The Physics Asset Tool (or PhAT for short) is used for creating Physics Assets for use with [](Engine/Content/Types/SkeletalMeshes). You can start from nothing and build to a full ragdoll setup or use the automation tools to create a basic set of Physics Bodies and Physics Constraints. 
 
 
 
@@ -189,7 +189,7 @@ For more information, see: [](Engine/Physics/PhAT)
 [EXCERPT:staticmesheditor]
 **Static Mesh Editor**
 
-The **Static Mesh Editor** is used to preview the look, collision, and UVs as well as set and manipulate the properties of [](Engine/Content/Types/StaticMeshes). Inside the **Static Mesh Editor** you can also setup [LODs](Engine/Content/Types/StaticMeshes/HowTo/LODs) (or Level of Detail settings) for your Static Mesh assets. 
+The Static Mesh Editor is used to preview the look, collision, and UVs as well as set and manipulate the properties of [](Engine/Content/Types/StaticMeshes). Inside the Static Mesh Editor you can also setup [LODs](Engine/Content/Types/StaticMeshes/HowTo/LODs) (or Level of Detail settings) for your Static Mesh assets. 
 
 For more information, see: [](Engine/Content/Types/StaticMeshes/Editor)
 [/EXCERPT:staticmesheditor]
@@ -199,7 +199,7 @@ For more information, see: [](Engine/Content/Types/StaticMeshes/Editor)
 [EXCERPT:mediaeditor]
 **Media Player Editor**
 
-The **Media Player Editor** allows you to define media files or URLs to use as source media for playback inside Unreal Engine 4. 
+The Media Player Editor allows you to define media files or URLs to use as source media for playback inside Unreal Engine 4. 
 
 Although not necessarily an "Editor" in that you cannot edit media files in this Editor, you define settings to how your source media will playback such as auto-playing, play rate or looping for example. 
 
