UEDOC-1787

Updated GIF to Youtube video

diff --git a/Engine/Documentation/Source/Engine/Landscape/Materials/Materials.INT.udn b/Engine/Documentation/Source/Engine/Landscape/Materials/Materials.INT.udn
--- a/Engine/Documentation/Source/Engine/Landscape/Materials/Materials.INT.udn
+++ b/Engine/Documentation/Source/Engine/Landscape/Materials/Materials.INT.udn
@@ -328,7 +328,17 @@ You can use the same Material for your Landscape Material and your Landscape Hol
 	![](Landscape_Hole_Materail_Assign.png)
 
 1. When completed with everything, you should now be able to mask and un-mask the visibility of your terrain like in the following video. 
-	![](Landscape_Hole_Material_Usage.gif)(convert:false)
+	[OBJECT:EmbeddedVideo]
+	[PARAMLITERAL:width]
+	640
+	[/PARAMLITERAL]
+	[PARAMLITERAL:height]
+	360
+	[/PARAMLITERAL]
+	[PARAMLITERAL:videoid]
+	6NxNppmZRss
+	[/PARAMLITERAL]
+	[/OBJECT]
 
 [REGION:note]
 For information about using the painting tools to create a Landscape hole, see [Using Opacity Masking to Create a Hole](Engine\Landscape\Editing\SculptMode\#UsingOpacityMaskingtoCreateaHole).
