#UE4doc #docs: Removed extraneous bolding as according to the instructional guide. Unbolded hyperlink in accordance to the intstructional guide.

diff --git a/Engine/Documentation/Source/Platforms/Android/GettingStarted/2/AndroidQS_2.INT.udn b/Engine/Documentation/Source/Platforms/Android/GettingStarted/2/AndroidQS_2.INT.udn
--- a/Engine/Documentation/Source/Platforms/Android/GettingStarted/2/AndroidQS_2.INT.udn
+++ b/Engine/Documentation/Source/Platforms/Android/GettingStarted/2/AndroidQS_2.INT.udn
@@ -59,14 +59,14 @@ Description:Set up your Android device for development.
 [EXCERPT:WindowsSetup]
 
 1. Connect your device.
-1.	Device driver software may automatically install, but if not, visit **[Android's OEM USB Drivers](http://developer.android.com/tools/extras/oem-usb.html)** page to find driver installation links and additional information.  
+1.	Device driver software may automatically install, but if not, visit [Android's OEM USB Drivers](http://developer.android.com/tools/extras/oem-usb.html) page to find driver installation links and additional information.  
 1.	Open the **Settings** application on your device.
-1.	Look for the **Developer Options** menu. If it is not available, you need to enable **Developer Options** on your device. For more about developer options, visit **[Android's Using Hardware Devices](http://developer.android.com/tools/device.html)** page.
+1.	Look for the **Developer Options** menu. If it is not available, you need to enable Developer Options on your device. For more about developer options, visit [Android's Using Hardware Devices](http://developer.android.com/tools/device.html) page.
 	1.	Scroll down and select **About Phone**. Depending on your Android device, you may have an **About Device**, **About Tablet**, or **About Shield** option instead.
 		[REGION: note]
 		Note: This may be in a **More** section on newer Android versions.
 		[/REGION]
-	1.	Tap the **Build Number** 7 times. After the first few presses, you will see a message stating **You are now # steps away from becoming a developer.**
+	1.	Tap the **Build Number** 7 times. After the first few presses, you will see a message stating You are now # steps away from becoming a developer.
 	1.	You will get a success message once **Developer Mode** has been activated.
 	1.	Navigate back to the **Settings** menu and select **Developer Options** which is now available.
 1.	In the **Developer Options** menu, enable USB debugging. 
