#UE4doc #docs:
Removed Overview header
Updated Engine version
Edited/update images to have yellow indicators instead of green.
Unbolded a few terms/lines

diff --git a/Engine/Documentation/Source/Programming/Plugins/Plugins.INT.udn b/Engine/Documentation/Source/Programming/Plugins/Plugins.INT.udn
--- a/Engine/Documentation/Source/Programming/Plugins/Plugins.INT.udn
+++ b/Engine/Documentation/Source/Programming/Plugins/Plugins.INT.udn
@@ -2,7 +2,7 @@ Availability:Public
 Title:Plugins
 Crumbs:%ROOT%, Programming
 Description:How to create Unreal Engine plugins.
-Version: 4.5
+Version: 4.9
 
 
 [REGION:warning]
@@ -14,7 +14,6 @@ document.
 
 [TOC(start:2)]
 
-## Overview
 
 This page describes how to develop your own plugins to use with the Unreal Engine tools and runtime. 
 
@@ -63,7 +62,7 @@ section for more information.
 For plugins with code modules, the plugin will have its own Binaries folder that contains compiled code for that plugin. 
 Also, temporary build product files will be stored to an Intermediate folder under the plugin's directory. 
 
-Plugins can have their own Content folder that contains asset files specific to that plugin!  See the [Content in Plugins](#ContentInPlugins) 
+Plugins can have their own Content folder that contains asset files specific to that plugin.  See the [Content in Plugins](#ContentInPlugins) 
 section for more information.
 
 [REGION:warning]
@@ -153,7 +152,7 @@ they will still be loaded at startup by the game as long as binary files exist.
 
 ## Engine Plugins
 
-Unreal Engine has some **built-in** plugins included under the Engine directory. Engine plugins are just like plugins you would have in 
+Unreal Engine 4 has some built-in plugins included under the Engine directory. Engine plugins are just like plugins you would have in 
 your game project, except they are available for all game projects. Typically, these plugins are created by engine and tools programmers 
 to provide baseline functionality in a plugin-like way. This can enable the user to remove or override whole engine features without 
 modifying any of the engine code. 
@@ -174,9 +173,9 @@ Here are the steps required to distribute a plugin.
 3. If your plugin contains a **Source** directory, **compile plugin binaries** for **Win64 Development** configuration. This configuration will always be needed by an editor on PC platform to load your plugin. 
 4. For plugins with **Source**, if you intend to support additional configurations such as Mac, you should build those too.
 5. Make a copy of your plugin to a new "staging" folder outside of your project directory. It does not matter where, but you are going to be making changes to the plugin here.
-6. In the staging directory for your plugin, delete the Intermediate directory. These are temporary files that should never be distributed.
-7. If your plugin includes source code and you do not want to distribute it, delete any Source directories in the staging folder for your plugin.
-8. In the staging directory for your plugin, delete all files that you do not want to distribute, such as Private source files.
+6. In the staging directory for your plugin, delete the **Intermediate** directory. These are temporary files that should never be distributed.
+7. If your plugin includes source code and you do not want to distribute it, delete any **Source** directories in the staging folder for your plugin.
+8. In the staging directory for your plugin, delete all files that you do not want to distribute.
 9. Your plugin is now ready to distribute!
 
 
@@ -187,7 +186,7 @@ Important notes:
 directory, or no Source directory at all. We intend to add support for distributing Public source (as well as the *.Build.cs files) 
 at a later date.
 * We are investigating tools to automate most of this process. More details will be coming at a later date.
-* For plugins with **Content**, you may want to **generate Derived Data** to include with your plugin, to avoid having to require it 
+* For plugins with Content, you may want to generate Derived Data to include with your plugin, to avoid having to require it 
 to be generated on demand by an end-user. This feature is not yet available, but will be added at a later date. Stay tuned.
 * Including an EULA or other documentation with distributed plugins is not yet supported. We may provide support for this at a later date.
 [/REGION]
@@ -196,7 +195,7 @@ to be generated on demand by an end-user. This feature is not yet available, but
 
 ## Plugin Descriptor files
 
-Plugin descriptors are files that end in a **.uplugin** extension. The first part of the file name is always the **name of your plugin**. 
+Plugin descriptors are files that end in a **.uplugin** extension. The first part of the file name is always the name of your plugin. 
 Plugin descriptor files are always located in your plugin's directory, where the engine will discover them at startup time.
 
 Plugin descriptors are in the Json ([JavaScript Object Notation](http://www.json.org/)) file format.
@@ -245,7 +244,7 @@ This example plugin descriptor is from the **UObjectPlugin** [downloadable examp
 
 
 [REGION:warning]
-This .uplugin file format is still in development and highly likely to change!  For example, we expect the way that 
+This .uplugin file format is still in development and highly likely to change. For example, we expect the way that 
 versioning is handled will be modified in a future release.
 [/REGION]
 
