#UE4doc #docs:
Moved abbreviation declaration for UE4 to the first mention of Unreal Engine 4.
Updated Engine Version
Removed unnecessary italics.

diff --git a/Engine/Documentation/Source/Programming/QuickStart/QuickStart.INT.udn b/Engine/Documentation/Source/Programming/QuickStart/QuickStart.INT.udn
--- a/Engine/Documentation/Source/Programming/QuickStart/QuickStart.INT.udn
+++ b/Engine/Documentation/Source/Programming/QuickStart/QuickStart.INT.udn
@@ -4,7 +4,7 @@ Crumbs:%ROOT%, Programming
 Description:Create your first code project and add a new C++ class.
 prereq:Engine/Basics/Projects/Browser
 SkillLevel: Beginner
-Version: 4.8
+Version: 4.9
 
 [VAR:Topic]
 [OBJECT:Topic]
@@ -56,7 +56,7 @@ Version: 4.8
 
 ## Goals
 
-The primary goal of this tutorial is to introduce programmers to Unreal Engine 4's development environment. By the end of this quick start guide, developers should be comfortable setting up and developing C++ Projects in Unreal Engine 4 (UE4).
+The primary goal of this tutorial is to introduce programmers to Unreal Engine 4's (UE4) development environment. By the end of this quick start guide, developers should be comfortable setting up and developing C++ Projects in UE4.
 
 ## Objectives
 
@@ -67,7 +67,7 @@ Upon completion, developers will have the skill to:
 * Compile a Project 
 * Add an Instance of a New Class to a Level
 
-_When we're finished, we'll be able to see an Actor we programmed in C++ moving around on screen._
+When we're finished, we'll be able to see an Actor we programmed in C++ moving around on screen.
 
 [OBJECT:TopicButtonList]
 	[PARAM:icon]
