UEDOC-1795

Updated multiple GIFs to Youtube videos

diff --git a/Engine/Documentation/Source/Engine/Rendering/Materials/HowTo/Transparency/Transparency.INT.udn b/Engine/Documentation/Source/Engine/Rendering/Materials/HowTo/Transparency/Transparency.INT.udn
--- a/Engine/Documentation/Source/Engine/Rendering/Materials/HowTo/Transparency/Transparency.INT.udn
+++ b/Engine/Documentation/Source/Engine/Rendering/Materials/HowTo/Transparency/Transparency.INT.udn
@@ -85,7 +85,17 @@ content page for information about how to move content between projects. This wa
 
 1. Find the **Transparency_Material** in the **Content Browser** and **Right-Click** on it with the mouse and from the menu that is displayed, select the **Create Material Instance** option.  
 
-	![](Tras_Creat_New_MI.gif)(convert:false)
+	[OBJECT:EmbeddedVideo]
+	[PARAMLITERAL:width]
+	640
+	[/PARAMLITERAL]
+	[PARAMLITERAL:height]
+	360
+	[/PARAMLITERAL]
+	[PARAMLITERAL:videoid]
+	PhhorP5ae28
+	[/PARAMLITERAL]
+	[/OBJECT]
 
 1. Now we need to find an object to place the Material Instance on so we can see the Material in action. Inside of the **Content Browser**, under the Starter Content folder there is a **Shapes** folder.
 Locate the  **Shape_Sphere** and then select it with the **Left Mouse Button** and then drag it into the current level with the **Left Mouse Button** still held down. Once in a location you like, release the **Left Mouse Button** to place the mesh in the level.
@@ -95,7 +105,17 @@ Locate the  **Shape_Sphere** and then select it with the **Left Mouse Button** a
 1. With something in the level to test with, find the **Transparency_Material_Inst** that was created before in the **Content Browser**. Once located, select it with the **Left Mouse Button** and then drag it from the **Content Browser** onto the **Shape_Sphere**.
 Once over the **Shape_Sphere**, release the **Left Mouse Button** to apply the Material Instance to the mesh.
 
-	![](Tras_Apply_Trans_Material.gif)(convert:false)
+	[OBJECT:EmbeddedVideo]
+	[PARAMLITERAL:width]
+	640
+	[/PARAMLITERAL]
+	[PARAMLITERAL:height]
+	360
+	[/PARAMLITERAL]
+	[PARAMLITERAL:videoid]
+	zJP1NO1DdB8
+	[/PARAMLITERAL]
+	[/OBJECT]
 
 1. Once the Material Instance has been applied to the mesh, open the Material Instance up by **Double-Clicking** on it with the **Left Mouse Button** in the **Content Browser**. Once open enable the **Transparency_Amount** by clicking on the grayed out check mark box with the **Left Mouse Button**. Once enabled you can adjust the value of **Transparency_Amount** to change how transparent the object will be. 
 
@@ -191,7 +211,17 @@ to figure out exactly where your performance issues could be coming from.
 When having multiple Transparent objects in the scene close together, you will sometime notice that the Transparent objects have issues figuring out which object to render in front or 
 behind the other transparent objects. The following example demonstrates this happening using the Materials that were created above and the lit smoke VFX from the Starter Content.
 
-![](Tras_Sort_Issue_Example.gif)(convert:false)
+[OBJECT:EmbeddedVideo]
+[PARAMLITERAL:width]
+640
+[/PARAMLITERAL]
+[PARAMLITERAL:height]
+360
+[/PARAMLITERAL]
+[PARAMLITERAL:videoid]
+BipY-HUpClk
+[/PARAMLITERAL]
+[/OBJECT]
 
 Notice how the smoke will be in front of the meshes and then all of a sudden it will pop behind a mesh. 
 To prevent this from happening, the **Translucency Sort Priority** needs to be set so that the VFX of the smoke will always render on top of other translucency in the scene.
@@ -203,7 +233,17 @@ To change the **Translucency Sort Priority** on the VFX, do the following.
 
 1. In the **Details** panel, under the **Rendering** section, click on the **White Triangle** to expose additional properties.
 
-	![](Tras_Rendering_Advanced.gif)(convert:false)
+	[OBJECT:EmbeddedVideo]
+	[PARAMLITERAL:width]
+	640
+	[/PARAMLITERAL]
+	[PARAMLITERAL:height]
+	360
+	[/PARAMLITERAL]
+	[PARAMLITERAL:videoid]
+	78ZOswFZbj8
+	[/PARAMLITERAL]
+	[/OBJECT]
 
 1. From the new list of options, look for the **Translucency Sort Priority** and then input a value of 100. 
 
@@ -215,11 +255,31 @@ To change the **Translucency Sort Priority** on the VFX, do the following.
 
 1. Now check out the smoke VFX in the scene. With a **Translucency Sort Priority**  set to 100, the smoke will always draw on top of the other translucent objects.
 
-	![](Tras_New_Trans_Order_Front.gif)(convert:false)
+	[OBJECT:EmbeddedVideo]
+	[PARAMLITERAL:width]
+	640
+	[/PARAMLITERAL]
+	[PARAMLITERAL:height]
+	360
+	[/PARAMLITERAL]
+	[PARAMLITERAL:videoid]
+	A4OcPB_wSoc
+	[/PARAMLITERAL]
+	[/OBJECT]
 
 1. Setting the **Translucency Sort Priority** to -100 will make it so that the smoke always draws behind other translucent objects.
 
-	![](Tras_New_Trans_Order_Behind.gif)(convert:false)
+	[OBJECT:EmbeddedVideo]
+	[PARAMLITERAL:width]
+	640
+	[/PARAMLITERAL]
+	[PARAMLITERAL:height]
+	360
+	[/PARAMLITERAL]
+	[PARAMLITERAL:videoid]
+	MBjshrSST4Q
+	[/PARAMLITERAL]
+	[/OBJECT]
 
 While adjusting the **Translucency Sort Priority** can fix sorting issue in the level, it could introduce new sorting issues when spawning effects or objects during game play. 
 Because of this, it is a good idea to define what **Translucency Sort Priority** number to use and where.
