#UE4doc #docs:
Updated Engine version
Changed opening note to specify Visual Studio 2013 instead of  the blanket statement of telling the user to have Visual Studio, no matter the version.

diff --git a/Engine/Documentation/Source/Programming/QuickStart/1/QuickStart_1.INT.udn b/Engine/Documentation/Source/Programming/QuickStart/1/QuickStart_1.INT.udn
--- a/Engine/Documentation/Source/Programming/QuickStart/1/QuickStart_1.INT.udn
+++ b/Engine/Documentation/Source/Programming/QuickStart/1/QuickStart_1.INT.udn
@@ -3,7 +3,7 @@ Title:1 - Essential Project Setup
 Crumbs:%ROOT%, Programming, Programming/QuickStart
 Description:Creating a new project with the specifications we need
 SkillLevel: Beginner
-Version: 4.8
+Version: 4.9
 
 [VAR:Steps]
 [OBJECT:Navigation]
@@ -23,7 +23,7 @@ Version: 4.8
 %Steps%
 
 [REGION:note]
-You should have a copy of **Visual Studio** installed before starting this tutorial.
+You should have a copy of **Visual Studio 2013** for Deskstop (Preferably Community or Professional) installed before starting this tutorial.
 [/REGION]
 
 1. Open **Unreal Engine** from the Launcher. The [](Engine/Basics/Projects/Browser) will appear.
