#UE4doc #docs: ##Overview removed as according to the instructional guide.

diff --git a/Engine/Documentation/Source/Programming/Assets/AsyncLoading/AsyncLoading.INT.udn b/Engine/Documentation/Source/Programming/Assets/AsyncLoading/AsyncLoading.INT.udn
--- a/Engine/Documentation/Source/Programming/Assets/AsyncLoading/AsyncLoading.INT.udn
+++ b/Engine/Documentation/Source/Programming/Assets/AsyncLoading/AsyncLoading.INT.udn
@@ -6,8 +6,6 @@ Version: 4.5
 
 [TOC(start:2)]
 
-## Overview
-
 There are several new systems in UE4 that make it much easier to asynchronously load asset data, and they replace a lot of the functionality that existed for seekfree content packages in UE3. These new methods work identically in development and with cooked data on devices, so you do not need to maintain two code paths for loading data on demand. There are two general methods you can use to refer to and load data on demand:
 
 ## FStringAssetReferences and TAssetPtr
