#UE4doc #docs:
Updated Engine version

diff --git a/Engine/Documentation/Source/Programming/QuickStart/2/QuickStart_2.INT.udn b/Engine/Documentation/Source/Programming/QuickStart/2/QuickStart_2.INT.udn
--- a/Engine/Documentation/Source/Programming/QuickStart/2/QuickStart_2.INT.udn
+++ b/Engine/Documentation/Source/Programming/QuickStart/2/QuickStart_2.INT.udn
@@ -3,7 +3,7 @@ Title:2 - Create a C++ Class
 Crumbs:%ROOT%, Programming, Programming/QuickStart
 Description:Creating a C++ class
 SkillLevel: Beginner
-Version: 4.8
+Version: 4.9
 
 [VAR:Steps]
 [OBJECT:Navigation]
