#UE4doc #docs:
Updated Engine version
Bolded a couple terms
Replaced 'Unreal Engine' with 'Unreal Editor' in some places.

diff --git a/Engine/Documentation/Source/Programming/QuickStart/3/QuickStart_3.INT.udn b/Engine/Documentation/Source/Programming/QuickStart/3/QuickStart_3.INT.udn
--- a/Engine/Documentation/Source/Programming/QuickStart/3/QuickStart_3.INT.udn
+++ b/Engine/Documentation/Source/Programming/QuickStart/3/QuickStart_3.INT.udn
@@ -3,7 +3,7 @@ Title:3 - Write and Compile C++ Code
 Crumbs:%ROOT%, Programming, Programming/QuickStart
 Description:Writing a C++ class and compiling it with Unreal Engine
 SkillLevel: Beginner
-Version: 4.8
+Version: 4.9
 
 [VAR:Steps]
 [OBJECT:Navigation]
@@ -56,15 +56,15 @@ Version: 4.8
 	active_button_content
 	[/PARAMLITERAL]
 	[PARAM:content]
-		*	Now that we're done coding, we can compile by right-clicking our project in the **Solution Browser** and selecting the **Build** command, or by pressing the **Compile** button in the **Unreal Editor**. Once the compile succeeds, **Unreal Engine** will automatically load our changes.
+		*	Now that we're done coding, we can compile by right-clicking our project in the **Solution Browser** and selecting the **Build** command, or by pressing the **Compile** button in the **Unreal Editor**. Once the compile succeeds, **Unreal Editor** will automatically load our changes.
 
 			[![](CompileFromVS.png)](CompileFromVS.png)
 	
-			(compiling from Visual Studio pictured above)
+			(compiling from **Visual Studio** pictured above)
 
 			[![](CompileFromEditor.png)](CompileFromEditor.png)
 
-			(compiling from the Unreal Editor pictured above)
+			(compiling from the **Unreal Editor** pictured above)
 	[/PARAM]
 [/OBJECT]
 
@@ -107,7 +107,7 @@ Version: 4.8
 
 ----
 
-We're now ready to create objects in **Unreal Engine** based on our code. For reference, all code used on this page is included below.
+We're now ready to create objects in **Unreal Editor** based on our code. For reference, all code used on this page is included below.
 
 [OBJECT:Section]
     [PARAMLITERAL:id]
