#UE4doc #docs:
Updated Engine version
Unbolded and unitalicized a couple terms.

diff --git a/Engine/Documentation/Source/Programming/QuickStart/5/QuickStart_5.INT.udn b/Engine/Documentation/Source/Programming/QuickStart/5/QuickStart_5.INT.udn
--- a/Engine/Documentation/Source/Programming/QuickStart/5/QuickStart_5.INT.udn
+++ b/Engine/Documentation/Source/Programming/QuickStart/5/QuickStart_5.INT.udn
@@ -3,7 +3,7 @@ Title:5 - On Your Own!
 Crumbs:%ROOT%, Programming, Programming/QuickStart
 Description:Ideas for further exploration!
 SkillLevel: Beginner
-Version: 4.8
+Version: 4.9
 
 [VAR:Steps]
 [OBJECT:Navigation]
@@ -26,9 +26,9 @@ Version: 4.8
 
 Using what you have learned, try to do the following:
 
-* Add a [Particle System Component](Engine/Rendering/ParticleSystems) to your **FloatingActor**. Some pre-built **Particle Systems** are already included in your project.
+* Add a [Particle System Component](Engine/Rendering/ParticleSystems) to your **FloatingActor**. Some pre-built Particle Systems are already included in your project.
 * Use **Unreal Engine's** [UProperty](Programming/UnrealArchitecture/Reference/Properties) macro to expose a variable for the magnitude of your **FloatingActor's** movement, instead of using a hardcoded value. You might want to check the [%Programming/Tutorials/VariablesTimersEvents:title%](Programming/Tutorials/VariablesTimersEvents) tutorial for help on this topic.
-* Add periodic motion on the X and/or Y axis, and multiply the _DeltaTime_ value by a number between 0.6 and 1.4, so your **FloatingActor** appears to float around freely. This can look great for powerups!
+* Add periodic motion on the X and/or Y axis, and multiply the `DeltaTime` value by a number between 0.6 and 1.4, so your **FloatingActor** appears to float around freely. This can look great for powerups!
 
 As for the specifics covered in this tutorial:
 
