#UE4doc #docs: Removed extraneous bolding as according to the instructional guide.

diff --git a/Engine/Documentation/Source/Resources/ContentExamples/PostProcessing/1_14/1_14.INT.udn b/Engine/Documentation/Source/Resources/ContentExamples/PostProcessing/1_14/1_14.INT.udn
--- a/Engine/Documentation/Source/Resources/ContentExamples/PostProcessing/1_14/1_14.INT.udn
+++ b/Engine/Documentation/Source/Resources/ContentExamples/PostProcessing/1_14/1_14.INT.udn
@@ -51,5 +51,5 @@ Related: Engine/Rendering/PostProcessEffects/PostProcessMaterials
 
 Anti-Aliasing improves image quality by removing the stepping effect from the edges that appear angled on screen. **Temporal AA** is the default and gives the best result. Although it is on by default, it only shows in the editor when Real-time is activated in viewport or while in Game Mode **(G-key)**. 
 
-Temporal AA does not currently work with the console command "HighResShot" to make screenshots that are bigger than the screen resolution. You have to turn **Temporal AA** off or use FXAA to make those screenshots.
+Temporal AA does not currently work with the console command "HighResShot" to make screenshots that are bigger than the screen resolution. You have to turn Temporal AA off or use FXAA to make those screenshots.
 
