#UE4doc #docs: Removed extraneous bolding as according to the instructional guide.

diff --git a/Engine/Documentation/Source/Resources/ContentExamples/PostProcessing/1_16/1_16.INT.udn b/Engine/Documentation/Source/Resources/ContentExamples/PostProcessing/1_16/1_16.INT.udn
--- a/Engine/Documentation/Source/Resources/ContentExamples/PostProcessing/1_16/1_16.INT.udn
+++ b/Engine/Documentation/Source/Resources/ContentExamples/PostProcessing/1_16/1_16.INT.udn
@@ -53,5 +53,5 @@ Related: Engine/Rendering/PostProcessEffects/PostProcessMaterials
 
 [INCLUDE:Engine/Rendering\PostProcessEffects\PostProcessMaterials#PostProcessMaterials_Desc]
 
-**Blendables**, also known as **PostProcessMaterials**, are Materials that can read the buffers then compose and modify them before they output on the screen. The easiest way to think of them is as a material overlay on the screen.
+**Blendables**, also known as PostProcessMaterials, are Materials that can read the buffers then compose and modify them before they output on the screen. The easiest way to think of them is as a material overlay on the screen.
 
