#UE4doc #docs: Removed extraneous bolding as according to the instructional guide.

diff --git a/Engine/Documentation/Source/Resources/ContentExamples/PostProcessing/1_3/1_3.INT.udn b/Engine/Documentation/Source/Resources/ContentExamples/PostProcessing/1_3/1_3.INT.udn
--- a/Engine/Documentation/Source/Resources/ContentExamples/PostProcessing/1_3/1_3.INT.udn
+++ b/Engine/Documentation/Source/Resources/ContentExamples/PostProcessing/1_3/1_3.INT.udn
@@ -43,7 +43,7 @@ Description:An overview of the Post Processing example level, example 1.3: Scene
 [/VAR]
 
 
-The **Scene Color** category on a Post Process Volume covers a variety of effects that modify the color of the scene in some way other than what is presented in the **Film** category. Where the options in the **Film** category give you access to crush shadows, re-mix your color channels, and desaturate the scene; **Scene Color** is about **Vignette**, **Scene Fringe**, and **Tone Mapping**.
+The **Scene Color** category on a Post Process Volume covers a variety of effects that modify the color of the scene in some way other than what is presented in the **Film** category. Where the options in the Film category give you access to crush shadows, re-mix your color channels, and desaturate the scene; Scene Color is about **Vignette**, **Scene Fringe**, and **Tone Mapping**.
 
 ## Vignette
 
@@ -71,4 +71,4 @@ The **Scene Color** category on a Post Process Volume covers a variety of effect
 
 ## Scene Color Tint
 
-This option allows you to simply "tint" the scene. It is not about balancing the color channels like the color mixers in **Film**, this is just raw "tint the scene to the color set here."
+This option allows you to simply "tint" the scene. It is not about balancing the color channels like the color mixers in Film, this is just raw "tint the scene to the color set here."
