#UE4doc #docs:
Updated Engine version
Expanded on some abbreviations / acronyms upon their first use.
Bolded a term.

diff --git a/Engine/Documentation/Source/Programming/Rendering/Programming.INT.udn b/Engine/Documentation/Source/Programming/Rendering/Programming.INT.udn
--- a/Engine/Documentation/Source/Programming/Rendering/Programming.INT.udn
+++ b/Engine/Documentation/Source/Programming/Rendering/Programming.INT.udn
@@ -2,12 +2,12 @@ Availability:Public
 Title:Graphics Programming
 Crumbs:%ROOT%, Programming
 Description:Information for graphics programmers working with the rendering systems and writing shaders.
-Version: 4.5
+Version: 4.9
 
 
 ## Getting Started
 
-There is a lot of rendering code in UE4 so it is hard to get a quick high level view of what is going on. A good place to start reading through the code is `FDeferredShadingSceneRenderer::Render`, which is where a new frame is rendered on the rendering thread. It is also useful to do a 'profilegpu' command and look through the draw events. You can then do a **Find in Files** in Visual Studio on the draw event name to find the corresponding C++ implementation.
+There is a lot of rendering code in Unreal Engine 4 (UE4) so it is hard to get a quick high level view of what is going on. A good place to start reading through the code is `FDeferredShadingSceneRenderer::Render`, which is where a new frame is rendered on the rendering thread. It is also useful to do a 'profilegpu' command and look through the draw events. You can then do a **Find in Files** in Visual Studio on the draw event name to find the corresponding C++ implementation.
 
 * See [Shader Development](Programming/Rendering/ShaderDevelopment) for information on working with shaders.
 
@@ -106,7 +106,7 @@ And the same classes arranged by which thread has ownership of their state.  It
 
 Primitive components are the basic unit of visibility and relevance determination. For example, occlusion and frustum culling happen on a per-primitive basis. Therefore it is important when designing a system to think about how big to make components. Each component has a bounds that is used for various operations like culling, shadow casting, and light influence determination.
 
-Components only become visible to the scene (and therefore the renderer) when they are registered.  Game thread code that changes a component's properties must call MarkRenderStateDirty() on the component to propagate the change to the rendering thread.
+Components only become visible to the scene (and therefore the renderer) when they are registered.  Game thread code that changes a component's properties must call **MarkRenderStateDirty()** on the component to propagate the change to the rendering thread.
 
 #### FPrimitiveSceneProxy and FPrimitiveSceneInfo
 
@@ -121,7 +121,7 @@ FPrimitiveSceneProxy is the rendering thread version of UPrimitiveComponent that
 
 #### Scene rendering order
 
-The renderer processes the scene in the order that it wants to composite data to the render targets. For example, the Depth only pass is rendered before the Base pass, so that HiZ will be populated to reduce shading cost in the base pass. This order is statically defined by the order pass functions are called in C++.
+The renderer processes the scene in the order that it wants to composite data to the render targets. For example, the Depth only pass is rendered before the Base pass, so that Heirarchical Z (HiZ) will be populated to reduce shading cost in the base pass. This order is statically defined by the order pass functions are called in C++.
 #### Relevance
 
 FPrimitiveViewRelevance is the information on what effects (and therefore passes) are relevant to the primitive. A primitive may have multiple elements with different relevance, so FPrimitiveViewRelevance is effectively a logical OR of all the element's relevancies. This means that a primitive can have both opaque and translucent relevance, or dynamic and static relevance; they are not mutually exclusive.
@@ -129,7 +129,7 @@ FPrimitiveViewRelevance is the information on what effects (and therefore passes
 FPrimitiveViewRelevance also indicates whether a primitive needs to use the dynamic and/or static rendering path with bStaticRelevance and bDynamicRelevance.
 #### Drawing Policies
 
-Drawing policies contain the logic to render meshes with pass specific shaders. They use the FVertexFactory interface to abstract the mesh type, and the FMaterial interface to abstract the material details. At the lowest level, a drawing policy takes a set of mesh material shaders and a vertex factory, binds the vertex factory's buffers to the RHI, binds the mesh material shaders to the RHI, sets the appropriate shader parameters, and issues the RHI draw call.
+Drawing policies contain the logic to render meshes with pass specific shaders. They use the FVertexFactory interface to abstract the mesh type, and the FMaterial interface to abstract the material details. At the lowest level, a drawing policy takes a set of mesh material shaders and a vertex factory, binds the vertex factory's buffers to the Rendering Hardware Interface (RHI), binds the mesh material shaders to the RHI, sets the appropriate shader parameters, and issues the RHI draw call.
 #### Drawing Policy methods
 
 |Function| Description|
