#UE4 Docs: reorg quick search terms and topic links

diff --git a/Engine/Documentation/Source/UE4Home.INT.udn b/Engine/Documentation/Source/UE4Home.INT.udn
--- a/Engine/Documentation/Source/UE4Home.INT.udn
+++ b/Engine/Documentation/Source/UE4Home.INT.udn
@@ -27,34 +27,34 @@ Navigation:home
 			How-To
 		[/PARAMLITERAL]
 		[PARAMLITERAL:terms]
-			Actor
+			C++
 		[/PARAMLITERAL]
 		[PARAMLITERAL:terms]
-			Component
+			Blueprint
 		[/PARAMLITERAL]
 		[PARAMLITERAL:terms]
-			Blueprint
+			iOS
 		[/PARAMLITERAL]
 		[PARAMLITERAL:terms]
-			Materials
+			Android
 		[/PARAMLITERAL]
 		[PARAMLITERAL:terms]
-			Lighting
+			Actor
 		[/PARAMLITERAL]
 		[PARAMLITERAL:terms]
-			Camera
+			Component
 		[/PARAMLITERAL]
 		[PARAMLITERAL:terms]
-			UMG UI
+			Materials
 		[/PARAMLITERAL]
 		[PARAMLITERAL:terms]
-			C++
+			Lighting
 		[/PARAMLITERAL]
 		[PARAMLITERAL:terms]
-			iOS
+			Camera
 		[/PARAMLITERAL]
 		[PARAMLITERAL:terms]
-			Android
+			UMG UI
 		[/PARAMLITERAL]
 	[/OBJECT]
 [/REGION]
@@ -68,11 +68,11 @@ Navigation:home
 	%GettingStarted:role%
 	%Engine/Editor:role%
 	%Engine:role%
-	%Gameplay:role%
-	%Engine/Blueprints:role%
 	%Programming:topicrole%
-	%Platforms:role%
+	%Engine/Blueprints:role%
+	%Gameplay:role%
 	%Resources:role%
+	%Platforms:role%
 	%Support/Builds:role%
 	[/REGION]
 [/REGION]
