#UE4doc #docs: Edited "Note:" to match instructional guide.

diff --git a/Engine/Documentation/Source/Engine/Rendering/LightingAndShadows/AmbientCubemap/AmbientCubemap.INT.udn b/Engine/Documentation/Source/Engine/Rendering/LightingAndShadows/AmbientCubemap/AmbientCubemap.INT.udn
--- a/Engine/Documentation/Source/Engine/Rendering/LightingAndShadows/AmbientCubemap/AmbientCubemap.INT.udn
+++ b/Engine/Documentation/Source/Engine/Rendering/LightingAndShadows/AmbientCubemap/AmbientCubemap.INT.udn
@@ -29,7 +29,9 @@ The **Ambient Cubemap** lights the scene from a provided image. The image is map
 The blurry versions are computed in a way that they can be used for specular lighting with varying glossiness (sharp vs. blurry reflections) and they can be used for diffuse lighting as well. 
 This effect is independent of the position a material is lit from. Viewer position, material roughness (for specular effects), and the material surface normal are all taken into account. 
 
-**Note: The Ambient Cubemap has very limited local shadowing (SSAO only).  Use a [Sky Light](Engine/Rendering/LightingAndShadows/LightTypes/SkyLight) instead when representing the sky's lighting.**  
+[REGION:note]
+The Ambient Cubemap has very limited local shadowing (SSAO only).  Use a [Sky Light](Engine/Rendering/LightingAndShadows/LightTypes/SkyLight) instead when representing the sky's lighting.
+[/REGION]
 [/EXCERPT:AmbientCubemap_Desc]
 
 The Ambient Cubemap is only intended to be used for a subtle directional ambient term (useful in games with completely dynamic lighting), or for model viewer applications.
