UEDOC-1796

Updated GIF to Youtube video

diff --git a/Engine/Documentation/Source/Engine/Rendering/Materials/HowTo/AnimatingUVCoords/AnimatingUVCoordinates.INT.udn b/Engine/Documentation/Source/Engine/Rendering/Materials/HowTo/AnimatingUVCoords/AnimatingUVCoordinates.INT.udn
--- a/Engine/Documentation/Source/Engine/Rendering/Materials/HowTo/AnimatingUVCoords/AnimatingUVCoordinates.INT.udn
+++ b/Engine/Documentation/Source/Engine/Rendering/Materials/HowTo/AnimatingUVCoords/AnimatingUVCoordinates.INT.udn
@@ -86,7 +86,17 @@ content page for information about how to move content between projects. This wa
 
 1. With the texture now placed in the Material Graph, go to the **Palette** and search for **Panner**. Once you find the **Panner** Material Expression, select it with the **Left Mouse Button** and then drag it into the Material Graph. Once over the Material Graph, release the  **Left Mouse Button**  to place the **Panner** Material Expression in the Material Graph. Once completed, your Material Graph should look like this. 
 	
-	![](AMUC_Adding_Panner_Node.gif)(convert:false)
+	[OBJECT:EmbeddedVideo]
+	[PARAMLITERAL:width]
+	640
+	[/PARAMLITERAL]
+	[PARAMLITERAL:height]
+	360
+	[/PARAMLITERAL]
+	[PARAMLITERAL:videoid]
+	Adj0C-FY_Zg
+	[/PARAMLITERAL]
+	[/OBJECT]
 
 1. Now with all the required nodes added, it is now time to start connecting everything together. First plug the **White** Output of the Fire Texture into the **Base Color**  input. Then connect the Output of the **Panner** to the **UVs** input on the Fire Texture. When completed, you should have something that looks like this. 
 
