#UE4doc #docs:
Updated Engine version
Unbolded and unitalicized a few terms

diff --git a/Engine/Documentation/Source/Programming/Rendering/ParticleEmitters/CreatingParticleEmitters.INT.udn b/Engine/Documentation/Source/Programming/Rendering/ParticleEmitters/CreatingParticleEmitters.INT.udn
--- a/Engine/Documentation/Source/Programming/Rendering/ParticleEmitters/CreatingParticleEmitters.INT.udn
+++ b/Engine/Documentation/Source/Programming/Rendering/ParticleEmitters/CreatingParticleEmitters.INT.udn
@@ -2,7 +2,7 @@ Availability:Public
 Title:Particle Emitter Technical Guide
 Crumbs:%ROOT%, Programming, Programming/Rendering
 Description:Programmer's guide to adding new particle emitter types.
-Version: 4.5
+Version: 4.9
 
 [TOC (start:2 end:3)]
 
@@ -52,7 +52,7 @@ The `ParticleEmitterInstance` struct contains the following public variables:
 | `BurstFired` | An array of entries for tracking the firing of bursts. |
 | `LoopCount` | The number of loops completed by the instance. |
 | `IsRenderDataDirty` | Flag indicating if the render data is dirty. |
-| `Module_AxisLock` | The **AxisLock** module, if present. Held to avoid searching for it each Tick. |
+| `Module_AxisLock` | The AxisLock module, if present. Held to avoid searching for it each Tick. |
 | `EmitterDuration` | The current duration of the instance. |
 | `EmitterDurations` | An array of the durations at each LOD level of the instance. |
 | `TrianglesToRender` | The number of triangles the emitter will render this frame. |
@@ -72,14 +72,14 @@ The ParticleEmitterInstance struct contains the following member functions which
 | [REGION:wrap]`void Tick(float DeltaTime, bool bSuppressSpawning)`[/REGION] | Tick the instance with the given time slice. If `bSuppressSpawning` is `true`, do not spawn any new particles. |
 | [REGION:wrap]`void Rewind()`[/REGION] | Rewind the emitter instance. |
 | [REGION:wrap]`FBox GetBoundingBox()`[/REGION] | Return the bounding box for the instance. |
-| [REGION:wrap]`void UpdateBoundingBox(float DeltaTime)`[/REGION] | Update the bounding box for the instance. **This is where the final positioning of the particles takes place for the frame as all updates are guaranteed to have occurred by now.** |
+| [REGION:wrap]`void UpdateBoundingBox(float DeltaTime)`[/REGION] | Update the bounding box for the instance. This is where the final positioning of the particles takes place for the frame as all updates are guaranteed to have occurred by now. |
 | [REGION:wrap]`uint32 RequiredBytes()`[/REGION] | Retrieve the number of per-particle bytes that the emitter requires. |
 | [REGION:wrap]`uint8* GetModuleInstanceData(UParticleModule* Module)`[/REGION] | Get the pointer to the per-instance data allocated for the given module. |
 | [REGION:wrap]`uint32 CalculateParticleStride(uint32 ParticleSize)`[/REGION] | Calculate the stride of a single particle for this instance. |
 | [REGION:wrap]`void ResetBurstList()`[/REGION] | Reset the burst list information for the instance. |
 | [REGION:wrap]`float GetCurrentBurstRateOffset(float& DeltaTime, int32& Burst)`[/REGION] | Get the current burst rate offset - artificially increases `DeltaTime` to generate the bursts. |
 | [REGION:wrap]`float Spawn(float OldLeftover, float Rate, float DeltaTime, int32 Burst = 0, float BurstTime = 0.0f)`[/REGION] | Spawn particles for the instance given the current time slice (DeltaTime). Take the leftover from last from (OldLeftover) into account. |
-| [REGION:wrap]`void PreSpawn(FBaseParticle* Particle)`[/REGION] | Handle any _pre-spawning_ actions required for the particles in this instance. |
+| [REGION:wrap]`void PreSpawn(FBaseParticle* Particle)`[/REGION] | Handle any pre-spawning actions required for the particles in this instance. |
 | [REGION:wrap]`bool HasCompleted()`[/REGION] | Return `true` if the instance has completed its run. |
 | [REGION:wrap]`void PostSpawn(FBaseParticle* Particle, float InterpolationPercentage, float SpawnTime)`[/REGION] | Handle any _post-spawning_ actions required for the particles in this instance. |
 | [REGION:wrap]`void KillParticles()`[/REGION] | Kill off any dead particle - simply remove them from the active array. |
@@ -91,7 +91,7 @@ The ParticleEmitterInstance struct contains the following member functions which
 
 ### The ParticleModuleTypeDataBase class
 
-The `ParticleModuleTypeDataBase` class provides the mechanism for generating custom emitter instances when creating a **ParticleSystem** for use in the engine. For example, the `ParticleModuleTypeDataMesh` class will result in a `FParticleMeshEmitterInstance` getting created for the **ParticleSystem**.
+The `ParticleModuleTypeDataBase` class provides the mechanism for generating custom emitter instances when creating a **ParticleSystem** for use in the engine. For example, the `ParticleModuleTypeDataMesh` class will result in a `FParticleMeshEmitterInstance` getting created for the ParticleSystem.
 
 
 #### Member Functions
