#UE4doc #docs: Removed extraneous bolding as according to the instructional guide.

diff --git a/Engine/Documentation/Source/Resources/ContentExamples/PostProcessing/1_9/1_9.INT.udn b/Engine/Documentation/Source/Resources/ContentExamples/PostProcessing/1_9/1_9.INT.udn
--- a/Engine/Documentation/Source/Resources/ContentExamples/PostProcessing/1_9/1_9.INT.udn
+++ b/Engine/Documentation/Source/Resources/ContentExamples/PostProcessing/1_9/1_9.INT.udn
@@ -53,7 +53,7 @@ Related: Engine/Rendering/PostProcessEffects/PostProcessMaterials
 [/REGION]
 
 
-**Screen Space Reflection** is an effect which is enabled by default that alters the appearance of objects on the surface of **Materials**. There are only a handful of options that can be tweaked for this effect: Intensity, Quality, and Max Roughness. See the Screen Space Reflection page link located below for more information on these options.  
+**Screen Space Reflection** is an effect which is enabled by default that alters the appearance of objects on the surface of Materials. There are only a handful of options that can be tweaked for this effect: Intensity, Quality, and Max Roughness. See the Screen Space Reflection page link located below for more information on these options.  
 
 
 
