#UE4doc #docs: Removed extraneous bolding as according to the instructional guide.

diff --git a/Engine/Documentation/Source/Resources/ContentExamples/Reflections/Reflections.INT.udn b/Engine/Documentation/Source/Resources/ContentExamples/Reflections/Reflections.INT.udn
--- a/Engine/Documentation/Source/Resources/ContentExamples/Reflections/Reflections.INT.udn
+++ b/Engine/Documentation/Source/Resources/ContentExamples/Reflections/Reflections.INT.udn
@@ -47,31 +47,28 @@ Related: Engine/Rendering/PostProcessEffects
 
 [TOC(start:2 end:2)]
 
+The **Reflections** Content Example map provides examples on how to adjust the way reflections appear within your levels by using a Reflection Capture Actor (which projects a reflection within the Actor's effect volume or radius). 
+Examples are also shown on how the Roughness, Metallic, and Specular properties of a Material can be used to produce different reflection results. 
 
-## Overview
-
-The **Reflections** Content Example map provides examples on how to adjust the way reflections appear within your levels by using a **Reflection Capture Actor** (which projects a reflection within the **Actor's** effect volume or radius). 
-Examples are also shown on how the **Roughness**, **Metallic**, and **Specular** properties of a **Material** can be used to produce different reflection results. 
-
-In addition to tweaking reflection settings, the use of both a **Box Reflection Actor** and the more commonly used **Sphere Reflection Actor** are shown with sample effects that can be achieved.
-Lastly, **Scene Capture** methods are demonstrated which can be used to capture a fully dynamic cube map or capture an image of the scene on every frame. 
+In addition to tweaking reflection settings, the use of both a Box Reflection Actor and the more commonly used*Sphere Reflection Actor are shown with sample effects that can be achieved.
+Lastly, Scene Capture methods are demonstrated which can be used to capture a fully dynamic cube map or capture an image of the scene on every frame. 
 
 
 ## Reflections Map
 
 ![](Reflections_Map.png)(w:640)
 
-Listed below are the examples provided inside the **Reflections** map:
+Listed below are the examples provided inside the Reflections map:
 
 | Example | What is Demonstrated |
 | ------- | ------------ |
-| **[1.1 Reflections Metallic Materials](Resources\ContentExamples\Reflections\1_1)** | Metallic **Materials** reflecting objects based on the **Metallic** and **Roughness** properties (increasing **Metallic** property results in more chrome like appearance). |
-| **[1.2 Reflections Non-Metallic Materials](Resources\ContentExamples\Reflections\1_2)** | Non-metallic **Materials** (0.0 for the **Metallic** property) reflecting objects based on the **Roughness** and **Specular** values of the **Material**. |
-| **[1.3 Box Reflection Capture Actor (Static)](Resources\ContentExamples\Reflections\1_3)** | Demonstrates reflection using the **Box Reflection Capture Actor** which is good for use with interiors that have a box shape. |
-| **[1.4 Sphere Reflection Capture Actor (Static)](Resources\ContentExamples\Reflections\1_4)** | The more commonly used **Sphere Reflection Capture Actor** used to reflect the environment, which blends better and is more flexible than the **Box Reflection Capture**. |
-| **[1.5 Screenspace Reflections (Dynamic)](Resources\ContentExamples\Reflections\1_5)** | **Screenspace Reflections** are a dynamic **Post Process** effect that can only reflect what is on screen. |
-| **[1.6 Scene Capture Cube (Dynamic)](Resources\ContentExamples\Reflections\1_6)** | **Scene Capture Cubes** capture a fully dynamic cube map from their location on every frame (these are expensive performance wise and should only be used when absolutely required to achieve a given effect). |
-| **[1.7 Scene Capture 2D (Dynamic)](Resources\ContentExamples\Reflections\1_7)** | **Scene Capture 2D** captures an image of the scene on every frame and projects the image in a **Material** (they work more like a TV screen receiving an image from a camera and are most commonly used for security camera effects). |
+| **[1.1 Reflections Metallic Materials](Resources\ContentExamples\Reflections\1_1)** | Metallic Materials reflecting objects based on the Metallic and Roughness properties (increasing Metallic property results in more chrome like appearance). |
+| **[1.2 Reflections Non-Metallic Materials](Resources\ContentExamples\Reflections\1_2)** | Non-metallic Materials (0.0 for the Metallic property) reflecting objects based on the Roughness and Specular values of the Material. |
+| **[1.3 Box Reflection Capture Actor (Static)](Resources\ContentExamples\Reflections\1_3)** | Demonstrates reflection using the Box Reflection Capture Actor which is good for use with interiors that have a box shape. |
+| **[1.4 Sphere Reflection Capture Actor (Static)](Resources\ContentExamples\Reflections\1_4)** | The more commonly used Sphere Reflection Capture Actor used to reflect the environment, which blends better and is more flexible than the Box Reflection Capture. |
+| **[1.5 Screenspace Reflections (Dynamic)](Resources\ContentExamples\Reflections\1_5)** | Screenspace Reflections are a dynamic Post Process effect that can only reflect what is on screen. |
+| **[1.6 Scene Capture Cube (Dynamic)](Resources\ContentExamples\Reflections\1_6)** | Scene Capture Cubes capture a fully dynamic cube map from their location on every frame (these are expensive performance wise and should only be used when absolutely required to achieve a given effect). |
+| **[1.7 Scene Capture 2D (Dynamic)](Resources\ContentExamples\Reflections\1_7)** | Scene Capture 2D captures an image of the scene on every frame and projects the image in a Material (they work more like a TV screen receiving an image from a camera and are most commonly used for security camera effects). |
 
 
 
