#UE4doc #docs: Removed extraneous bolding as according to the instructional guide. ##Overview removed as according to the instructional guide.

diff --git a/Engine/Documentation/Source/Resources/ContentExamples/Reflections/1_1/1_1.INT.udn b/Engine/Documentation/Source/Resources/ContentExamples/Reflections/1_1/1_1.INT.udn
--- a/Engine/Documentation/Source/Resources/ContentExamples/Reflections/1_1/1_1.INT.udn
+++ b/Engine/Documentation/Source/Resources/ContentExamples/Reflections/1_1/1_1.INT.udn
@@ -51,9 +51,7 @@ Related: Engine/Rendering/PostProcessEffects
 
 ![](Reflections_1_1_Header.png)
 
-## Overview
-
-**Reflection Capture** Actors, either boxes or spheres, project a reflection within their effect volume or radius. For metallic materials, the greater the **Metallic** property, the more chrome like the material will appear. These Materials gain a vast majority of the color from their reflective properties, while their **Base Color** will colorize the reflection.
+**Reflection Capture** Actors, either boxes or spheres, project a reflection within their effect volume or radius. For metallic materials, the greater the Metallic property, the more chrome like the material will appear. These Materials gain a vast majority of the color from their reflective properties, while their Base Color will colorize the reflection.
 
 
 [/REGION]
