#UE4doc #docs: Removed extraneous bolding as according to the instructional guide. ##Overview removed as according to the instructional guide.

diff --git a/Engine/Documentation/Source/Resources/ContentExamples/Reflections/1_2/1_2.INT.udn b/Engine/Documentation/Source/Resources/ContentExamples/Reflections/1_2/1_2.INT.udn
--- a/Engine/Documentation/Source/Resources/ContentExamples/Reflections/1_2/1_2.INT.udn
+++ b/Engine/Documentation/Source/Resources/ContentExamples/Reflections/1_2/1_2.INT.udn
@@ -51,10 +51,9 @@ Related: Engine/Rendering/PostProcessEffects
 [REGION:fullwidth]
 
 ![](Reflections_1_2_Header.png)
-## Overview
 
-**Reflection Capture** Actors, either box or sphere, project a reflection within their effect volume or radius. Non-metallic materials (Materials with 0.0 in the **Metallic** property) will reflect based on their **Roughness** and **Specular** values. With higher Roughness decreasing and lower specularity reducing the reflective properties of the Material. **Roughness** and **Specular** affect the reflective properties of a surface differently, and should be adjusted until a desired effect is achieved.
+**Reflection Capture** Actors, either box or sphere, project a reflection within their effect volume or radius. Non-metallic materials (Materials with 0.0 in the Metallic property) will reflect based on their Roughness and Specular values. With higher Roughness decreasing and lower specularity reducing the reflective properties of the Material. Roughness and Specular affect the reflective properties of a surface differently, and should be adjusted until a desired effect is achieved.
 
-It is also worth noting that non-metallic materials gain the majority of the color from their **Base Color** property.
+It is also worth noting that non-metallic materials gain the majority of the color from their Base Color property.
 
 [/REGION]
