#UE4doc #docs: Removed extraneous bolding as according to the instructional guide. ##Overview removed as according to the instructional guide.

diff --git a/Engine/Documentation/Source/Resources/ContentExamples/Reflections/1_3/1_3.INT.udn b/Engine/Documentation/Source/Resources/ContentExamples/Reflections/1_3/1_3.INT.udn
--- a/Engine/Documentation/Source/Resources/ContentExamples/Reflections/1_3/1_3.INT.udn
+++ b/Engine/Documentation/Source/Resources/ContentExamples/Reflections/1_3/1_3.INT.udn
@@ -50,11 +50,9 @@ Related: Engine/Rendering/PostProcessEffects
 [REGION:fullwidth]
 
 ![](Reflections_1_3_Header.png)
-## Overview
-
 
 The **Box Reflection Capture Actor** is one of two Reflection Capture Actors available in Unreal Engine 4. Many of these can be placed around the world with almost no performance impact as they are calculated before run time. While the box shape is a less commonly used in comparison to the sphere shape, it is very good for interiors that have a box shape themselves.
 
-One big draw back to consider for **Box Reflection Captures** is that if there are objects close to the Actor center (near the reflection sprite in the center of the box volume), they will appear very large in the reflection. This is due to the fact that the image used from the reflection is calculated from that sprite location, meaning that objects closer to that point will exhibit perspective-like depth from that point.
+One big drawback to consider for Box Reflection Captures is that if there are objects close to the Actor center (near the reflection sprite in the center of the box volume), they will appear very large in the reflection. This is due to the fact that the image used from the reflection is calculated from that sprite location, meaning that objects closer to that point will exhibit perspective-like depth from that point.
 
 [/REGION]
