#UE4doc #docs:
Removed Overview header
Updated Engine version
Changed font settings for multiple terms/lines

diff --git a/Engine/Documentation/Source/Programming/Rendering/ParticleModules/ParticleModules.INT.udn b/Engine/Documentation/Source/Programming/Rendering/ParticleModules/ParticleModules.INT.udn
--- a/Engine/Documentation/Source/Programming/Rendering/ParticleModules/ParticleModules.INT.udn
+++ b/Engine/Documentation/Source/Programming/Rendering/ParticleModules/ParticleModules.INT.udn
@@ -2,19 +2,19 @@ Availability:Public
 Title:Particle Module Technical Guide
 Crumbs:%ROOT%, Programming, Programming/Rendering
 Description:Programmer's guide to adding new particle modules for customizing the behavior of particle systems.
-Version: 4.5
+Version: 4.9
 
 [TOC (start:2 end:3)]
 
 
 
-## Overview
 
-A particle system in UE4 consists of any number of particle emitters, each of which contains modules that effect how its particles behave. Extending the system with custom modules and emitter types is quite simple and this doc will give examples of how to do so.
+
+A particle system in Unreal Editor 4 (UE4) consists of any number of particle emitters, each of which contains modules that effect how its particles behave. Extending the system with custom modules and emitter types is quite simple and this doc will give examples of how to do so.
 
 ## The ParticleModule base class 
 
-All particle modules derive from the same base class, `ParticleModule` (defined in `//depot/UE4/Engine/Source/Runtime/Engine/Classes/Particles/Modules/ParticleModule.h`). 
+All particle modules derive from the same base class, `ParticleModule` (defined in //depot/UE4/Engine/Source/Runtime/Engine/Classes/Particles/Modules/ParticleModule.h). 
 
 
 ### Member Variables
@@ -36,10 +36,10 @@ A `ModuleType` enumeration is also defined in this class. This type indicates wh
 
 | Type (EPMT) | Description |
 | ------------ | ----------- |
-| _General_ | A general module that is usable by all emitter types. |
-| _TypeData_ | The module is a TypeData module - dictating the emitter class that gets instanced. |
-| _Beam_ | The module is only usable by beam emitters. |
-| _Trail_ | The module is only usable by trail emitters. |
+| General | A general module that is usable by all emitter types. |
+| TypeData | The module is a TypeData module - dictating the emitter class that gets instanced. |
+| Beam | The module is only usable by beam emitters. |
+| Trail | The module is only usable by trail emitters. |
 
 
 ### Member Functions
@@ -58,7 +58,7 @@ The `ParticleModule` class contains several virtual member functions that provid
 
 Writing a custom module simply involves overriding the functions mentioned above to implement the desired effect. As an example, a module that implements setting the particle color to the base color times the velocity scaled by some distribution-generated factor will be created.
 
-The base class of the module defines the _heading_ it will go under in the **right-click** module menu of Cascade. So, in our example, we will want to derive from the `ParticleModuleColorBase` class to ensure that the module shows up in the **Color** sub-menu.
+The base class of the module defines the heading it will go under in the **right-click** module menu of Cascade. So, in our example, we will want to derive from the `ParticleModuleColorBase` class to ensure that the module shows up in the **Color** sub-menu.
 
 ### Module Declaration
 
@@ -109,7 +109,7 @@ The constructor creates a `DistributionVectorConstant` to assign to `ScaleVeloci
 		ScaleVelocity.Distribution = DistributionScaleVelocity;
 	}
 
-In the `Spawn()` function, the code will be setting up any initial effect it has on particles. Not every module will require a `Spawn()` call, but this one does.
+In the **Spawn()** function, the code will be setting up any initial effect it has on particles. Not every module will require a `Spawn()` call, but this one does.
 
 	void UParticleModuleColorByVelocity::Spawn(FParticleEmitterInstance* Owner, int32 Offset, float SpawnTime)
 	{
@@ -123,7 +123,7 @@ In the `Spawn()` function, the code will be setting up any initial effect it has
 	   }
 	}
 
-The `Spawn()` function retrieves the `ScaleVelocity` value from the distribution using the `Particle.RelativeTime`. This value is multiplied by the velocity to give the scaled velocity value. The `Particle.Color` is then set to the `BaseColor` multiplied by the result of the scaled velocity calculation. 
+The **Spawn()** function retrieves the `ScaleVelocity` value from the distribution using the `Particle.RelativeTime`. This value is multiplied by the velocity to give the scaled velocity value. The `Particle.Color` is then set to the `BaseColor` multiplied by the result of the scaled velocity calculation. 
 
 `SPAWN_INIT` is a macro which sets up a `FBaseParticle` reference to the particle being spawned, as well as some other commonly used values when accessing particle data, such as offset tracking into the particle payload, etc. For full details, please refer to `//depot/UE4/Engine/Source/Runtime/Engine/Public/ParticleHelper.h`.
 
@@ -131,7 +131,7 @@ The `Spawn()` function retrieves the `ScaleVelocity` value from the distribution
 	This is a destructive module in that it sets the `Particle.Color` directly. Any module that came before it that only modified the `Particle.Color` would be pointless. However, if it modified the `BaseColor` as well, such as the **InitialColor** module, then it would still have an impact.
 [/REGION]
 
-The `Update()` function in this particular module is more or less identical to the `Spawn()`. The difference is that each active particle is updated in a single loop. 
+The **Update()** function in this particular module is more or less identical to the **Spawn()**. The difference is that each active particle is updated in a single loop. 
 
 	void UParticleModuleColorByVelocity::Update(FParticleEmitterInstance* Owner, int32 Offset, float DeltaTime)
 	{
@@ -148,7 +148,7 @@ The `Update()` function in this particular module is more or less identical to t
 
 `BEGIN_UPDATE_LOOP`/`END_UPDATE_LOOP` are macros that set up a code block that loops through all active particles and executes the code contained between them on each particle. See `UnParticleHelper.h` for full details.
 
-The following screen shot shows the **ColorByVelocity** module in action. The **InitialVelocity** of the particles is randomly set to **[(0,0,0),(200,200,0)]**. The **InitialColor** of each particle is set to white. The **VelocityScale** value of the **ColorByVelocity** is set to a constant of **(0.005, 0.005, 0)**. This results in each particle having its color set to   
+The following screen shot shows the **ColorByVelocity** module in action. The **InitialVelocity** of the particles is randomly set to **[(0,0,0),(200,200,0)]**. The **InitialColor** of each particle is set to white. The **VelocityScale** value of the **ColorByVelocity** is set to a constant of **(0.005, 0.005, 0)**. This results in each particle having its color set to:   
     
 	Velocity * VelocityScale
 
