#UE4doc #docs: Removed extraneous bolding as according to the instructional guide. ##Overview removed as according to the instructional guide.

diff --git a/Engine/Documentation/Source/Resources/ContentExamples/Reflections/1_4/1_4.INT.udn b/Engine/Documentation/Source/Resources/ContentExamples/Reflections/1_4/1_4.INT.udn
--- a/Engine/Documentation/Source/Resources/ContentExamples/Reflections/1_4/1_4.INT.udn
+++ b/Engine/Documentation/Source/Resources/ContentExamples/Reflections/1_4/1_4.INT.udn
@@ -51,11 +51,10 @@ Related: Engine/Rendering/PostProcessEffects
 [REGION:fullwidth]
 
 ![](Reflections_1_4_Header.png)
-## Overview
 
-The primarily used reflection capture is the **Sphere Reflection Capture**. Many of these can be placed around the world with almost no performance impact as they are calculated before run time. The **Sphere Reflection Capture** blends nicely into the scene and is far more flexible to use in comparison to the **Box Reflection Capture**.
+The primarily used reflection capture is the **Sphere Reflection Capture**. Many of these can be placed around the world with almost no performance impact as they are calculated before run time. The Sphere Reflection Capture blends nicely into the scene and is far more flexible to use in comparison to the Box Reflection Capture.
 
-In this example, you can clearly see how the captures are blending together on the surface, with a very noticeable transition from small example one, to the full room capture that is in place. By increasing the size of the example **Sphere Reflection Capture**, you can see how it will replace the other capture.
+In this example, you can clearly see how the captures are blending together on the surface, with a very noticeable transition from small example one, to the full room capture that is in place. By increasing the size of the example Sphere Reflection Capture, you can see how it will replace the other capture.
 
 
 [/REGION]
