#UE4doc #docs: ##Overview removed as according to the instructional guide.

diff --git a/Engine/Documentation/Source/Resources/ContentExamples/Reflections/1_5/1_5.INT.udn b/Engine/Documentation/Source/Resources/ContentExamples/Reflections/1_5/1_5.INT.udn
--- a/Engine/Documentation/Source/Resources/ContentExamples/Reflections/1_5/1_5.INT.udn
+++ b/Engine/Documentation/Source/Resources/ContentExamples/Reflections/1_5/1_5.INT.udn
@@ -50,8 +50,6 @@ Related: Engine/Rendering/PostProcessEffects
 [REGION:fullwidth]
 
 ![](Reflections_1_5_Header.png)
-## Overview
-
 
 **Screenspace Reflections** are a dynamic post process effect. While it can only reflect what is on screen, thusly creating some odd artifacts on perfectly reflective surfaces. However, by adding a bit of roughness to a Material, these artifacts will be blended away.
 
