#UE4doc #docs: ##Overview removed as according to the instructional guide.

diff --git a/Engine/Documentation/Source/Resources/ContentExamples/Reflections/1_6/1_6.INT.udn b/Engine/Documentation/Source/Resources/ContentExamples/Reflections/1_6/1_6.INT.udn
--- a/Engine/Documentation/Source/Resources/ContentExamples/Reflections/1_6/1_6.INT.udn
+++ b/Engine/Documentation/Source/Resources/ContentExamples/Reflections/1_6/1_6.INT.udn
@@ -51,7 +51,6 @@ Related: Engine/Rendering/PostProcessEffects
 [REGION:fullwidth]
 
 ![](Reflections_1_6_Header.png)
-## Overview
 
 **Scene Capture Cubes** capture a fully dynamic cube map from their location _on every frame_. This requires a Cubemap Render Target to actually capture to, and then that Render Target can be used as a texture within any Material. 
 
