#UE4doc #docs: Removed extraneous bolding as according to the instructional guide. ##Overview removed as according to the instructional guide.

diff --git a/Engine/Documentation/Source/Resources/ContentExamples/StaticMeshes/StaticMeshes.INT.udn b/Engine/Documentation/Source/Resources/ContentExamples/StaticMeshes/StaticMeshes.INT.udn
--- a/Engine/Documentation/Source/Resources/ContentExamples/StaticMeshes/StaticMeshes.INT.udn
+++ b/Engine/Documentation/Source/Resources/ContentExamples/StaticMeshes/StaticMeshes.INT.udn
@@ -47,33 +47,30 @@ Related: Engine/Content/FBX/BestPractices
 
 [TOC(start:2 end:2)]
 
-
-## Overview
-
-Inside the **StaticMeshes** Content Example map, you will find essential examples using **Static Meshes** including how to set **Mobility**, assign **UVs**, and add or adjust **Collision** settings. Other examples include how to customize your **Static Meshes** with **Material Elements** (which allow you to specify the **Material** applied to parts of a **Static Mesh**) as well as using the **MeshPaint Tool** to manually paint onto a **Static Mesh**. 
-Last, you will see how **LODs** can improve performance by scaling back the quality of a mesh based on player distance.
+Inside the **StaticMeshes** Content Example map, you will find essential examples using Static Meshes including how to set Mobility, assign UVs, and add or adjust Collision settings. Other examples include how to customize your Static Meshes with Material Elements (which allow you to specify the Material applied to parts of a Static Mesh) as well as using the MeshPaint Tool to manually paint onto a Static Mesh. 
+Last, you will see how LODs can improve performance by scaling back the quality of a mesh based on player distance.
 
 
 ## StaticMeshes Map
 
 ![](StaticMeshes_Map.png)(w:640)
 
-Listed below are the examples provided inside the **StaticMeshes** map:
+Listed below are the examples provided inside the StaticMeshes map:
 
 | Example | What is Demonstrated |
 | ------- | ------------ |
-| **1.1 Staticmesh** | A basic **StaticMesh** which was modeled in a 3rd Party Tool and imported as an FBX file. |
-| **1.2 Mobility** | The difference between **Static** (baked lighting/shadows, cannot be moved while playing) and **Moveable** (dynamic lighting/shadows, can move when playing) **Mobility** of **Static Meshes**.  |
-| **1.3 UV Mapping** | Examples of correct and incorrect **UV Mapping** which defines the alignment of a **Texture** in a **Material** and how it is applied to a **Static Mesh**.  |
-| **1.4 Lightmap UVs** | Examples of correct and incorrect **Lightmap UVs** which defines how a **Texture** within a **Material** reacts to light when applied to a **Static Mesh**.  |
+| **1.1 Staticmesh** | A basic StaticMesh which was modeled in a 3rd Party Tool and imported as an FBX file. |
+| **1.2 Mobility** | The difference between Static (baked lighting/shadows, cannot be moved while playing) and Moveable (dynamic lighting/shadows, can move when playing) Mobility of Static Meshes.  |
+| **1.3 UV Mapping** | Examples of correct and incorrect UV Mapping which defines the alignment of a Texture in a Material and how it is applied to a Static Mesh.  |
+| **1.4 Lightmap UVs** | Examples of correct and incorrect Lightmap UVs which defines how a Texture within a Material reacts to light when applied to a Static Mesh.  |
 | **1.5 Generate Unique UVs** | Three different methods of generating unique UVs: through a 3rd Party Tool, auto generated channel 1 from channel 0, auto generated new channel 1.  |
-| **1.6 Material Elements** | Illustrates how **Material Elements** allow you to assign different **Materials** to different parts of a **Static Mesh**. |
-| **1.7 Collision** | Several examples of how **Collision** can be applied to a **Static Mesh**. |
-| **1.8 Sockets** | How to attach two **Static Meshes** together with a **Socket** (e.g. a spotlight is attached to a hinge which can be rotated). |
-| **1.9 VertexColors** | Imported **VertexColors** applied to a **Static Mesh** with the **Mesh Paint** tool. |
-| **1.10 MeshPaint Tool (Vertices)** | How to paint colors on a **Static Mesh** using the **Mesh Paint** tool. |
-| **1.11 MeshPaint Tool (Texture)** | How to paint textures on a **Static Mesh** using the **Mesh Paint** tool. |
-| **1.12 LOD (Level of Detail)** | Demonstrates how the **Level of Detail** can be changed based on distance from the player. | 
+| **1.6 Material Elements** | Illustrates how Material Elements allow you to assign different Materials to different parts of a Static Mesh. |
+| **1.7 Collision** | Several examples of how Collision can be applied to a Static Mesh. |
+| **1.8 Sockets** | How to attach two Static Meshes together with a Socket (e.g. a spotlight is attached to a hinge which can be rotated). |
+| **1.9 VertexColors** | Imported VertexColors applied to a Static Mesh with the Mesh Paint tool. |
+| **1.10 MeshPaint Tool (Vertices)** | How to paint colors on a Static Mesh using the Mesh Paint tool. |
+| **1.11 MeshPaint Tool (Texture)** | How to paint textures on a Static Mesh using the Mesh Paint tool. |
+| **1.12 LOD (Level of Detail)** | Demonstrates how the Level of Detail can be changed based on distance from the player. | 
 
 
 
