#UE4doc #docs: Removed extraneous bolding as according to the instructional guide. ##Overview removed as according to the instructional guide.

diff --git a/Engine/Documentation/Source/Resources/ContentExamples/Volumes/Volumes.INT.udn b/Engine/Documentation/Source/Resources/ContentExamples/Volumes/Volumes.INT.udn
--- a/Engine/Documentation/Source/Resources/ContentExamples/Volumes/Volumes.INT.udn
+++ b/Engine/Documentation/Source/Resources/ContentExamples/Volumes/Volumes.INT.udn
@@ -46,33 +46,30 @@ Related: Engine/Actors/Triggers
 
 [TOC(start:2 end:2)]
 
+The **Volumes** Content Example map illustrates how you can use different types of Volumes to produce different gameplay effects, set boundaries, improve performance and (or) reduce lighting build times. 
+As seen in example 1.1 below, you can place a Volume by selecting any Volume from the Volumes Tab and dragging and dropping it into your level. Based on the type of Volume used, different settings may be available for editing. 
 
-## Overview
 
-The **Volumes** Content Example map illustrates how you can use different types of **Volumes** to produce different gameplay effects, set boundaries, improve performance and (or) reduce lighting build times. 
-As seen in example 1.1 below, you can place a **Volume** by selecting any **Volume** from the **Volumes Tab** and dragging and dropping it into your level. Based on the type of **Volume** used, different settings may be available for editing. 
-
-
-The most commonly used **Volumes** for gameplay purposes would be **Blocking Volumes** (to keep players confined to your play space with invisible collision), **Physics** volumes to add gameplay effects like decreasing/increasing gravity or simulating being under water, or **Pain Causing Volumes** used to inflict burst or tick damage to the **Actor** that enters it.
+The most commonly used Volumes for gameplay purposes would be Blocking Volumes (to keep players confined to your play space with invisible collision), Physics volumes to add gameplay effects like decreasing/increasing gravity or simulating being under water, or Pain Causing Volumes used to inflict burst or tick damage to the Actor that enters it.
 
 
 ## Volumes Map
 
 ![](Volumes_Map.png)(w:640)
 
-Listed below are the examples provided inside the **Volumes** map:
+Listed below are the examples provided inside the Volumes map:
 
 | Example | What is Demonstrated |
 | ------- | ------------ |
-| **1.1 Placement** | Describes how to place a **Volume** in your level and how to edit it in **Geometry Editing** mode. |
-| **1.2 Blocking Volume** | Demonstrates a how invisible collision can be created by placing a **Blocking Volume** in your level. |
-| **1.3 Cull Distance Volume** | How to make a rule set for **Cull Distances** depending on the size of the objects (e.g. have objects hidden based on distance to the player). |
-| **1.4 Lightmass Importance Volume** | Generates indirect light samples at player height for **Dynamic** objects within the volume and speeds up light rebuild times (e.g. typically used to enclose a playable area). |
+| **1.1 Placement** | Describes how to place a Volume in your level and how to edit it in Geometry Editing mode. |
+| **1.2 Blocking Volume** | Demonstrates a how invisible collision can be created by placing a Blocking Volume in your level. |
+| **1.3 Cull Distance Volume** | How to make a rule set for Cull Distances depending on the size of the objects (e.g. have objects hidden based on distance to the player). |
+| **1.4 Lightmass Importance Volume** | Generates indirect light samples at player height for Dynamic objects within the volume and speeds up light rebuild times (e.g. typically used to enclose a playable area). |
 | **1.5 Lightmass Character Indirect Detail Volume** | Similar to 1.4, however this generates indirect light samples within the entire volume not just at player height (e.g. used in elevator shaft to ensure indirect lighting is correct when it moves). |
-| **1.6 Precomputed Visibility Volume** | A performance optimization, this stores the visibility of **Actors** for its location in the world (note: should be placed only in areas a player could access). |
-| **1.7 Precomputed Visibility Override Volume** | This allows you to manually override the visibility of a **Precomputed Visibility Volume** if the auto generated visibility is undesired. |
-| **1.8 Physics Volume** | Override the physics settings of an **Actor** that enters the volume (e.g. can be used to simulate a loss of gravity or used to simulate being underwater).  |
-| **1.9 Pain Causing Volume** | **Actors** that enter a **Pain Causing Volume** register and receive damage events (e.g. can be used to apply damage from a fire). |
+| **1.6 Precomputed Visibility Volume** | A performance optimization, this stores the visibility of Actors for its location in the world (note: should be placed only in areas a player could access). |
+| **1.7 Precomputed Visibility Override Volume** | This allows you to manually override the visibility of a Precomputed Visibility Volume if the auto generated visibility is undesired. |
+| **1.8 Physics Volume** | Override the physics settings of an Actor that enters the volume (e.g. can be used to simulate a loss of gravity or used to simulate being underwater).  |
+| **1.9 Pain Causing Volume** | Actors that enter a Pain Causing Volume register and receive damage events (e.g. can be used to apply damage from a fire). |
 
 
 
