#UE4doc #docs: ##Overview removed as according to the instructional guide.

diff --git a/Engine/Documentation/Source/Resources/Downloads/Downloads.INT.udn b/Engine/Documentation/Source/Resources/Downloads/Downloads.INT.udn
--- a/Engine/Documentation/Source/Resources/Downloads/Downloads.INT.udn
+++ b/Engine/Documentation/Source/Resources/Downloads/Downloads.INT.udn
@@ -5,8 +5,6 @@ Description:Repository of downloadable content and example materials for use in
 
 [TOC]
 
-## Overview
-
 This page is a repository of downloadable content and example materials for use in developing games with UE4.
 
 ## Tutorial Content
