#UE4doc #docs: Removed extraneous bolding as according to the instructional guide.

diff --git a/Engine/Documentation/Source/Resources/SampleGames/CardGame/CardGame.INT.udn b/Engine/Documentation/Source/Resources/SampleGames/CardGame/CardGame.INT.udn
--- a/Engine/Documentation/Source/Resources/SampleGames/CardGame/CardGame.INT.udn
+++ b/Engine/Documentation/Source/Resources/SampleGames/CardGame/CardGame.INT.udn
@@ -60,11 +60,6 @@ Related: Engine/Content/Types/StaticMeshes
 
 [TOC(start:2)]
 
-
-
-
-## Overview
-
 The Memory Game sample is an example of a simple PC/mobile logic game. It is built only with _Blueprints_, with no additional code support.
 
 
@@ -75,22 +70,22 @@ The Memory Game sample is an example of a simple PC/mobile logic game. It is bui
 The **BP_LogicMem** _Blueprint_ is derived from the _GameMode_ _Class_.  This _Blueprint_ is where the Default Pawn Class, the HUDClass, the Player Controller Class, the Spectator Class, and the GameReplicationInfo Class
 are set.  
 
-There are two _Graphs_ in the **BP_LogicMem** _Blueprint_ as well.  **EventGraph** contains the initial setup of the game, the logic tests necessary for gameplay, and additional functionality such as moving back to the menu 
-and displaying all the cards with the **Hint** button. **FinishGame** handles whether the player wins or loses, the number of stars awarded to the player, and the necessary end-of-game cleanup. 
+There are two _Graphs_ in the BP_LogicMem _Blueprint_ as well.  **EventGraph** contains the initial setup of the game, the logic tests necessary for gameplay, and additional functionality such as moving back to the menu 
+and displaying all the cards with the Hint button. **FinishGame** handles whether the player wins or loses, the number of stars awarded to the player, and the necessary end-of-game cleanup. 
 
-**BP_Counter** is called by **BP_LogicMem** using a _Blueprint Interface_ to change the displayed number of moves, as well as to change the color of the text for the number of moves once there are fewer than 3 moves left.  
+**BP_Counter** is called by BP_LogicMem using a _Blueprint Interface_ to change the displayed number of moves, as well as to change the color of the text for the number of moves once there are fewer than 3 moves left.  
 
 ## Card Handling
 
-The **BP_DeckMem** and **BP_CardMem** _Blueprints_ are the primary _Blueprints_ that manage how the cards are defined and laid out in the game. One instance of **BP_DeckMem** is placed in the game level,
-and it handles the spawning of **BP_CardMem**.
+The **BP_DeckMem** and **BP_CardMem** _Blueprints_ are the primary _Blueprints_ that manage how the cards are defined and laid out in the game. One instance of BP_DeckMem is placed in the game level,
+and it handles the spawning of BP_CardMem.
 
-**BP_DeckMem** contains logic for creating and shuffling the card deck, spawning the cards, placing them on the board, and showing all the cards when the **Hint** button is pressed. The network for shuffling the 
-deck is stored in the **Shuffle** function,  which is called in the main **EventGraph** during new game preparation.  After the deck is shuffled, **BP_DeckMem** spawns the cards, with different card placement based 
-on the total number of cards.  Every location on the board has an index: the **CardDeckObjects** array contains the **BP_CardMem** _Blueprint_ that is at that index, while the **CardDeckShuffled** array holds 
-information about the card type for that index. The **Rotate** event in **BP_DeckMem** flips all the cards when the **Hint** button is pressed.  
+BP_DeckMem contains logic for creating and shuffling the card deck, spawning the cards, placing them on the board, and showing all the cards when the Hint button is pressed. The network for shuffling the 
+deck is stored in the Shuffle function,  which is called in the main **EventGraph** during new game preparation.  After the deck is shuffled, BP_DeckMem spawns the cards, with different card placement based 
+on the total number of cards.  Every location on the board has an index: the **CardDeckObjects** array contains the BP_CardMem _Blueprint_ that is at that index, while the **CardDeckShuffled** array holds 
+information about the card type for that index. The Rotate event in BP_DeckMem flips all the cards when the Hint button is pressed.  
 
-**BP_CardMem** defines the picture on each card, defines the animations for placing, flipping, and removing the cards, and sets the _ParticleEffects_ and sounds used when comparing cards.
+BP_CardMem defines the picture on each card, defines the animations for placing, flipping, and removing the cards, and sets the _ParticleEffects_ and sounds used when comparing cards.
 
 [REGION:fullwidth]
 ![](CardGame_ParticleEffects_Right.png)
@@ -98,22 +93,22 @@ information about the card type for that index. The **Rotate** event in **BP_Dec
 
 ## Click Input
 
-The **BP_PlayerControllerMem** _Blueprint_ enables the Mouse Interface.  Click events are enabled, while mouseover events are not.  The _Blueprints_ that contain click events are **BP_MenuMem**, **BP_CardMem**, and **BP_Scoreboard**.
-These **Assign OnClicked** _Nodes_ are all activated from the **ReceivedBeginPlay** event, which will fire once the player begins the game. 
+The **BP_PlayerControllerMem** _Blueprint_ enables the Mouse Interface.  Click events are enabled, while mouseover events are not.  The _Blueprints_ that contain click events are BP_MenuMem, BP_CardMem, and BP_Scoreboard.
+These Assign OnClicked _Nodes_ are all activated from the ReceivedBeginPlay event, which will fire once the player begins the game. 
 
-**Assign OnClicked** _Nodes_ create events tied to the delegate OnClicked. Delegates are ways for components of _Blueprints_ to report on their state and trigger events based on their state changes.   For example, once the **Hint Button** is clicked 
-in the game, it reports that it has been clicked, and the **Hint Button**'s OnClicked delegate triggers, activating a sequence of nodes that causes all the placed cards to flip over.
+Assign OnClicked _Nodes_ create events tied to the delegate OnClicked. Delegates are ways for components of _Blueprints_ to report on their state and trigger events based on their state changes.   For example, once the Hint Button is clicked 
+in the game, it reports that it has been clicked, and the Hint Button's OnClicked delegate triggers, activating a sequence of nodes that causes all the placed cards to flip over.
 
 The variable **BlockClick** serves to prevent players from clicking on a card while it is flipped over or while it is in the process of flipping.
 
 ## Frozen Camera
-The **BP_LogicMem** _Blueprint_ defines **SpectatorPawn** as the default Pawn Class.  The _Level Blueprint_ for the Memory Game places the **SpectatorPawn** at the PlayerStart, and the properties
-of the **SpectatorPawn** keeps the camera frozen in place throughout gameplay. 
+The BP_LogicMem _Blueprint_ defines **SpectatorPawn** as the default Pawn Class.  The _Level Blueprint_ for the Memory Game places the SpectatorPawn at the PlayerStart, and the properties
+of the SpectatorPawn keeps the camera frozen in place throughout gameplay. 
 
 ## Menu
-The **Memory Game** start menu is set up in the **BP_MenuMem** _Blueprint_ using simple _StaticMeshComponents_.  _TextRenderComponents_ are used for the text and numbers, and a **Timeline** is used to move the menu out of the game window.  
-Using the menu, sets variables for the number of cards to use and whether or not to use a pair of black cards that will end the game if paired.  These variables are sent to the **BP_LogicMem** and the **BP_DeckMem** _Blueprints_ with
-the **BPI_Communication** _Blueprint Interface_. 
+The **Memory Game** start menu is set up in the BP_MenuMem _Blueprint_ using simple _StaticMeshComponents_.  _TextRenderComponents_ are used for the text and numbers, and a Timeline is used to move the menu out of the game window.  
+Using the menu, sets variables for the number of cards to use and whether or not to use a pair of black cards that will end the game if paired.  These variables are sent to the BP_LogicMem and the BP_DeckMem _Blueprints_ with
+the BPI_Communication _Blueprint Interface_. 
 
 [REGION:fullwidth]
 ![](CardGameMenu.png)
