#UE4doc #docs: ##Overview removed as according to the instructional guide.

diff --git a/Engine/Documentation/Source/Resources/SampleGames/PlatformerGame/PlatformerGame.INT.udn b/Engine/Documentation/Source/Resources/SampleGames/PlatformerGame/PlatformerGame.INT.udn
--- a/Engine/Documentation/Source/Resources/SampleGames/PlatformerGame/PlatformerGame.INT.udn
+++ b/Engine/Documentation/Source/Resources/SampleGames/PlatformerGame/PlatformerGame.INT.udn
@@ -52,8 +52,6 @@ Description:Overview of the Platformer Sample Game.
 
 [TOC(start:2)]
 
-## Overview
-
 The PlatformerGame sample is an example of a PC/mobile side-scroller game. It includes basic implementations of 
 different movement types along with a simple front end menu system.
 
