#UE4doc #docs: ##Overview removed as according to the instructional guide.

diff --git a/Engine/Documentation/Source/Resources/SampleGames/ShooterGame/ShooterGame.INT.udn b/Engine/Documentation/Source/Resources/SampleGames/ShooterGame/ShooterGame.INT.udn
--- a/Engine/Documentation/Source/Resources/SampleGames/ShooterGame/ShooterGame.INT.udn
+++ b/Engine/Documentation/Source/Resources/SampleGames/ShooterGame/ShooterGame.INT.udn
@@ -65,8 +65,6 @@ Related: Engine/Rendering/ParticleSystems
 
 [TOC(start:2)]
 
-## Overview
-
 The First Person Shooter sample is an example of a PC multiplayer first-person shooter. It includes basic implementations of 
 weapons and gametypes along with a simple front end menu system.
 
