UEDOC-1807

Updated GIF to Youtube video

diff --git a/Engine/Documentation/Source/Engine/Animation/Tools/VertexAnimationTool/VAT_TL_Meshes/VAT_TL_Meshes.INT.udn b/Engine/Documentation/Source/Engine/Animation/Tools/VertexAnimationTool/VAT_TL_Meshes/VAT_TL_Meshes.INT.udn
--- a/Engine/Documentation/Source/Engine/Animation/Tools/VertexAnimationTool/VAT_TL_Meshes/VAT_TL_Meshes.INT.udn
+++ b/Engine/Documentation/Source/Engine/Animation/Tools/VertexAnimationTool/VAT_TL_Meshes/VAT_TL_Meshes.INT.udn
@@ -24,7 +24,17 @@ In the following section we will cover how you can use the 3Ds Max Animation Tim
 
 1. Now enable **Auto Key** and then move the **Time Slider** to **10** and set the **Angle** on the **Bend Modifier** to **-50** and then move the **Time Slider** to **25** and set the Bend **Angle** to **50**. Finally move the **Time Slider** to **35** and then set the Bend **Angle** to **0**. Once all of that has been completed make sure to turn off **Auto Key** to ensure that you do not place any extra keys by accident. 
 	
-	![](VAT_Max_Bend_Setup.gif)(convert:false)
+	[OBJECT:EmbeddedVideo]
+	[PARAMLITERAL:width]
+	640
+	[/PARAMLITERAL]
+	[PARAMLITERAL:height]
+	360
+	[/PARAMLITERAL]
+	[PARAMLITERAL:videoid]
+	i_kBjC7ULis
+	[/PARAMLITERAL]
+	[/OBJECT]
 
 1. With the Teapot now bending, open up the Vertex Animation Tools, if not already open, and set the **Animation Options** to **Animated Meshes** and the **Anim End** to **35**.
 	
@@ -41,7 +51,17 @@ a **.EXR** Texture which contains the vertices position data and a **.BMP** Text
 1. With the Textures now created go back to 3Ds Max and select the Teapot mesh with "**_MorphExport**" added to the name and export it as a .FBX file to the same folder that the Textures where exported to.
 Name the mesh "**SM_Teapot_Bend_00**" and in the FBX options make sure that Smoothing Groups are enabled and leave all other setting with the default values.
 
-	![](VAT_3DsMax_Teapot_Export.gif)(convert:false)
+	[OBJECT:EmbeddedVideo]
+	[PARAMLITERAL:width]
+	640
+	[/PARAMLITERAL]
+	[PARAMLITERAL:height]
+	360
+	[/PARAMLITERAL]
+	[PARAMLITERAL:videoid]
+	swfWWVYfj20
+	[/PARAMLITERAL]
+	[/OBJECT]
 
 
 ## UE4 Content Setup & Use for Timeline Animated Meshes
@@ -59,7 +79,17 @@ To do select each Texture from its export location in Windows Explorer and drag
 When you do this, you will be prompted with a warning message about the Textures not being a power of 2 resulting in the Texture not working with Streaming or Mip Mapping.
 Press the **Yes** button to dismiss the prompt and continue with the import process.
 
-![](VAT_Texture_Import.gif)(convert:false)
+[OBJECT:EmbeddedVideo]
+[PARAMLITERAL:width]
+640
+[/PARAMLITERAL]
+[PARAMLITERAL:height]
+360
+[/PARAMLITERAL]
+[PARAMLITERAL:videoid]
+yqqHsxrPk9o
+[/PARAMLITERAL]
+[/OBJECT]
 
 Once each Texture has been imported, they then need to be setup in a very specific manner to ensure that they work correctly.
 The first Texture to setup is the .EXR Texture that contains the vertices position data.
@@ -125,7 +155,17 @@ If you are unfamiliar with how Materials work inside UE4 please refer to the [](
 
 1. Next open up the **MAT_Vertex_Animation** Material and locate the two Textures that were imported into the Content Browser in the previous step. Select both Textures and then Drag them from the Content Browser and Drop them into the Material Graph.  
 	
-	![](VAT_Materail_Adding_Textures.gif)(convert:false)
+	[OBJECT:EmbeddedVideo]
+	[PARAMLITERAL:width]
+	640
+	[/PARAMLITERAL]
+	[PARAMLITERAL:height]
+	360
+	[/PARAMLITERAL]
+	[PARAMLITERAL:videoid]
+	B9NfExH-KYk
+	[/PARAMLITERAL]
+	[/OBJECT]
 
 1. Then select the **Main Material Node** and in **Details Panel** search for **Num Customized UVs**. When you located the **Num Customized UVs** input, input a value of **4**. You should now see four new inputs on the Main Material Node with Customized UV and then a number starting at 0 going to 3.
 	
@@ -141,11 +181,31 @@ If you are unfamiliar with how Materials work inside UE4 please refer to the [](
 
 1. Now it is time to add the Material Function that powers the effect. To do this Right-Click with the Mouse inside the Material Graph and then from the menu that is displayed, type **MS** to show the Maxscript Material Functions. Then click on the **MS_VertexAnimationTools_MorphTargets** Material Function with the mouse to add it to the Material Graph. 
 
-	![](VAT_Add_MF.gif)(convert:false)
+	[OBJECT:EmbeddedVideo]
+	[PARAMLITERAL:width]
+	640
+	[/PARAMLITERAL]
+	[PARAMLITERAL:height]
+	360
+	[/PARAMLITERAL]
+	[PARAMLITERAL:videoid]
+	Nh8RkYoNe_M
+	[/PARAMLITERAL]
+	[/OBJECT]
 
 1. With all the required Textures and Material Function now added to the Material Graph, we can now begin to connect everything together but first, select both Texture Samplers and convert them to **Texture Objects** by Right-Clicking on one of the Textures and selecting the **Convert To Texture Object** option.
 	
-	![](VAT_Convert_To_TexOBJ.gif)(convert:false)
+	[OBJECT:EmbeddedVideo]
+	[PARAMLITERAL:width]
+	640
+	[/PARAMLITERAL]
+	[PARAMLITERAL:height]
+	360
+	[/PARAMLITERAL]
+	[PARAMLITERAL:videoid]
+	mwaJYXw38kQ
+	[/PARAMLITERAL]
+	[/OBJECT]
 
 1. Next connect the output from the Normal Texture Object to the **Morph Normal(T2d)** input. Then connect the .EXR Texture Object to the **Morph Texture (T2d)** input. Finally connect the outputs of the **MS_VertexAnimationTools_MorphTargets** Material Function to the related inputs on the Main Material node.
 
@@ -163,7 +223,17 @@ If you are unfamiliar with how Materials work inside UE4 please refer to the [](
 
 1. From the **Content Browser** select the Teapot mesh that was imported and drag it from the Content Browser into the level. Then once the Teapot mesh is in the level, select the Material that was created, **MAT_Vertex_Animation**, and apply it to the Teapot mesh by dragging it from the Content Browser on to the Teapot mesh that is placed in the level.
 
-	![](VAT_Final_Effect.gif)(convert:false)
+	[OBJECT:EmbeddedVideo]
+	[PARAMLITERAL:width]
+	640
+	[/PARAMLITERAL]
+	[PARAMLITERAL:height]
+	360
+	[/PARAMLITERAL]
+	[PARAMLITERAL:videoid]
+	TuTRafXORuU
+	[/PARAMLITERAL]
+	[/OBJECT]
 
 
 
