#UE4doc #docs: Fixed typo: Defnese -> Defense Removed extraneous bolding as according to the instructional guide. ##Overview removed as according to the instructional guide.

diff --git a/Engine/Documentation/Source/Resources/SampleGames/StrategyGame/StrategyGame.INT.udn b/Engine/Documentation/Source/Resources/SampleGames/StrategyGame/StrategyGame.INT.udn
--- a/Engine/Documentation/Source/Resources/SampleGames/StrategyGame/StrategyGame.INT.udn
+++ b/Engine/Documentation/Source/Resources/SampleGames/StrategyGame/StrategyGame.INT.udn
@@ -56,9 +56,6 @@ Related: Engine/Rendering/ParticleSystems
 
 [TOC(start:2 end:2)]
 
-
-## Overview
-
 The Tower Defense sample is an example of a RTS/Tower Defense game. 
 
 A complete list of the featured concepts:
@@ -87,7 +84,7 @@ Both enemy and friendly Pawns operate with the AI logic, moving towards the oppo
 or behavior of their friendly Pawns, they are able to buy new units to spawn.  
 
 _Blueprints_ are also used to add logic to Minion Pawns.  Both friendly and enemy Pawns can be equipped with shields; friendly Pawns get shields if the armory upgrade has been purchased for the brewery,
-and enemy Pawns get shields if they were spawned due to the **SpawnHeavyFunction** or **SpawnEndBossFunction** being called in the _Level Blueprint_. If a Pawn has a shield, projectiles from the auto-arbalest will be destroyed and do no damage. This logic is carried out using _Blueprint Interfaces_.  The
+and enemy Pawns get shields if they were spawned due to the SpawnHeavyFunction or SpawnEndBossFunction being called in the _Level Blueprint_. If a Pawn has a shield, projectiles from the auto-arbalest will be destroyed and do no damage. This logic is carried out using _Blueprint Interfaces_.  The
 **Minion** _Blueprint_ also contains a network to slow down enemy Pawns when they are hit by a charged-up fire turret.  
 
 ## Building Construction
@@ -101,24 +98,24 @@ Again, the code in Tower Defense simply creates the base building classes. All l
 
 ### Brewery
 
-The **Brewery** _Blueprint_ has the parent class `StrategyBuilding_Brewery`, and also includes an **AIDirector** component.  There are two breweries placed in the TowerDefenseMap, one for the enemy where enemy
-Pawns spawn, and one friendly brewery where an armory and a smithy upgrade can be built, and friendly Pawns can be spawned.  There is no graph logic present in the **Brewery** _Blueprint_, just **Defaults** set for the 
-building properties and a **Components** list including the **AIDirector**, a **TriggerBox**, and a _Static Mesh_. 
+The **Brewery** _Blueprint_ has the parent class `StrategyBuilding_Brewery`, and also includes an AIDirector component.  There are two breweries placed in the TowerDefenseMap, one for the enemy where enemy
+Pawns spawn, and one friendly brewery where an armory and a smithy upgrade can be built, and friendly Pawns can be spawned.  There is no graph logic present in the **Brewery** _Blueprint_, just Defaults set for the 
+building properties and a Components list including the AIDirector, a TriggerBox, and a _Static Mesh_. 
 
 #### Upgrades
 
 ![](upgraded_brewery.png)
 
-There are two upgrade slots attached to the friendly brewery.  These slots are _Blueprint Classes_ which are also derived from the `StrategyBuilding` class.  If an upgrade is selected from the **Brewery** menu,
+There are two upgrade slots attached to the friendly brewery.  These slots are _Blueprint Classes_ which are also derived from the `StrategyBuilding` class.  If an upgrade is selected from the Brewery menu,
 one slot will be replaced by that upgrade. Only one smithy upgrade and one armory upgrade can be purchased. 
 
-Once a smithy upgrade is purchased, the build will be started, firing the **OnBuildStarted** event in the **Wall_Smithy** _Blueprint_.  This _Blueprint_ also informs the system that the upgrade has been built,
+Once a smithy upgrade is purchased, the build will be started, firing the OnBuildStarted event in the **Wall_Smithy** _Blueprint_.  This _Blueprint_ also informs the system that the upgrade has been built,
 once the build completes.  After this point, any friendly Pawns that are spawned will be equipped with a shield attachment, a _Blueprint_ derived from the `StrategyAttachment` class.  The network which assigns the 
-"shield-attaching" behavior after the armory reports that it has been constructed is present in the **PlayerBaseUpgrades** collapsed graph in the TowerDefenseMap _Level Blueprint_. The `StrategyAttachment` 
-class simply contains a  `SkeletalMeshComponent`; the mesh and attachment point for the attachment are set in the **Defaults** for the **Attachment_Armorer** _Blueprint_.
+"shield-attaching" behavior after the armory reports that it has been constructed is present in the PlayerBaseUpgrades collapsed graph in the TowerDefenseMap _Level Blueprint_. The `StrategyAttachment` 
+class simply contains a  `SkeletalMeshComponent`; the mesh and attachment point for the attachment are set in the Defaults for the **Attachment_Armorer** _Blueprint_.
 
-The armory _Blueprint_ contains the same logic setup with **OnBuildStarted** and **OnBuildFinished** events.  After construction of the armory, any friendly Pawns spawned will be equipped with a hammer, also derived from the 
-`StrategyAttachment` class.  The network which assigns the "hammer-attaching" behavior after the smithy reports that it has been constructed is also present in the **PlayerBaseUpgrades** collapsed graph in the TowerDefenseMap _Level Blueprint_.
+The armory _Blueprint_ contains the same logic setup with OnBuildStarted and OnBuildFinished events.  After construction of the armory, any friendly Pawns spawned will be equipped with a hammer, also derived from the 
+`StrategyAttachment` class.  The network which assigns the "hammer-attaching" behavior after the smithy reports that it has been constructed is also present in the PlayerBaseUpgrades collapsed graph in the TowerDefenseMap _Level Blueprint_.
 
 ### Turrets
 
@@ -126,10 +123,10 @@ The armory _Blueprint_ contains the same logic setup with **OnBuildStarted** and
 
 ![](building_empty_slot.png)
 
-The empty slot is also a _Blueprint_ with the `StrategyBuilding` parent class, **Wall_EmptySlot**. No logic is present in the _Blueprint_ graphs. This is a _Blueprint Class_ with **Defaults** set for building 
-properties and Static Meshes and a trigger box set as the **Components**.
+The empty slot is also a _Blueprint_ with the `StrategyBuilding` parent class, Wall_EmptySlot. No logic is present in the _Blueprint_ graphs. This is a _Blueprint Class_ with Defaults set for building 
+properties and Static Meshes and a trigger box set as the Components.
 
-The possible turret upgrades are all set in the **Defaults** of the **Wall_EmptySlot** _Blueprint_, in the **Upgrades** section of the **Building** category.  
+The possible turret upgrades are all set in the Defaults of the Wall_EmptySlot _Blueprint_, in the Upgrades section of the Building category.  
 
 
 #### Arbalest
@@ -140,7 +137,7 @@ The **Wall_arbalest** _Blueprint_ contains the logic for the arbalest, the basic
 can manually fire the arbalest as well, by clicking on it and dragging in the direction they want the arrow to fire.  The longer the mouse drag is, the stronger the shot will be.  
 
 
-The arbalest projectile is stored in another _Blueprint_, **Projectile_arbalest**, derived from the _Blueprint_ **TestProjectile** which has `StrategyProjectile`.  The **Wall_arbalest** _Blueprint_ has a number of sub-networks, all contained within the **EventGraph**.  No _Blueprint_ logic is present in the **ConstructionScript**.   
+The arbalest projectile is stored in another _Blueprint_, Projectile_arbalest, derived from the _Blueprint_ TestProjectile which has `StrategyProjectile`.  The Wall_arbalest _Blueprint_ has a number of sub-networks, all contained within the EventGraph.  No _Blueprint_ logic is present in the ConstructionScript.   
 
 
 
@@ -152,8 +149,8 @@ The **Wall_arbalest_auto** _Blueprint_ contains the logic for the auto-arbalest.
 pass through.  The auto-arbalest arrows are not destroyed unless they hit a wall or an enemy with a shield.  It is possible to aim the auto-arbalest by clicking and dragging so that the auto-arbalest faces the desired direction;
 the auto-arbalest will continue to fire in the aimed direction while the mouse button is held down, and will return to its default firing position when the mouse button is released. 
 
-Like the arbalest, this turret shoots projectile arrows contained in a separate _Blueprint_.  The **Projectile_arbalest_auto** auto-arbalest arrows are not destroyed unless they hit an enemy Pawn with a shield or a wall, and this behavior
-is carried out with the aid of _Blueprint Interfaces_, **Interface_Auto_Arbalest** and **Interface_Auto_Projectile**. 
+Like the arbalest, this turret shoots projectile arrows contained in a separate _Blueprint_.  The Projectile_arbalest_auto auto-arbalest arrows are not destroyed unless they hit an enemy Pawn with a shield or a wall, and this behavior
+is carried out with the aid of _Blueprint Interfaces_, Interface_Auto_Arbalest and Interface_Auto_Projectile. 
 
 #### Flamethrower
 
@@ -165,7 +162,7 @@ there is a small cool down afterwards and then the regular fire attack continues
 
 ## Camera
 
-The camera in Tower Defnese has a fixed view angle and the ability to zoom in and out with the mouse scroll wheel.  The camera calculations are made inside the `CalcCamera` function in the `StrategyPlayerController` class, while
+The camera in Tower Defense has a fixed view angle and the ability to zoom in and out with the mouse scroll wheel.  The camera calculations are made inside the `CalcCamera` function in the `StrategyPlayerController` class, while
 constants like the camera minimum and maximum offset, camera angle, and camera speed can be set in `DefaultGame.ini`.
 
 A spectator Pawn is used to create a player without a visible Pawn.  
@@ -178,7 +175,7 @@ The in-game HUD is created with a mixture of Canvas drawing and Slate widgets.
 
 In the top-left corner, a game timer counts down the warm-up time for the game, using the function `GetGameTime` in the class `SStrategySlateHUDWidget`. After the game begins, this countdown disappears, 
 and the display for the number of lives left (1) is revealed.  The properties for the "lives left" display are set in the `DrawLives` function of the `AStrategyHUD` class; the initial number of lives 
-is set in the **PlayerBaseUpgrades** subgraph in the **TowerDefenseMap** _Level Blueprint_.  
+is set in the PlayerBaseUpgrades subgraph in the TowerDefenseMap _Level Blueprint_.  
 
 The current gold resources are displayed in the top-center of the screen (2).  Both the game timer and the resource display are defined using basic widgets in the `SStrategySlateHUDWidget`.  The same class
 is used to create all top level widgets, but not all of these widgets are displayed by default.
@@ -240,17 +237,17 @@ The _Level Blueprint_ has a modular structure for spawning each wave, as well as
 The waves are constructed using three _Blueprint Macros_: spawn fast, spawn normal, and spawn heavy.  Each of these sets up the desired unit parameters and attachments and then waits for the `SpawnMinions` function in the
 `StrategyAIDirector` to fire. The macro waits for the enemy brewery's `StrategyAIDirector` to report back that the wave has been spawned, and then execution is allowed to leave the sub-network.
 
-Each spawn macro takes two execution inputs, one for beginning the macro and one for opening the **Gate** after the **OnWaveSpawned** event fires, and an integer input for the number of Pawns to spawn.  
+Each spawn macro takes two execution inputs, one for beginning the macro and one for opening the **Gate** after the OnWaveSpawned event fires, and an integer input for the number of Pawns to spawn.  
 Functions from the class `StrategyAIDirector` are called, with inputs specific to each type of Pawn wave.  The three functions are `SetDefaultWeapon`, `SetDefaultArmor`, and `SetBuffModifier`.  
 `SetDefaultWeapon` and `SetDefaultArmor` take _Blueprints_ as inputs and assign those _Blueprints_ as the new default weapon or default armor for the spawn.  For example, all enemy Pawns spawned by 
-the **SpawnFastMacro** have the **Attachment_Smithy** hammer _Blueprint_ as their default weapon, and all enemy Pawns spawned by the **SpawnHeavyMacro** have the **Attachment_Armorer** shield _Blueprint_ as their default armor. 
+the SpawnFastMacro have the Attachment_Smithy hammer _Blueprint_ as their default weapon, and all enemy Pawns spawned by the SpawnHeavyMacro have the Attachment_Armorer shield _Blueprint_ as their default armor. 
 
 The last `StrategyAIDirector` function called by the spawning _Blueprint_ functions is `SetBuffModifier`, which has a number of data inputs including the attack abilities, health bonus, speed, and size
 of the Pawns.  These inputs are all exposed to the _Blueprint_, so it is straightforward for a level designer to create a new class of enemy Pawn to spawn.  Finally, each spawning _Blueprint_ function sets the `WaveSize`
 property of the enemy brewery's `StrategyAIDirector`.   
 
-There are five waves of enemies, each with different combinations of fast, normal, and heavy enemy Pawns.  At the beginning of a wave, the **Show Wave Title** node displays the wave number.  Then, the first enemy 
-spawn of the wave is called.  There are two types of delays after a spawn: a timer delay set by a **Delay** node, or a Pawn-based delay set by a **WaitForWaveMacro**. The **WaitForWaveMacro** macro continually checks to see how many enemy Pawns are alive
+There are five waves of enemies, each with different combinations of fast, normal, and heavy enemy Pawns.  At the beginning of a wave, the Show Wave Title node displays the wave number.  Then, the first enemy 
+spawn of the wave is called.  There are two types of delays after a spawn: a timer delay set by a Delay node, or a Pawn-based delay set by a WaitForWaveMacro. The WaitForWaveMacro macro continually checks to see how many enemy Pawns are alive
 and will not let execution leave the macro until either the delay time has expired or all of the enemy Pawns are dead. After all the spawns for a wave are complete and all of the enemy Pawns for that wave are dead (or two minutes have passed), 
 a _Remote Event_ node is used to call the _Custom Event_ for the next wave.
 
@@ -261,14 +258,14 @@ In the game, your base has 3 lives.  A life is subtracted if an enemy Pawn makes
 in losing the game. To win the game, you must defeat all five waves of enemy Pawns before losing all of your lives. The networks setting up both the win condition and the lose condition are in **TowerDefenseMap**'s _Level Blueprint_, and call functions set up
 in the class `StrategyGameMode`, which is derived from the `AGameMode` class.  The `StrategyGameMode` class also contains functions such as `InitGame` which initializes the game and is called before the Actors' PreInitializeComponents, `SetGamePaused`, and `SetGameDifficulty`.
 
-After the end boss spawned in Wave 5 has been killed, a _Remote Event_ node is used to call the **Winning** _Custom Event_.  This _Custom Event_, located in the **Win Condition** comment box, then fires and triggers execution of a sub-network **WaitForWin**,
-which checks that no other enemy Pawns are still alive.  If this is found to be _true_, the `Finish Game` function is called with the **WinningTeam** input set to "Player".  
+After the end boss spawned in Wave 5 has been killed, a _Remote Event_ node is used to call the Winning _Custom Event_.  This _Custom Event_, located in the **Win Condition** comment box, then fires and triggers execution of a sub-network WaitForWin,
+which checks that no other enemy Pawns are still alive.  If this is found to be _true_, the `Finish Game` function is called with the WinningTeam input set to "Player".  
 
-There are two nodes in the **Lose Condition** comment box which call the `Finish Game` function with the **WinningTeam** input set to "Enemy", causing the player to lose the game.  The first is triggered when all three lives have 
-been lost, due to enemy Pawns reaching the friendly brewery.  A **MultiGate** node is triggered each time an enemy Pawn reaches the friendly brewery.  The first and second output execution pins of the **MultiGate** node each connect to nodes which update the 
-NumberOfLives of the friendly brewery, decreasing the value by one each time. The last output execution pin sets the number of lives of the friendly brewery to zero, and then triggers the `Finish Game` function with the **WinningTeam** set to
-"Enemy".   Once the enemy boss has spawned, the **BossSpawned** _Custom Event_ closes the **Gate** node leading to the "3 lives version" **MultiGate**, and opens another **GateNode** leading to the second `Finish Game` function with the **WinningTeam** set to
-"Enemy". This creates an open network so that if the end boss reaches the friendly brewery, the **FinishGame** function will be activated and the player will lose the game.  
+There are two nodes in the **Lose Condition** comment box which call the `Finish Game` function with the WinningTeam input set to "Enemy", causing the player to lose the game.  The first is triggered when all three lives have 
+been lost, due to enemy Pawns reaching the friendly brewery.  A **MultiGate** node is triggered each time an enemy Pawn reaches the friendly brewery.  The first and second output execution pins of the MultiGate node each connect to nodes which update the 
+NumberOfLives of the friendly brewery, decreasing the value by one each time. The last output execution pin sets the number of lives of the friendly brewery to zero, and then triggers the `Finish Game` function with the WinningTeam set to
+"Enemy".   Once the enemy boss has spawned, the BossSpawned _Custom Event_ closes the Gate node leading to the "3 lives version" MultiGate, and opens another GateNode leading to the second `Finish Game` function with the WinningTeam set to
+"Enemy". This creates an open network so that if the end boss reaches the friendly brewery, the FinishGame function will be activated and the player will lose the game.  
  
 
   
@@ -279,7 +276,7 @@ The gold nodes are _Blueprints_ with the parent class `StrategyResourceNode`.  T
 a `BlueprintImplementableEvent` to report when the resource has been depleted, and the protected property `NumResources` to set the amount of the resource present in the node.  
 
 The _Blueprint_ for a gold node contains sub-networks to make the node appear and disappear on a timer.  The **ConstructionScript** in the _Blueprint_ sets the node to automatically be hidden when it is placed in the level.
-When the gold node appears, **AppearFX** particle effects are played along with an apparition noise. If the **OnDepleted** event fires because the node is successfully collected, the amount of gold present in the node
-is added to the player's total gold pool.  The **CollectFX** particle effects and the **CoinSound** is played, and then the node is hidden again.  If the player fails to collect the node in time by clicking on it,
-**FadeFX** particle effects and an appropriate sound is played.  
+When the gold node appears, AppearFX particle effects are played along with an apparition noise. If the OnDepleted event fires because the node is successfully collected, the amount of gold present in the node
+is added to the player's total gold pool.  The CollectFX particle effects and the CoinSound is played, and then the node is hidden again.  If the player fails to collect the node in time by clicking on it,
+FadeFX particle effects and an appropriate sound is played.  
 
