#UE4doc #docs: ##Overview removed as according to the instructional guide.

diff --git a/Engine/Documentation/Source/Resources/Showcases/BlueprintExamples/BlueprintExamples.INT.udn b/Engine/Documentation/Source/Resources/Showcases/BlueprintExamples/BlueprintExamples.INT.udn
--- a/Engine/Documentation/Source/Resources/Showcases/BlueprintExamples/BlueprintExamples.INT.udn
+++ b/Engine/Documentation/Source/Resources/Showcases/BlueprintExamples/BlueprintExamples.INT.udn
@@ -60,10 +60,6 @@ Related: Engine/Rendering/Materials
 
 [TOC(start:2 end:2)]
 
-
-
-## Overview
-
 **Blueprints** are Unreal Engine 4's approach to visual scripting. This means that tasks generally reserved for programming scripts are instead created via a graph of nodes and connections, rather than having to type out any actual code. This enables artists and other non-coder type users to create intricate and sophisticated gameplay systems that were previously only available to programmers.
 
 
