#UE4doc #docs: Fixed typo: Editiable -> editable. ##Overview removed as according to the instructional guide.

diff --git a/Engine/Documentation/Source/Resources/Showcases/BlueprintExamples/Butterfly/Butterfly.INT.udn b/Engine/Documentation/Source/Resources/Showcases/BlueprintExamples/Butterfly/Butterfly.INT.udn
--- a/Engine/Documentation/Source/Resources/Showcases/BlueprintExamples/Butterfly/Butterfly.INT.udn
+++ b/Engine/Documentation/Source/Resources/Showcases/BlueprintExamples/Butterfly/Butterfly.INT.udn
@@ -13,15 +13,13 @@ Related: Gameplay/AI
 ![](ButterflyHeader.png)
 [/REGION]
 
-## Overview
-
 [EXCERPT:Overview]
 The two mechanical butterflies in the center of the level are examples of how Blueprints can be used to create basic AI behavior. The butterflies will fly around a specified center point, randomly changing direction for a time, before trying to return to the center point to rest. When they encounter an obstacle, they smoothly fly towards it to land, then take off again after a delay. The longer a butterfly has flown, the longer it will rest.
 [/EXCERPT:Overview]
 
 The butterfly Blueprint Actor does not move around the level in and of itself; rather, its components do. The butterfly is made up of two wing meshes parented to a single body mesh. Each mesh is animated by timeline nodes, causing the wings to flap while the body bobs up and down. The Event Graph drives the butterfly components to move forward while steering to face random points within a radius, picking a new point to fly towards every few seconds. 
 
-A vector variable named **Attractor Location** is _Editiable_, meaning that it is exposed to the user. This variable also displays a 3D widget, allowing it to be moved around to determine the center point for the butterfly's flight range. After a certain amount of flying, the script steers the butterfly back toward this spot, which prevents it from straying too far away.
+A vector variable named **Attractor Location** is _Editable_, meaning that it is exposed to the user. This variable also displays a 3D widget, allowing it to be moved around to determine the center point for the butterfly's flight range. After a certain amount of flying, the script steers the butterfly back toward this spot, which prevents it from straying too far away.
 
 While the butterfly is in flight, it continuously draws a line trace ahead of itself to check for obstacles. When this line trace hits a surface, it provides the location and normal (surface angle) of the impact point. This data is then used to calculate the angle and position the butterfly should assume when it lands. A Timeline then plays and Linear Interpolate nodes smoothly transition the butterfly from its current location to its landed position.
 
