#UE4doc #docs: ##Overview removed as according to the instructional guide.

diff --git a/Engine/Documentation/Source/Resources/Showcases/BlueprintExamples/FogSheet/FogSheet.INT.udn b/Engine/Documentation/Source/Resources/Showcases/BlueprintExamples/FogSheet/FogSheet.INT.udn
--- a/Engine/Documentation/Source/Resources/Showcases/BlueprintExamples/FogSheet/FogSheet.INT.udn
+++ b/Engine/Documentation/Source/Resources/Showcases/BlueprintExamples/FogSheet/FogSheet.INT.udn
@@ -13,8 +13,6 @@ Related: Engine/Rendering/Materials/MaterialInstances
 ![](FogSheetLightBeams.png)
 [/REGION]
 
-## Overview
-
 [EXCERPT:Overview]
  The Fog Sheet and Light Ray Blueprints serve as an example of using Blueprints to create atmospheric effects for use in level design. Arranging translucent Static Meshes such as planes of fog and beams of light when decorating a level tends to be time-consuming, and so the purpose of these blueprints is to make the process easier while giving the user more control over the appearance of these meshes.
 [/EXCERPT:Overview]
