#UE4doc #docs: Removed extraneous bolding as according to the instructional guide. ##Overview removed as according to the instructional guide.

diff --git a/Engine/Documentation/Source/Resources/Showcases/BlueprintExamples/Foliage/Foliage.INT.udn b/Engine/Documentation/Source/Resources/Showcases/BlueprintExamples/Foliage/Foliage.INT.udn
--- a/Engine/Documentation/Source/Resources/Showcases/BlueprintExamples/Foliage/Foliage.INT.udn
+++ b/Engine/Documentation/Source/Resources/Showcases/BlueprintExamples/Foliage/Foliage.INT.udn
@@ -12,8 +12,6 @@ Related: Engine/Blueprints/UserGuide/UserConstructionScript
 ![](Foliage.png)
 [/REGION]
 
-##Overview
-
 [EXCERPT:Overview]
 The Foliage Blueprint shows how to use a Blueprint to create randomly-generated foliage (or potentially other structures) that are automatically placed such that they conform to a surface. The user has control over the extent of the effect and how many meshes the system will try to spawn. As the Blueprint is moved, the meshes update to re-conform to the new environment.
 [/EXCERPT:Overview]
@@ -23,7 +21,7 @@ At a glance, the system takes a set radius then selects randomly-generated point
 
 ## Blueprint Setup
 
-The **Components List**, **Construction Script**, and **Event Graph** for this Blueprint perform the following tasks:
+The Components List, Construction Script, and Event Graph for this Blueprint perform the following tasks:
 
 
 |Blueprint Aspect |Purpose |
