#UE4doc #docs: Removed extraneous bolding as according to the instructional guide.

diff --git a/Engine/Documentation/Source/Resources/Showcases/BlueprintExamples/SecurityCameraAndAlarm/SecurityCameraAndAlarm.INT.udn b/Engine/Documentation/Source/Resources/Showcases/BlueprintExamples/SecurityCameraAndAlarm/SecurityCameraAndAlarm.INT.udn
--- a/Engine/Documentation/Source/Resources/Showcases/BlueprintExamples/SecurityCameraAndAlarm/SecurityCameraAndAlarm.INT.udn
+++ b/Engine/Documentation/Source/Resources/Showcases/BlueprintExamples/SecurityCameraAndAlarm/SecurityCameraAndAlarm.INT.udn
@@ -28,7 +28,7 @@ The catch? With many cameras running at the same time, it is possible to run int
 
 ### Blueprint Setup
 
-The **Components List**, **Construction Script**, and **Event Graph** for this Blueprint perform the following tasks:
+The Components List, Construction Script, and Event Graph for this Blueprint perform the following tasks:
 
 
 |Blueprint Aspect |Purpose |
@@ -72,7 +72,7 @@ The button Blueprint also contains a Box Component which is used as a trigger vo
 
 ### Blueprint Setup
 
-The **Components List**, **Construction Script**, and **Event Graph** for this Blueprint perform the following tasks:
+The Components List, Construction Script, and Event Graph for this Blueprint perform the following tasks:
 
 
 |Blueprint Aspect |Purpose |
@@ -106,19 +106,19 @@ The button has the following exposed properties:
 
 
 [EXCERPT:ScreenCL]
-The **Components List** for the screen is about as simple as you can get, containing only a single Static Mesh for the screen itself.
+The Components List for the screen is about as simple as you can get, containing only a single Static Mesh for the screen itself.
 [/EXCERPT:ScreenCL]
 
 
 [EXCERPT:ScreenCS]
-The **Construction Script** performs the following tasks:
+The Construction Script performs the following tasks:
 * Creates and applies a [Material Instance Dynamic (MID)](Engine/Rendering/Materials/MaterialInstances) for the screen mesh.
 [/EXCERPT:ScreenCS]
 
 
 
 [EXCERPT:ScreenEG]
-The **Event Graph** performs the following tasks:
+The Event Graph performs the following tasks:
 When we receive a "Camera Swap" signal from the alarm Interface:
 * Turn off the previous capture Actor.
 * Play a Timeline that drives noise on the screen material.
@@ -129,20 +129,20 @@ When we receive a "Camera Swap" signal from the alarm Interface:
 
 
 [EXCERPT:ButtonCL]
-The **Components List** for the Button is set up as follows:
+The Components List for the Button is set up as follows:
 ![](ButtonComponents.png)
 
 [/EXCERPT:ButtonCL]
 
 
 [EXCERPT:ButtonCS]
-The **Construction Script** has no function.
+The Construction Script has no function.
 [/EXCERPT:ButtonCS]
 
 
 
 [EXCERPT:ButtonEG]
-The **Event Graph** performs the following tasks:
+The Event Graph performs the following tasks:
 * Check to see when the player enters the Box Component.
 * Once the player is within the Component, create some help text.
 * Keep that help text rotated toward the player.
