#UEDoc #docs - Bold lettering

diff --git a/Engine/Documentation/Source/Programming/UnrealArchitecture/Reference/Properties/Specifiers/Interp/PropertySpecifier_Interp.INT.udn b/Engine/Documentation/Source/Programming/UnrealArchitecture/Reference/Properties/Specifiers/Interp/PropertySpecifier_Interp.INT.udn
--- a/Engine/Documentation/Source/Programming/UnrealArchitecture/Reference/Properties/Specifiers/Interp/PropertySpecifier_Interp.INT.udn
+++ b/Engine/Documentation/Source/Programming/UnrealArchitecture/Reference/Properties/Specifiers/Interp/PropertySpecifier_Interp.INT.udn
@@ -4,4 +4,4 @@ Description:
 Availability:Public
 Version: 4.5
 
-Indicates that the value can be driven over time by a track in Matinee. 
+Indicates that the value can be driven over time by a track in **Matinee**. 
