#UE4doc #docs: Removed extraneous bolding as according to the instructional guide. ##Overview removed as according to the instructional guide.

diff --git a/Engine/Documentation/Source/Resources/Showcases/BlueprintExamples/SecurityDoor/SecurityDoor.INT.udn b/Engine/Documentation/Source/Resources/Showcases/BlueprintExamples/SecurityDoor/SecurityDoor.INT.udn
--- a/Engine/Documentation/Source/Resources/Showcases/BlueprintExamples/SecurityDoor/SecurityDoor.INT.udn
+++ b/Engine/Documentation/Source/Resources/Showcases/BlueprintExamples/SecurityDoor/SecurityDoor.INT.udn
@@ -12,8 +12,6 @@ Related: Gameplay/Framework/Pawn
 
 ![](SecurityCameraAndDoorHeader.png)(w:600)
 
-## Overview
-
 [EXCERPT:Overview]
 The security camera Blueprints seen throughout the level are designed to interact with the player when in Play mode. During play, they will pan back and forth in a search pattern. When the green search cone overlaps a pawn, the light turns red and the camera begins to track player movement. It continues to track the player until the player moves out of range. Then, after a few seconds, it returns to searching.
 [/EXCERPT:Overview]
@@ -37,7 +35,7 @@ A Timeline can contain any number of animated tracks, whose data can then be sen
 
 ### Blueprint Setup
 
-The **Components List**, **Construction Script**, and **Event Graph** for this Blueprint perform the following tasks:
+The Components List, Construction Script, and Event Graph for this Blueprint perform the following tasks:
 
 
 |Blueprint Aspect |Purpose |
@@ -91,7 +89,7 @@ The alarm is turned on and off by receiving messages from the Security Alarm [Bl
 
 ### Blueprint Setup
 
-The **Components List**, **Construction Script**, and **Event Graph** for this Blueprint perform the following tasks:
+The Components List, Construction Script, and Event Graph for this Blueprint perform the following tasks:
 
 
 |Blueprint Aspect |Purpose |
@@ -129,7 +127,7 @@ Also interesting about the door setup is that the door Timelines do not animate
 
 ### Blueprint Setup
 
-The **Components List**, **Construction Script**, and **Event Graph** for this Blueprint perform the following tasks:
+The Components List, Construction Script, and Event Graph for this Blueprint perform the following tasks:
 
 
 |Blueprint Aspect |Purpose |
@@ -156,14 +154,14 @@ The door has the following exposed properties:
 <!-----------------------------------
 
 [EXCERPT:CameraCL]
-The **Components List** for the camera contains the following components:
+The Components List for the camera contains the following components:
 [REGION:raw]![](CameraCS.png)[/REGION]
 Most of the elements are self-explanatory. The SceneCapture Component works just like a camera, except that instead of outputting what it sees to the screen, it instead populates a texture, which can then be used in a Material.
 [/EXCERPT:CameraCL]
 
 
 [EXCERPT:CameraCS]
-The **Construction Script** performs the following tasks:
+The Construction Script performs the following tasks:
 * Set up the Scene Capture to record to a Render Target Texture.
 * Use incoming values from Editable variables to set up the Spot Light for the camera.
 * Set the camera's default rotation angle.
@@ -173,7 +171,7 @@ The **Construction Script** performs the following tasks:
 
 
 [EXCERPT:CameraEG]
-The **Event Graph** performs the following tasks when a Pawn enters the view cone:
+The Event Graph performs the following tasks when a Pawn enters the view cone:
 * Engage the alarm.
 * Change the color of the view cone and the light.
 * Begin tracking the player's movement.
@@ -188,14 +186,14 @@ When the Pawn exits the view cone:
 
 
 [EXCERPT:AlarmCL]
-The **Components List** for the camera contains the following components:
+The Components List for the camera contains the following components:
 [REGION:raw]![](AlarmCS.png)[/REGION]
 The **Component List** for the alarm is basically a base mesh and 2 lights: one light for the alarm (Spot Light) and another for the neutral state (OffLight).
 [/EXCERPT:AlarmCL]
 
 
 [EXCERPT:AlarmCS]
-The **Construction Script** performs the following tasks:
+The Construction Script performs the following tasks:
 * Sets the light colors from the exposed user inputs.
 * Creates and applies a [Material Instance Dynamic (MID)](Engine/Rendering/Materials\MaterialInstances) for the alarm mesh.
 * If the user is testing the alarm, it swaps out the lights and material settings to show what the alarm looks like when it is on.
@@ -204,7 +202,7 @@ The **Construction Script** performs the following tasks:
 
 
 [EXCERPT:AlarmEG]
-The **Event Graph** performs the following tasks:
+The Event Graph performs the following tasks:
 When we receive a signal from the alarm Interface:
 * Parse whether the signal was to turn the alarm on or off.
 * If turning the alarm on, turn on the Spot Light, turn off the Off Light, and change material colors. Then open a Gate to start spinning the Spot Light.
@@ -216,14 +214,14 @@ When we receive a signal from the alarm Interface:
 
 
 [EXCERPT:DoorCL]
-The **Components List** for the Camera contains the following components:
+The Components List for the Camera contains the following components:
 [REGION:raw]![](DoorCL.png)[/REGION]
 The Component list for the door essentially contains a trigger volume and 2 door panels.
 [/EXCERPT:DoorCL]
 
 
 [EXCERPT:DoorCS]
-The **Construction Script** performs the following tasks:
+The Construction Script performs the following tasks:
 * Sets up the position and size for the trigger box volume.
 * Sets the location and scale for each door.
 [/EXCERPT:DoorCS]
@@ -231,7 +229,7 @@ The **Construction Script** performs the following tasks:
 
 
 [EXCERPT:DoorEG]
-The **Event Graph** performs the following tasks when we receive a signal from the alarm Interface:
+The Event Graph performs the following tasks when we receive a signal from the alarm Interface:
 * Check to see if the signal is turning the alarm on or off.
 * If the alarm is switching on, we check the current state of the doors and then, if they are not open already, we open them.
 * If the alarm is switching off, we check the current state of the doors and then, if they are not already closed, we close them.
