#UE4doc #docs: Removed extraneous bolding as according to the instructional guide. ##Overview removed as according to the instructional guide. Altered images to match the instructional guide: green -> yellow.

diff --git a/Engine/Documentation/Source/Resources/Showcases/BlueprintMultiplayer/BlueprintMultiplayer.INT.udn b/Engine/Documentation/Source/Resources/Showcases/BlueprintMultiplayer/BlueprintMultiplayer.INT.udn
--- a/Engine/Documentation/Source/Resources/Showcases/BlueprintMultiplayer/BlueprintMultiplayer.INT.udn
+++ b/Engine/Documentation/Source/Resources/Showcases/BlueprintMultiplayer/BlueprintMultiplayer.INT.udn
@@ -56,8 +56,6 @@ Related: Resources/ContentExamples/Networking
 
 [TOC(start:2 end:3)]
 
-## Overview
-
 The **Simple Blueprint Multiplayer** showcase is a game built entirely in **Blueprints** and **UMG UI** and is an example of how you might use Blueprint's **Session Nodes** to facilitate multiplayer in your game. It has a Main Menu, a Server List, and a simple map with a Scoreboard on the HUD. Clicking **Play** at the Main Menu will host a session and load the gameplay map. Other players can then click **Find games** to see a list of all hosted games on a new menu screen, where clicking on one of the found games will attempt to join it. If there are any errors, a dialog is displayed and the user is returned to the Main Menu.
 
 This game is also an example of how to manage state using a **GameInstance** Blueprint. The GameInstance is a convenient place to do this, since it persists between map loads, and is also the object that receives error events. Most of the session calls and menu transitions in this sample game are handled in the GameInstance.
@@ -179,15 +177,15 @@ Each asset can be found in their respective folders in the **Content Browser**:
 | Name | Description |
 |---|---|
 | **Gunslinger_BP** | This is the playable character in-game. |
-| **GunslingerTTP** | Skeletal Mesh used with the **Gunslinger_BP**. |
-| **Gunslinger_AnimBP** | Used to drive animations for the **Gunslinger_BP**. |
+| **GunslingerTTP** | Skeletal Mesh used with the Gunslinger_BP. |
+| **Gunslinger_AnimBP** | Used to drive animations for the Gunslinger_BP. |
 
 **Game/Character**
 
 | Name | Description |
 |---|---|
-| **Character_Parent** | (Parent) Material used for the **Gunslinger_BP**. |
-| **Character_Instance** | (Instanced) Material used with **Gunslinger_BP**. |
+| **Character_Parent** | (Parent) Material used for the Gunslinger_BP. |
+| **Character_Instance** | (Instanced) Material used with Gunslinger_BP. |
 
 **Game/Fonts**
 
@@ -277,11 +275,11 @@ The first thing that **ShowMainMenuEvent** does is run the [IsCurrentState Macro
 [REGION:note]
 ##### TransitionToState Macro
 
-Below, on the **TransitionToState**, the first thing it does is run the [IsCurrentState Macro](#iscurrentstatemacro) (1) to check the **Current State** versus the **Desired State** and if they are equal to each other an error is printed to the screen (2). If they are not equal to each other, the **Switch on State** node is executed (3) which gets the **Current State** and either runs a Function (if applicable) called **Hide Widget** (which hides the UI Widget currently being displayed) or runs a custom event called **Destroy Session Caller** (which destroys the session for the player who calls it). 
+Below, on the TransitionToState, the first thing it does is run the [IsCurrentState Macro](#iscurrentstatemacro) (1) to check the **Current State** versus the **Desired State** and if they are equal to each other an error is printed to the screen (2). If they are not equal to each other, the **Switch on State** node is executed (3) which gets the **Current State** and either runs a Function (if applicable) called **Hide Widget** (which hides the UI Widget currently being displayed) or runs a custom event called **Destroy Session Caller** (which destroys the session for the player who calls it). 
 
 ![](Startup6.png)
 
-After the **Switch on State** is complete, the **Current State** is updated to reflect the specified **Desired State**. 
+After the Switch on State is complete, the Current State is updated to reflect the specified Desired State. 
 [/REGION]
 
 After running the **TransitionToState** macro, we can continue with the **ShowMainMenuEvent** custom event. 
@@ -292,7 +290,7 @@ Above, the **IsValid** node (1) is executed which checks if the variable **Main
 
 ![](Startup8.png)
 
-Next we will look at what happens when you press **Play** to host a game. 
+Next we will look at what happens when you press Play to host a game. 
 
 ### Hosting a Game
 
@@ -327,7 +325,7 @@ On the **Create Session** node, the number of **Public Connections** (or players
 [REGION:note]
 ##### DisplayErrorDialog Macro
 
-The **DisplayErrorDialog** macro when called first executes the (1) [TransitionToState Macro](#transitiontostatemacro) to transition to the **Error Dialog** state. After the new state has been applied, a custom event called **Destroy Session Caller** is executed (2) which destroys the session for the current player (this can be found on the Event Graph). Following the session being destroyed, the engine macro **IsValid** checks to see if the variable called **Error Dialog** (which points to the **Error Dialog** Widget Blueprint) is indeed valid. 
+The DisplayErrorDialog macro when called first executes the (1) [TransitionToState Macro](#transitiontostatemacro) to transition to the **Error Dialog** state. After the new state has been applied, a custom event called **Destroy Session Caller** is executed (2) which destroys the session for the current player (this can be found on the Event Graph). Following the session being destroyed, the engine macro **IsValid** checks to see if the variable called Error Dialog (which points to the Error Dialog Widget Blueprint) is indeed valid. 
 
 ![](Play7.png)
 
@@ -335,7 +333,7 @@ Upon first run, the **Error Dialog** is not valid, if it were it would skip (1)
 
 ![](Play8.png)
 
-Upon first run however, (1) above which uses the **Create Widget** node to create the **ErrorDialog** Widget Blueprint is executed. In (2), the **ErrorDialog** Widget Blueprint is promoted to a variable called **Error Dialog** (which can now be used the next time this macro is performed). In (3), the Function **Set Message** from inside the **ErrorDialog** Widget Blueprint is called which gets the message specified on the macro node and sets it as the message to display in the dialog box (yellow box below) while the **Add to Viewport** (4) displays the **Error Dialog** Widget Blueprint on screen and (5) uses the **Set Input Mode UIOnly** sets the input mode to only UI. 
+Upon first run however, (1) above which uses the **Create Widget** node to create the ErrorDialog Widget Blueprint is executed. In (2), the ErrorDialog Widget Blueprint is promoted to a variable called Error Dialog (which can now be used the next time this macro is performed). In (3), the Function **Set Message** from inside the ErrorDialog Widget Blueprint is called which gets the message specified on the macro node and sets it as the message to display in the dialog box (yellow box below) while the **Add to Viewport** (4) displays the Error Dialog Widget Blueprint on screen and (5) uses the **Set Input Mode UIOnly** sets the input mode to only UI. 
 
 ![](Play9.png)
 
@@ -355,26 +353,26 @@ The **OnClicked** Event here is tied to the **Button** specified as the "Find ga
 
 The **ShowServerListEvent** starts off by running the [TransitionToState Macro](#transitiontostatemacro) (1) which sets the game state to **Sever List**. From there, the **IsValid** (2) check occurs against the **Server List** Widget Blueprint variable, where if valid uses the **Add to Viewport** (5) node to display the Server List menu onscreen. Upon first run however, the Server List screen needs to be created using the **Create Widget** (3) node, then promoted to a variable (4) before displaying it onscreen. Once the screen has been set to displayed, the **Set Input Mode UIOnly** node sets the input method to only UI. 
 
-During (5) above and before (6) and actually displaying the Server List Menu, additional script is executed inside the **ServerList** Widget Blueprint which is used to populate the Server List menu with any active games. As seen below inside the **Server List** Widget Blueprint, the **Event Construct** node is used and fired off when the Widget Blueprint is called to be constructed (created) and immediately executes a created macro called **RefreshListMacro**. 
+During (5) above and before (6) and actually displaying the Server List Menu, additional script is executed inside the ServerList Widget Blueprint which is used to populate the Server List menu with any active games. As seen below inside the Server List Widget Blueprint, the **Event Construct** node is used and fired off when the Widget Blueprint is called to be constructed (created) and immediately executes a created macro called **RefreshListMacro**. 
 
 ![](Findgames3.png)
 
 [REGION:note]
 ##### RefreshListMacro
 
-The **RefreshListMacro** performs quite a few tasks, the first section of which is illustrated below. 
+The RefreshListMacro performs quite a few tasks, the first section of which is illustrated below. 
 
 ![](Findgames4.png)
 
-Above when the macro is called, the first thing it does is set a **Boolean** variable called **Refresh Button Enabled** to **False** (1) which disables the **Refresh** Button and prevents the player from being able to press it. Next, a variable called **Status Text** is set to **Searching...** (2) and its **Visibility** is set to **Visible** (3) which allows the Widget holding the Status Text to be updated and displayed. 
+Above when the macro is called, the first thing it does is set a Boolean variable called **Refresh Button Enabled** to **False** (1) which disables the **Refresh** Button and prevents the player from being able to press it. Next, a variable called **Status Text** is set to **Searching...** (2) and its **Visibility** is set to **Visible** (3) which allows the Widget holding the Status Text to be updated and displayed. 
 
-In (4) an array variable called **Found Session Widgets**, which is a collection of **ServerRow** Widgets used to populate the Server List, each one that exists is removed using the **Remove Child** (5) node. This is effectively removing any servers that are to be listed when the screen is displayed. In (6), the **Found Session Widgets** variable is also cleared removing any previously stored sessions, so that everything is wiped clean before performing a new search. 
+In (4) an array variable called **Found Session Widgets**, which is a collection of ServerRow Widgets used to populate the Server List, each one that exists is removed using the **Remove Child** (5) node. This is effectively removing any servers that are to be listed when the screen is displayed. In (6), the **Found Session Widgets** variable is also cleared removing any previously stored sessions, so that everything is wiped clean before performing a new search. 
 
 Next, the macro gets the game instance and uses the **Cast To MyGameInstance** node (1 below) in order to access whether **Enable LAN** is true or not (which is used when searching for sessions). 
 
 ![](Findgames5.png)
 
-The **Find Sessions** node (2) is then used to discover any sessions that exist (which are output to the **Results** value). Also on this node, you can set the **Max Results** (number of search results to return). The results from the **Find Sessions** query are assigned and stored in an array variable called **Found Sessions** (3). A **Branch** is then performed (4) to check the **Results**, and if equal to 0, the **True** pin is executed (if not, the **False** is executed). If for any reason there is a failure when using the **Find Sessions** node, the **OnFailure** event is called and changes the **Status Text** that is displayed to **Search failed** (5), informing the user that the search for sessions did not complete. 
+The **Find Sessions** node (2) is then used to discover any sessions that exist (which are output to the Results value). Also on this node, you can set the **Max Results** (number of search results to return). The results from the **Find Sessions** query are assigned and stored in an array variable called **Found Sessions** (3). A **Branch** is then performed (4) to check the Results, and if equal to 0, the **True** pin is executed (if not, the **False** is executed). If for any reason there is a failure when using the **Find Sessions** node, the **OnFailure** event is called and changes the **Status Text** that is displayed to **Search failed** (5), informing the user that the search for sessions did not complete. 
 
 Off of the **Branch** (4 in the previous image), if **True** (which means no sessions were found) the **Status Text** is set to **No sessions found** (7 below) and the **Refresh Button Enabled** variable is set to **True** (8) so that the player can press the **Refresh** button to attempt to search for sessions again. 
 
@@ -382,9 +380,9 @@ Off of the **Branch** (4 in the previous image), if **True** (which means no ses
 
 If sessions are found, the **Status Text Visibility** is set to **Collapsed** (1) so that it is hidden from the player. A **ForEachLoop** node (2) is then used to take each of the **Results** found and, using the **Create ServerRow** node (3), create a Server Row Widget for each of them. These will be the individual rows of sessions found in the server list. 
 
-In (4), the **Results** are also added to the **Found Session Widgets** array before calling the **Set Search Result** Function inside the **ServerRow** Widget Blueprint (which used to generate the Server Name, Player Count, and Ping to be displayed). After setting the content to be displayed in the **ServerRow** Widget Blueprint with the **Set Search Result** node, the **Add Child** node (6) is used to add each of the **ServerRow** Widget Blueprints created to a **Scrolling Servers** scroll box that is displayed in the Server List menu. 
+In (4), the **Results** are also added to the **Found Session Widgets** array before calling the Set Search Result Function inside the ServerRow Widget Blueprint (which used to generate the Server Name, Player Count, and Ping to be displayed). After setting the content to be displayed in the ServerRow Widget Blueprint with the Set Search Result node, the **Add Child** node (6) is used to add each of the ServerRow Widget Blueprints created to a **Scrolling Servers** scroll box that is displayed in the Server List menu. 
 
-After each of the sessions found has had a **ServerRow** widget created, populated and added to the **Scrolling Servers** widget, the **ForEachLoop** is completed and the **Referesh Button Enabled** variable is set to **True** (8). 
+After each of the sessions found has had a ServerRow widget created, populated and added to the **Scrolling Servers** widget, the **ForEachLoop** is completed and the **Referesh Button Enabled** variable is set to **True** (8). 
 [/REGION]
 
 After having run through the **RefreshListMacro** to search for sessions, any sessions found will be displayed in the **Server List** menu. 
@@ -403,7 +401,7 @@ The **JoinFromServerListEvent** inside the **MyGameInstance** Blueprint (picture
 
 ![](Findgames9.png)
 
-While the loading screen is being displayed, the **Join Session** node is used to attempt to join the **Search Result** provided from the **ServerRow** Widget Blueprint the player clicked on to attempt to join. If successful, the player will join the server they selected and spawn into that game. If there are any errors, the [DisplayErrorDialog](#displayerrordialogmacro) is run and the **Error Message** of **Failed to join session** is displayed. 
+While the loading screen is being displayed, the Join Session node is used to attempt to join the Search Result provided from the ServerRow Widget Blueprint the player clicked on to attempt to join. If successful, the player will join the server they selected and spawn into that game. If there are any errors, the [DisplayErrorDialog](#displayerrordialogmacro) is run and the **Error Message** of **Failed to join session** is displayed. 
 
 ### Gameplay
 
@@ -431,7 +429,7 @@ Now that we have identified where the script should occur, the **Setup Ingame UI
 
 ![](Gameplay4.png)
 
-This Function starts with a **Branch** (1) to check if the controller that is being targeted is a Local Player Controller and if it is, the **Create Widget** (2) is used to create the **HUD** Widget Blueprint. The **HUD** Widget is then promoted to a variable called **HUD Widget** (3) and the **Add to Viewport** (4) node adds the widget to the player's viewport. In (5), another **Create Widget** node used to create the **InGame Menu** Widget Blueprint which is promoted to the variable **InGameMenuWidget** (6) and can be used to call display the In Game Menu when the player presses the button needed to summon it. 
+This Function starts with a **Branch** (1) to check if the controller that is being targeted is a Local Player Controller and if it is, the **Create Widget** (2) is used to create the HUD Widget Blueprint. The HUD Widget is then promoted to a variable called **HUD Widget** (3) and the **Add to Viewport** (4) node adds the widget to the player's viewport. In (5), another **Create Widget** node used to create the **InGame Menu** Widget Blueprint which is promoted to the variable **InGameMenuWidget** (6) and can be used to call display the In Game Menu when the player presses the button needed to summon it. 
 
 [REGION:note]
 At this point, the "login" process is complete and the setting up of the character itself occurs through the **Character Animation Blueprint**. Each of the character's movements states are defined and animations are set, which will not be covered here but if you would like to read more on **Animation Blueprints**, please see the [](Engine/Animation/AnimBlueprints) documentation.
