#UE4doc #docs: Removed duplicated words: the the -> the. ##Overview removed as according to the instructional guide.

diff --git a/Engine/Documentation/Source/Resources/Showcases/Effects/Effects.INT.udn b/Engine/Documentation/Source/Resources/Showcases/Effects/Effects.INT.udn
--- a/Engine/Documentation/Source/Resources/Showcases/Effects/Effects.INT.udn
+++ b/Engine/Documentation/Source/Resources/Showcases/Effects/Effects.INT.udn
@@ -59,15 +59,11 @@ Related: Engine/Rendering/ParticleSystems/Reference
 
 [TOC(start:2 end:2)]
 
-
-
-## Overview
-
 The Particle Effects Example shows off a variety of particle systems that demonstrate various environmental and elemental effects. The goal of these examples is to create a variety of believably realistic visual effects, focusing on the use of [GPU Sprite](Engine/Rendering/ParticleSystems\Reference\TypeData\GPUSprites) particles and their features.
 
 ## Scene Breakdown
 
-The Particle Effects Example showcases a cavern level, high up in a snowy arctic mountain range. Within the cave, water cascades across moss-covered rock formations, trickling into rich, flowing pools on the cavern floor. Mysterious and seemingly magical flame sources provide an eerie light. Outside, a harsh blizzard rips through the the rocky crags. 
+The Particle Effects Example showcases a cavern level, high up in a snowy arctic mountain range. Within the cave, water cascades across moss-covered rock formations, trickling into rich, flowing pools on the cavern floor. Mysterious and seemingly magical flame sources provide an eerie light. Outside, a harsh blizzard rips through the rocky crags. 
 
 
 ## Particle System Examples
