#UEDoc #docs - Bold lettering adjustments.

diff --git a/Engine/Documentation/Source/Programming/UnrealArchitecture/CoreArchitecture.INT.udn b/Engine/Documentation/Source/Programming/UnrealArchitecture/CoreArchitecture.INT.udn
--- a/Engine/Documentation/Source/Programming/UnrealArchitecture/CoreArchitecture.INT.udn
+++ b/Engine/Documentation/Source/Programming/UnrealArchitecture/CoreArchitecture.INT.udn
@@ -3,19 +3,18 @@ Title: Gameplay Programming
 Crumbs:%ROOT%, Programming
 Description:Explanations of the basic gameplay elements, Actors and Objects.
 
-## Overview
 
 The goal of this document is to provide a brief tour of the engine's core features and introduce some of the core architectural concepts to help you get started with the codebase.
 
 ## UObjects and Actors
 
-[](Programming\UnrealArchitecture\Actors) are instances of classes that derive from the `AActor` class; the base class of all gameplay objects that can be placed into the world. 
-[](Programming\UnrealArchitecture\Objects) are instances of classes that inherit from the `UObject` class; the base class of all objects in Unreal Engine, including Actors. So, in reality, all 
-instances in Unreal Engine are **Objects**; however, the term **Actors** is commonly used to refer to instances of classes that derive from `AActor` in their hierarchy, while the term **Objects** 
+**[](Programming\UnrealArchitecture\Actors)** are instances of classes that derive from the `AActor` class; the base class of all gameplay objects that can be placed into the world. 
+**[](Programming\UnrealArchitecture\Objects)** are instances of classes that inherit from the `UObject` class; the base class of all objects in Unreal Engine, including Actors. So, in reality, all 
+instances in Unreal Engine are **Objects**; however, the term **Actors** is commonly used to refer to instances of classes that derive from `AActor` in their hierarchy, while the term Objects 
 is used to refer to instances of classes that do not inherit from the `AActor` class. The majority of the classes you create will inherit from `AActor` at some point in their hierarchy.
 
-In general, **Actors** can be thought of as whole items or entities, while **Objects** are more specialized parts. **Actors** often make use of [](Programming\UnrealArchitecture\Actors\Components), which are specialized **Objects**, 
-to define certain aspects of their functionality or hold values for a collection of properties. Take a car as an example. The car as a whole is an **Actor**, whereas the parts of the car, 
+In general, Actors can be thought of as whole items or entities, while Objects are more specialized parts. Actors often make use of **[](Programming\UnrealArchitecture\Actors\Components)**, which are specialized Objects, 
+to define certain aspects of their functionality or hold values for a collection of properties. Take a car as an example. The car as a whole is an Actor, whereas the parts of the car, 
 like the wheels and doors, would all be Components of that Actor. 
 
 
