#UE4doc #docs: removed duplicate words: to to -> to.

diff --git a/Engine/Documentation/Source/Resources/Showcases/Effects/WaterExamples/WaterExamples.INT.udn b/Engine/Documentation/Source/Resources/Showcases/Effects/WaterExamples/WaterExamples.INT.udn
--- a/Engine/Documentation/Source/Resources/Showcases/Effects/WaterExamples/WaterExamples.INT.udn
+++ b/Engine/Documentation/Source/Resources/Showcases/Effects/WaterExamples/WaterExamples.INT.udn
@@ -38,7 +38,7 @@ As the water falls down the rock face, you will notice it splashing and collidin
 
 ### Particle Color
 
-It is important to keep in mind that GPU Sprite particles do **not** support Dynamic Parameters, which are typically used to send data from a Particle System into a Material. Because of this, we have employed some trickery, using a Particle Color node within the Material. Typically, this node allows us to  to change the color of the particles as they fall, by animating their color and alpha values across their lifespan. Here, however, we are using that Red, Green, Blue, and Alpha (RGBA) data to drive a few more things. Since the water does not need to actually change color (water is colorless, after all), we use the Particle Color data in the following manner:
+It is important to keep in mind that GPU Sprite particles do **not** support Dynamic Parameters, which are typically used to send data from a Particle System into a Material. Because of this, we have employed some trickery, using a Particle Color node within the Material. Typically, this node allows us to change the color of the particles as they fall, by animating their color and alpha values across their lifespan. Here, however, we are using that Red, Green, Blue, and Alpha (RGBA) data to drive a few more things. Since the water does not need to actually change color (water is colorless, after all), we use the Particle Color data in the following manner:
 
 * **Red** - Drives refractive index of the water, or how much it bends light.
 * **Green** - Drives the actual color. This means that the Green channel is driving all RGB data for color, which gives a grayscale result.
