#UEDoc #docs - Bold lettering adjustments.

diff --git a/Engine/Documentation/Source/Programming/UnrealArchitecture/Actors/Actors.INT.udn b/Engine/Documentation/Source/Programming/UnrealArchitecture/Actors/Actors.INT.udn
--- a/Engine/Documentation/Source/Programming/UnrealArchitecture/Actors/Actors.INT.udn
+++ b/Engine/Documentation/Source/Programming/UnrealArchitecture/Actors/Actors.INT.udn
@@ -5,7 +5,6 @@ Description:Explanations of the basic gameplay elements, Actors and Objects.
 
 [TOC (start:2 end:3)]
 
-## Overview
 
 <!-- An Actor is the base type of Object that can be placed or [spawned](Programming\UnrealArchitecture\Actors\Spawning) into the world. A good short description of Actor would be "a collection of Components that tick". In reality, Actors can do 
 many things; check out `Actor.h` for all the included variables and functions. 
@@ -23,13 +22,13 @@ as "pending kill", which means they will hang around until they are cleaned up o
 
 Creating new instances of `AActor` classes is called **spawning**. This can be done using the generic `SpawnActor()` function or one of its specialized templated versions.
 
-See [](Programming/UnrealArchitecture/Actors/Spawning) for detailed info on the various methods of creating instances of `AActor` classes for gameplay.
+See **[](Programming/UnrealArchitecture/Actors/Spawning)** for detailed info on the various methods of creating instances of `AActor` classes for gameplay.
 
 
 ## Components
 
-**Actors** can be thought of, in one sense, as containers that hold special types of **Objects** called [Components](Programming/UnrealArchitecture/Actors/Components). Different types of Components 
-can be used to control how **Actors** move, how they are rendered, etc. The other main function of **Actors** is the [replication](#Replication) of properties and function calls across the network 
+**Actors** can be thought of, in one sense, as containers that hold special types of **Objects** called **[Components](Programming/UnrealArchitecture/Actors/Components)**. Different types of Components 
+can be used to control how Actors move, how they are rendered, etc. The other main function of Actors is the **[replication](#Replication)** of properties and function calls across the network 
 during play.
 
 Components are associated with their containing Actor when they are created.
@@ -44,9 +43,9 @@ in the world. They can also join together in an attachment hierarchy.
 $ UPrimitiveComponent : This is a Component that can not only exist in the world, but has a representation of some kind (e.g. a mesh or a particle system). Many of the interesting physics and collision 
 settings are here.
 
-Actors support having a hierarchy of SceneComponents. Each **Actor** also has a `RootComponent` property that designates which Component acts as the root for the **Actor**. **Actors** themselves do not have transforms, and thus do not have locations, 
+Actors support having a hierarchy of SceneComponents. Each Actor also has a `RootComponent` property that designates which Component acts as the root for the Actor. Actors themselves do not have transforms, and thus do not have locations, 
 rotations, or scales. Instead, they rely on the transforms of their Components; more specifically, their root Component. If this Component is a **SceneComponent**, it provides the transformation 
-information for the **Actor**. Otherwise, the **Actor** will have no transform. Other attached Components have a transform relative to the Component they are attached to. 
+information for the Actor. Otherwise, the Actor will have no transform. Other attached Components have a transform relative to the Component they are attached to. 
 
 An example Actor and hierarchy might look something like this:
 
@@ -70,14 +69,14 @@ For more information, see the [](Programming/UnrealArchitecture/Actors/Component
 
 Ticking refers to how Objects are updated in Unreal Engine. All Objects have the ability to be ticked each frame, allowing you to perform any update calculations or actions that are necessary.
 
-**Actors** all have the ability to be ticked by default via the `TickActor()` function. This function is executed every frame by the engine.
+Actors all have the ability to be ticked by default via the `TickActor()` function. This function is executed every frame by the engine.
 
-**ActorComponents** also have the ability to be updated by default. See the [Updating section](Programming/UnrealArchitecture/Actors/Components#Updating) of the Components page for more information.
+**ActorComponents** also have the ability to be updated by default. See the **[Updating section](Programming/UnrealArchitecture/Actors/Components#Updating)** of the Components page for more information.
 
 
 ## Lifecycle
 
-See the [Actor Lifecycle](Programming/UnrealArchitecture/Actors/ActorLifecycle) documentation for more information on how an Actor is created and removed from the game.
+See the **[Actor Lifecycle](Programming/UnrealArchitecture/Actors/ActorLifecycle)** documentation for more information on how an Actor is created and removed from the game.
 
 
 ## Replication
