#UEDoc #docs - Bold lettering adjustments.

diff --git a/Engine/Documentation/Source/Programming/UnrealArchitecture/Actors/ActorLifecycle/ActorLifecycle.INT.udn b/Engine/Documentation/Source/Programming/UnrealArchitecture/Actors/ActorLifecycle/ActorLifecycle.INT.udn
--- a/Engine/Documentation/Source/Programming/UnrealArchitecture/Actors/ActorLifecycle/ActorLifecycle.INT.udn
+++ b/Engine/Documentation/Source/Programming/UnrealArchitecture/Actors/ActorLifecycle/ActorLifecycle.INT.udn
@@ -5,9 +5,8 @@ Description:What actually happens when an Actor is loaded or spawned, and eventu
 
 [TOC (start:2 end:3)]
 
-## Overview
 
-This document is a high level overview of the lifecycle of an Actor: how an Actor is instantiated (**Spawned**) into the level and then how it is removed (**Destroyed**).
+This document is a high level overview of the lifecycle of an **Actor**: how an Actor is instantiated (**Spawned**) into the level and then how it is removed (**Destroyed**).
 
 The flow graph below shows the three primary paths for how an Actor is instanced. No matter how an Actor is created, they all follow the same path to their destruction.
 
@@ -22,7 +21,7 @@ The flow graph below shows the three primary paths for how an Actor is instanced
 This path occurs for any Actor that is already in the level, like when LoadMap occurs, or AddToWorld (from streaming or sub levels) is called.
 
 1. Actors in a package/level are loaded from disk.
-1. **PostLoad** - is called by serialized Actor after they have finished loading from disk. Any custom versioning and fixup behavior should go here. **PostLoad** is mutually exclusive with **PostActorCreated**.
+1. **PostLoad** - is called by serialized Actor after they have finished loading from disk. Any custom versioning and fixup behavior should go here. PostLoad is mutually exclusive with **PostActorCreated**.
 1. **InitializeActorsForPlay**
 1. **RouteActorInitialize** for any non-initialized Actors (covers seamless travel carry over)
 	1. **PreInitializeComponents** - Called before InitializeComponent is called on the Actor's components
@@ -49,7 +48,7 @@ When spawning (instancing) an Actor, this is the path that will be followed.
 
 1. **SpawnActor** called
 1. **PostSpawnInitialize** 
-1. **PostActorCreated** - called for spawned Actors after its creation, constructor like behavior should go here. **PostActorCreated** is mutually exclusive with **PostLoad**.
+1. **PostActorCreated** - called for spawned Actors after its creation, constructor like behavior should go here. PostActorCreated is mutually exclusive with PostLoad.
 1. **ExecuteConstruction**:
 	* **OnConstruction** - The construction of the Actor, this is where Blueprint Actors have their components created and blueprint variables are initialized
 1. **PostActorConstruction**:
