#UE4doc #docs: Removed extraneous bolding as according to the instructional guide. ##Overview removed as according to the instructional guide. Altered images to match the instructional guide: green -> yellow.

diff --git a/Engine/Documentation/Source/Resources/Showcases/MatineeFightScene/MatineeFightScene.INT.udn b/Engine/Documentation/Source/Resources/Showcases/MatineeFightScene/MatineeFightScene.INT.udn
--- a/Engine/Documentation/Source/Resources/Showcases/MatineeFightScene/MatineeFightScene.INT.udn
+++ b/Engine/Documentation/Source/Resources/Showcases/MatineeFightScene/MatineeFightScene.INT.udn
@@ -60,13 +60,11 @@ Related: Engine/Basics/Projects/CreatingTemplates
 
 [TOC(start:2 end:2)]
 
-## Overview
-
-The **Matinee** example shows how to create highly stylized cinematic sequences using Unreal Engine 4's **Matinee Editor**. If you do not know what **Matinee** is, the **Matinee Editor** (or simply **Matinee**) is a system which is based on the use of specialized **Tracks** in which you can place **Keyframes** and set the values of certain properties for the **Actors** in your level where you can adjust those values over time. 
+The **Matinee** example shows how to create highly stylized cinematic sequences using Unreal Engine 4's **Matinee Editor**. If you do not know what Matinee is, the Matinee Editor (or simply Matinee) is a system which is based on the use of specialized **Tracks** in which you can place **Keyframes** and set the values of certain properties for the **Actors** in your level where you can adjust those values over time. 
 
 On this page, we will jump into and discuss each of the different sections inside the **Matinee Actor** used to create the fight sequence scene and how each of the elements come together to produce the final result shown in the **Scene Breakdown** section below. With an understanding of how it was constructed, you will be able to modify it and put a different spin on it or take the knowledge, tips-and-tricks learned here and apply them to your own scenes. 
 
-You can also head to the [](Engine\Matinee) documentation for more details on how to work with and use **Matinee**. 
+You can also head to the [](Engine\Matinee) documentation for more details on how to work with and use Matinee. 
 
 
 ## Scene Breakdown
@@ -89,9 +87,9 @@ Before breaking down the scene, you can view the end result below:
 
 ## Opening the Matinee Actor
 
-When going through this page, it will help if you have the project up and the **Matinee** open so you can follow along and can adjust values in real time. 
+When going through this page, it will help if you have the project up and the Matinee open so you can follow along and can adjust values in real time. 
 
-You can open the **Matinee** by searching the **World Outliner** for **Matinee**: 
+You can open the **Matinee** by searching the **World Outliner** for Matinee: 
 
 ![](SceneOutlinerMatinee.png)
 
