#UEDoc #docs - bold lettering adjustments.

diff --git a/Engine/Documentation/Source/Programming/UnrealArchitecture/Actors/Components/Components.INT.udn b/Engine/Documentation/Source/Programming/UnrealArchitecture/Actors/Components/Components.INT.udn
--- a/Engine/Documentation/Source/Programming/UnrealArchitecture/Actors/Components/Components.INT.udn
+++ b/Engine/Documentation/Source/Programming/UnrealArchitecture/Actors/Components/Components.INT.udn
@@ -9,10 +9,17 @@ Version: 4.5
 ## Overview
 
 [EXCERPT:Overview]
-**Components** are a special type of **Object** designed to be used as sub-objects within **Actors**. These are generally used in situations where it is necessary to have easily swappable parts in order to change the behavior or functionality of some particular aspect of the owning **Actor**. For instance, a car is controlled and moves very differently from an aircraft, which is controlled and moves differently from a boat, etc.; yet all of these are vehicles that share other commonalities. By using a Component to handle the controls and movement, the same vehicle can easily be made to behave like any one of these specific types.
+**Components** are a special type of **Object** designed to be used as sub-objects within **Actors**. These are generally used in situations where it is necessary to have easily swappable parts in 
+order to change the behavior or functionality of some particular aspect of the owning Actor. For instance, a car is controlled and moves very differently from an aircraft, which is controlled and 
+moves differently from a boat, etc.; yet all of these are vehicles that share other commonalities. By using a Component to handle the controls and movement, the same vehicle can easily be made to behave 
+like any one of these specific types.
 [/EXCERPT:Overview]
 
-Contrary to the default behavior of sub-objects in general, Components created as sub-objects within an **Actor** are instanced, meaning each **Actor** instance of a particular class gets its own unique instances of the Components. A simple way to visualize this is to imagine the vehicle described above. A **Car** class might use Components to represent the wheels of the car. Four **Wheel** Components would be created as sub-objects in the default properties of the class and assigned to a `Wheels[]` array. When a new **Car** instance is created, new instances of the wheel Components are created specifically for that **Car**. If this were not the case, when one **Car** in the world moved, the wheels of all **Cars** would turn; which is clearly not the desired behavior. Instancing of Components by default simplifies the process of quickly adding unique sub-objects to **Actors**. 
+Contrary to the default behavior of sub-objects in general, Components created as sub-objects within an Actor are instanced, meaning each Actor instance of a particular class gets its own 
+unique instances of the Components. A simple way to visualize this is to imagine the vehicle described above. A **Car** class might use Components to represent the wheels of the car. Four **Wheel** 
+Components would be created as sub-objects in the default properties of the class and assigned to a `Wheels[]` array. When a new Car instance is created, new instances of the wheel Components 
+are created specifically for that Car. If this were not the case, when one Car in the world moved, the wheels of all Cars would turn; which is clearly not the desired behavior. Instancing 
+of Components by default simplifies the process of quickly adding unique sub-objects to Actors. 
 
 [REGION:note]
 **Note:** Without Component instancing, all Component variables would need to be declared using the **[Instanced](Programming\UnrealArchitecture\Reference\Properties\Specifiers\Instanced) property specifier**.
@@ -20,17 +27,23 @@ Contrary to the default behavior of sub-objects in general, Components created a
 
 ## Actor Components
 
-**ActorComponent** is the base class for components that define reusable behavior that can be added to different types of **Actors**. These are generally used to associate geometry with an **Actor**, in the form of collision geometry or rendered meshes, control how the **Actor** moves through the world, play sounds associated with the **Actor**, provide the ability for the **Actor** to cast light and shadows onto the world, and more. In fact, everything the player sees or interacts with in the world when playing the game is actually the work of one type of **ActorComponent** or another. **ActorComponents** that have a transform are known as **SceneComponents** and those that can be rendered are **PrimitiveComponents**.
+**ActorComponent** is the base class for components that define reusable behavior that can be added to different types of Actors. These are generally used to associate geometry with an Actor, 
+in the form of collision geometry or rendered meshes, control how the Actor moves through the world, play sounds associated with the Actor, provide the ability for the Actor to cast light 
+and shadows onto the world, and more. In fact, everything the player sees or interacts with in the world when playing the game is actually the work of one type of ActorComponent or another. 
+ActorComponents that have a transform are known as **SceneComponents** and those that can be rendered are **PrimitiveComponents**.
 
 ### Registering Components
 
-In order for **ActorComponents** to be updated each frame and affect the scene, they must **register** with the scene. Registration is performed by calling `UActorComponent::RegisterComponent()`. 
+In order for ActorComponents to be updated each frame and affect the scene, they must **register** with the scene. Registration is performed by calling `UActorComponent::RegisterComponent()`. 
 
 	void UActorComponent::RegisterComponent()
 
-This function calls `UActorComponent::RegisterComponentWithScene()` to make sure that the **ActorComponent** is contained in the owning **Actor**'s `Components[]` array, associate it with the scene, and create the [render proxy](#SceneProxy) and [physics state](#PhysicsState) for the **ActorComponent**. 
+This function calls `UActorComponent::RegisterComponentWithScene()` to make sure that the ActorComponent is contained in the owning Actor's `Components[]` array, associate it with the 
+scene, and create the **[render proxy](#SceneProxy)** and **[physics state](#PhysicsState)** for the ActorComponent. 
 
-**ActorComponents** are automatically registered when the owning **Actor** is spawned, if they are created as sub-objects and added to the `Components[]` array in the default properties of the **Actor**. They can be registered at any time during play by calling `UActorComponent::RegisterComponent()`, but be aware that registering a Component is a somewhat heavy operation, so it should only be performed when necessary.
+ActorComponents are automatically registered when the owning Actor is spawned, if they are created as sub-objects and added to the `Components[]` array in the default properties of the 
+Actor. They can be registered at any time during play by calling `UActorComponent::RegisterComponent()`, but be aware that registering a Component is a somewhat heavy operation, so it should 
+only be performed when necessary.
 
 #### Register Events
 
@@ -45,7 +58,7 @@ When a Component is registered, the events below are fired off.
 
 ### UnRegistering Components
 
-**ActorComponents** can be unregistered as well to avoid them being updated, simulated, or rendered. Unregistering a Component is performed by calling `UActorComponent::UnregisterComponent()`. 
+ActorComponents can be unregistered as well to avoid them being updated, simulated, or rendered. Unregistering a Component is performed by calling `UActorComponent::UnregisterComponent()`. 
 
 	void UActorComponent::UnregisterComponent()
 
@@ -62,31 +75,43 @@ The events below are fired off when a Component is unregistered.
 
 ### Updating
 
-**ActorComponents** have the ability to be updated each frame via their `TickComponent()` function. This allows Components to perform type-specific, per-frame calculations. For example, **SkeletalMeshComponents** use `TickComponent()` to update animations and skeletal controllers, while **ParticleSystemComponents** use it to update the emitters of the system and check for particle events to handle.
+ActorComponents have the ability to be updated each frame via their `TickComponent()` function. This allows Components to perform type-specific, per-frame calculations. For example, 
+**SkeletalMeshComponents** use `TickComponent()` to update animations and skeletal controllers, while **ParticleSystemComponents** use it to update the emitters of the system and check for particle 
+events to handle.
 
 In order for a Component to be updated, it must be registered, set to tick (`bComponentNeverTicks=false`), and have its tick function setup. 
 
 ### Render State
 
-An **ActorComponent** must have a render state created in order to be rendered. The render state of an **ActorComponent** also lets the engine know whether something has changed about the Component that requires its render data to be updated. When such a change occurs, the render state is marked dirty. Then, at the end of the current frame, all dirty Components have their render data updated.
+An ActorComponent must have a render state created in order to be rendered. The render state of an ActorComponent also lets the engine know whether something has changed about the Component 
+that requires its render data to be updated. When such a change occurs, the render state is marked dirty. Then, at the end of the current frame, all dirty Components have their render data updated.
 
 ### Physics State
 
-In order to be simulated using the physics engine, an **ActorComponent** must have a physics state created. Unlike the render state, the physics state is never marked dirty as physics updates happen immediately when changes occur. This is important as it keeps issues like frame-behind artifacts from happening.
+In order to be simulated using the physics engine, an ActorComponent must have a physics state created. Unlike the render state, the physics state is never marked dirty as physics updates happen 
+immediately when changes occur. This is important as it keeps issues like frame-behind artifacts from happening.
 
 ## Scene Components
 
-**SceneComponents** are an extension of **ActorComponents** that have a transform, i.e. location, rotation, and scale. The addition of a transform means **SceneComponents** can also be attached to one another.
+SceneComponents are an extension of ActorComponents that have a transform, i.e. location, rotation, and scale. The addition of a transform means SceneComponents can also be attached to one 
+another.
 
 ### Attachment
 
-In Unreal Engine 4, all attachment is handled on the Component level and only **SceneComponents** can be attached to one another. **SceneComponents** have an `AttachParent` property that points to the Component it is attached to. This is mainly used when an **Actor** contains multiple Components to keep them all attached, if this is the desired behavior. However, it can also be used to specify a Component belonging to another **Actor** to attach a Component in one **Actor** to a Component in another **Actor**.
+In Unreal Engine 4, all attachment is handled on the Component level and only SceneComponents can be attached to one another. SceneComponents have an `AttachParent` property that points to 
+the Component it is attached to. This is mainly used when an Actor contains multiple Components to keep them all attached, if this is the desired behavior. However, it can also be used to specify 
+a Component belonging to another Actor to attach a Component in one Actor to a Component in another Actor.
 
-This means that in order to attach one **Actor** to another, each of the **Actors** involved must contain at least one **SceneComponent** - as these are what will actually be attached - and the **SceneComponent** of the **Actor** that is being attached must be the root Component. The reason for this should be apparent; each Component can only have one parent so if the attached Component is not the root Component of the **Actor** being attached, that Component and any of its child Components are the only ones that will be attached. The root Component would remain unattached and unaffected, as would the **Actor** as a whole since it relies on the root Component and its transform.
+This means that in order to attach one Actor to another, each of the Actors involved must contain at least one SceneComponent - as these are what will actually be attached - and the 
+SceneComponent of the Actor that is being attached must be the root Component. The reason for this should be apparent; each Component can only have one parent so if the attached Component 
+is not the root Component of the Actor being attached, that Component and any of its child Components are the only ones that will be attached. The root Component would remain unattached and 
+unaffected, as would the Actor as a whole since it relies on the root Component and its transform.
 
 ### Component Transforms
 
-The `FTransform` struct contains a `Translation` vector, a `Rotation` quaternion, and a `Scale3D` vector. Each **SceneComponent** has its own `FTransform`, generally for internal use, that describes its world-relative location, rotation, and scale. They also have an additional `RelativeLocation` vector, `RelativeRotation` rotator, and `RelativeScale3D` vector used which can act as either relative to their parent Component or relative to the world. These are generally what are used when getting or setting the transform for a Component.
+The `FTransform` struct contains a `Translation` vector, a `Rotation` quaternion, and a `Scale3D` vector. Each SceneComponent has its own `FTransform`, generally for internal use, that 
+describes its world-relative location, rotation, and scale. They also have an additional `RelativeLocation` vector, `RelativeRotation` rotator, and `RelativeScale3D` vector used which can act 
+as either relative to their parent Component or relative to the world. These are generally what are used when getting or setting the transform for a Component.
 
 By default, `RelativeLocation`, `RelativeRotation`, and `RelativeScale3D` are relative to their `AttachParent`. These values can be forced to be relative to the world using the `bAbsoluteLocation`, `bAbsoluteRotation`, and `bAbsoluteScale` properties. Setting any of these to `true` will cause the corresponding property to be world-relative.
 
@@ -100,11 +125,14 @@ Even when using parent-relative transform values, it is still possible to set tr
 
 ## Primitive Components
 
-**PrimitiveComponents** are **SceneComponents** that contain or generate some sort of geometry, generally to be rendered or used as collision data. There are several subclasses for the various types of geometry, but the most common by far are the **CapsuleComponent**, **StaticMeshComponent**, and **SkeletalMeshComponent**. **CapsuleComponents** generate geometry that is used for collision detection, but not rendered; while **StaticMeshComponents** and **SkeletalMeshComponents** contain pre-built geometry that is rendered, but can also be used for collision detection.
+**PrimitiveComponents** are SceneComponents that contain or generate some sort of geometry, generally to be rendered or used as collision data. There are several subclasses for the various types 
+of geometry, but the most common by far are the **CapsuleComponent**, **StaticMeshComponent**, and **SkeletalMeshComponent**. **CapsuleComponents** generate geometry that is used for collision detection, 
+but not rendered; while StaticMeshComponents and SkeletalMeshComponents contain pre-built geometry that is rendered, but can also be used for collision detection.
 
 ### Scene Proxy
 
-The `SceneProxy` of a **PrimitiveComponent** is an `FPrimitiveSceneProxy` instance that encapsulates the scene data that is mirrored in order to render the primitive in parallel to the game thread. A subclass of `FPrimitiveSceneProxy` is created for each primitive type to hold the specific render data necessary to render that type of primitive.
+The `SceneProxy` of a PrimitiveComponent is an `FPrimitiveSceneProxy` instance that encapsulates the scene data that is mirrored in order to render the primitive in parallel to the game thread. 
+A subclass of `FPrimitiveSceneProxy` is created for each primitive type to hold the specific render data necessary to render that type of primitive.
 
-See [Rendering System Overview](Programming/Rendering) for more details on primitives and rendering geometry.
+See **[Rendering System Overview](Programming/Rendering)** for more details on primitives and rendering geometry.
 
