UEDOC-1892

Gif files have been replaced with YouTube videos

diff --git a/Engine/Documentation/Source/Engine/Rendering/ParticleSystems/Reference/Modules/Size/SizeModules.INT.udn b/Engine/Documentation/Source/Engine/Rendering/ParticleSystems/Reference/Modules/Size/SizeModules.INT.udn
--- a/Engine/Documentation/Source/Engine/Rendering/ParticleSystems/Reference/Modules/Size/SizeModules.INT.udn
+++ b/Engine/Documentation/Source/Engine/Rendering/ParticleSystems/Reference/Modules/Size/SizeModules.INT.udn
@@ -9,7 +9,17 @@ These modules affect the sizes of particles across various points of their lifes
 
 ## Initial Size
 
-![psys_initial_size.gif](psys_initial_size.gif)(convert:false)
+[OBJECT:EmbeddedVideo]
+[PARAMLITERAL:width]
+640
+[/PARAMLITERAL]
+[PARAMLITERAL:height]
+360
+[/PARAMLITERAL]
+[PARAMLITERAL:videoid]
+e72tK-OB9oE
+[/PARAMLITERAL]
+[/OBJECT]
 
 The Initial Size module is used to set the initial size of a particle at spawn time. It contains the following member(s):
 
@@ -35,7 +45,17 @@ The Initial Size (Seeded) module is identical to the [Initial Size](#InitialSize
 
 ## Size By Life
 
-![](SizeByLife.gif)(convert:false)
+[OBJECT:EmbeddedVideo]
+[PARAMLITERAL:width]
+640
+[/PARAMLITERAL]
+[PARAMLITERAL:height]
+360
+[/PARAMLITERAL]
+[PARAMLITERAL:videoid]
+UPWGbVD8ync
+[/PARAMLITERAL]
+[/OBJECT]
 
 The Size By Life module scales the size of the particle by a given value over its lifetime. It contains the following member(s):
 
@@ -62,7 +82,17 @@ The Size Scale module sets the Size of the particle to the BaseSize times the gi
 
 ## Size By Speed
 
-![](SizeBySpeed.gif)(convert:false)
+[OBJECT:EmbeddedVideo]
+[PARAMLITERAL:width]
+640
+[/PARAMLITERAL]
+[PARAMLITERAL:height]
+360
+[/PARAMLITERAL]
+[PARAMLITERAL:videoid]
+LeUTwLRnh6M
+[/PARAMLITERAL]
+[/OBJECT]
 
 Rather than recede into the distance, these particles are actually shrinking as they slow down.
 
