#UEDoc #docs - Bold lettering adjustments.

diff --git a/Engine/Documentation/Source/Programming/UnrealArchitecture/Delegates/Delegates.INT.udn b/Engine/Documentation/Source/Programming/UnrealArchitecture/Delegates/Delegates.INT.udn
--- a/Engine/Documentation/Source/Programming/UnrealArchitecture/Delegates/Delegates.INT.udn
+++ b/Engine/Documentation/Source/Programming/UnrealArchitecture/Delegates/Delegates.INT.udn
@@ -7,9 +7,11 @@ Version: 4.5
 [TOC(start:2)]
 
 [EXCERPT:Overview]
-Delegates allow you to call member functions on C++ objects in a generic, yet type-safe way. Using delegates, you can dynamically bind to a member function of an arbitrary object, then call functions on the object, even if the caller does not know the object's type.
+Delegates allow you to call member functions on C++ objects in a generic, yet type-safe way. Using delegates, you can dynamically bind to a member function of an 
+arbitrary object, then call functions on the object, even if the caller does not know the object's type.
 
-It is perfectly safe to copy delegate objects.  Delegates can be passed around by value but this is generally not recommended since they do have to allocate memory on the heap.  **You should always pass delegates by reference when possible.**
+It is perfectly safe to copy delegate objects.  Delegates can be passed around by value but this is generally not recommended since they do have to allocate memory on 
+the heap.  **You should always pass delegates by reference when possible.**
 
 Both single-cast and multi-cast delegates are supported, as well as "dynamic" delegates which can be safely serialized to disk.
 
@@ -23,7 +25,9 @@ Both single-cast and multi-cast delegates are supported, as well as "dynamic" de
 
 ## Declaring Delegates
 
-Delegates are declared using one of the declaration macros provided. The macro used is determined by the signature of the function(s) to be bound to the delegate. The system predefines various combinations of generic function signatures with which you can declare a delegate type from, filling in the type names for return value and parameters with whichever types you need. Currently delegate signatures using any combination of the following are supported:
+Delegates are declared using one of the declaration macros provided. The macro used is determined by the signature of the function(s) to be bound to the delegate. The 
+system predefines various combinations of generic function signatures with which you can declare a delegate type from, filling in the type names for return value and 
+parameters with whichever types you need. Currently delegate signatures using any combination of the following are supported:
 
 * Functions returning a value
 * Up to four "payload" variables
@@ -56,11 +60,14 @@ Variations of the macros above for multi-cast, dynamic, and wrapped delegates ar
 
 Delegate signature declarations can exist at global scope, within a namespace or even within a class declaration (but not function bodies.)
 
-See [](Programming/UnrealArchitecture/Delegates/Dynamic) and [](Programming/UnrealArchitecture/Delegates/Multicast) for more information on declaring these types of delegates.
+See **[](Programming/UnrealArchitecture/Delegates/Dynamic)** and **[](Programming/UnrealArchitecture/Delegates/Multicast)** for more information on declaring these types of 
+delegates.
 
 ## Binding Delegates
 
-The delegate system understands certain types of objects, and additional features are enabled when using these objects.  If you bind a delegate to a member of a UObject or shared pointer class, the delegate system can keep a weak reference to the object, so that if the object gets destroyed out from underneath the delegate, you will be able to handle these cases by calling `IsBound()` or `ExecuteIfBound()` functions.  Note the special binding syntax for the various types of supported objects.
+The delegate system understands certain types of objects, and additional features are enabled when using these objects.  If you bind a delegate to a member of a 
+UObject or shared pointer class, the delegate system can keep a weak reference to the object, so that if the object gets destroyed out from underneath the delegate, 
+you will be able to handle these cases by calling `IsBound()` or `ExecuteIfBound()` functions.  Note the special binding syntax for the various types of supported objects.
 
 | Function | Description |
 | --- | --- |
@@ -75,14 +82,20 @@ See `DelegateSignatureImpl.inl` (located in `..\UE4\Engine\Source\Runtime\Core\P
 
 ### Payload Data
 
-When binding to a delegate, you can pass payload data along. These are arbitrary variables that will be passed directly to any bound function when it is invoked.  This is really useful as it allows you to store parameters within the delegate itself at bind-time.  All delegate types (except for "dynamic") supports payload variables automatically.  This example passes two custom variables, a bool and an int32 to a delegate.  Then when the delegate is invoked, these parameters will be passed to your bound function.  The extra variable arguments must always be accepted after the delegate type parameter arguments.
+When binding to a delegate, you can pass payload data along. These are arbitrary variables that will be passed directly to any bound function when it is invoked.  This is really 
+useful as it allows you to store parameters within the delegate itself at bind-time.  All delegate types (except for "dynamic") supports payload variables automatically.  
+This example passes two custom variables, a bool and an int32 to a delegate.  Then when the delegate is invoked, these parameters will be passed to your bound function.  The 
+extra variable arguments must always be accepted after the delegate type parameter arguments.
 
 	MyDelegate.BindRaw( &MyFunction, true, 20 );
 
 ## Executing Delegates
 
 [EXCERPT:ExecutingDelegates]
-The function bound to a delegate is executed by calling the delegate's `Execute()` function. You must check if delegates are "bound" before executing them.  This is to make the code more safe as there may be cases where delegates have return values and output parameters that are uninitialized and subsequently accessed.  Executing an unbound delegate could actually scribble over memory in some instances. You can call `IsBound()` to check if the delegate is safe to execute.  Also, for delegates that have no return value, you can call `ExecuteIfBound()`, but be wary of output parameters that may be left uninitialized.
+The function bound to a delegate is executed by calling the delegate's `Execute()` function. You must check if delegates are "bound" before executing them.  This is to make the 
+code more safe as there may be cases where delegates have return values and output parameters that are uninitialized and subsequently accessed.  Executing an unbound delegate 
+could actually scribble over memory in some instances. You can call `IsBound()` to check if the delegate is safe to execute.  Also, for delegates that have no return value, 
+you can call `ExecuteIfBound()`, but be wary of output parameters that may be left uninitialized.
 
 | Execution Functions | Description |
 | --- | --- |
@@ -90,7 +103,7 @@ The function bound to a delegate is executed by calling the delegate's `Execute(
 | `ExecuteIfBound()` |  |
 | `IsBound()` |  |
 
-See [](Programming/UnrealArchitecture/Delegates/Multicast) for details on executing multi-cast delegates.
+See **[](Programming/UnrealArchitecture/Delegates/Multicast)** for details on executing multi-cast delegates.
 [/EXCERPT:ExecutingDelegates]
 
 ## Example Usage
@@ -102,7 +115,8 @@ Suppose you have a class with a method that you would like to be able to call fr
 		void WriteToLog( FString );
 	};
 
-To call the WriteToLog function, we will need to create a delegate type for that function's signature. To do this, you will first declare the delegate using one of the macros below.  For example, here is a simple delegate type:
+To call the WriteToLog function, we will need to create a delegate type for that function's signature. To do this, you will first declare the delegate using one of the macros 
+below.  For example, here is a simple delegate type:
 
 	DECLARE_DELEGATE_OneParam( FStringDelegate, FString );
 
@@ -117,7 +131,9 @@ Here is an example of how you would use this 'FStringDelegate' in a class:
 
 This allows your class to hold a pointer to a method in an arbitrary class.  The only thing the class really knows about this delegate is its function signature.
 
-Now, to assign the delegate, simply create an instance of your delegate class, passing along the class that owns the method as a template parameter.  You will also pass the instance of your object and the actual function address of the method.  So, here we will create an instance of our 'FLogWriter' class, then create a delegate for the 'WriteToLog' method of that object instance:
+Now, to assign the delegate, simply create an instance of your delegate class, passing along the class that owns the method as a template parameter.  You will also pass the 
+instance of your object and the actual function address of the method.  So, here we will create an instance of our 'FLogWriter' class, then create a delegate for the 'WriteToLog' 
+method of that object instance:
 
 	FSharedRef< FLogWriter > LogWriter( new FLogWriter() );
 
@@ -125,7 +141,8 @@ Now, to assign the delegate, simply create an instance of your delegate class, p
 
 You have just dynamically bound a delegate to a method of a class!  Pretty simple, right?
 
-Note that the `SP` part of `BindSP` stands for `shared pointer`, because we are binding to an object that is owned by a shared pointer.  There are versions for different object types, such as BindRaw() and BindUObject().
+Note that the `SP` part of `BindSP` stands for `shared pointer`, because we are binding to an object that is owned by a shared pointer.  There are versions for different 
+object types, such as BindRaw() and BindUObject().
 
 Now, your 'WriteToLog' method can be called by FMyClass without it even knowing anything about the 'FLogWriter' class!  To call your delegate, just use the 'Execute()' method:
 
