#UEDoc #docs - Bold lettering adjustments.

diff --git a/Engine/Documentation/Source/Programming/UnrealArchitecture/Delegates/Dynamic/DynamicDelegates.INT.udn b/Engine/Documentation/Source/Programming/UnrealArchitecture/Delegates/Dynamic/DynamicDelegates.INT.udn
--- a/Engine/Documentation/Source/Programming/UnrealArchitecture/Delegates/Dynamic/DynamicDelegates.INT.udn
+++ b/Engine/Documentation/Source/Programming/UnrealArchitecture/Delegates/Dynamic/DynamicDelegates.INT.udn
@@ -8,7 +8,8 @@ Dynamic delegates can be serialized, their functions can be found by name, and t
 
 ## Declaring Dynamic Delegates
 
-Dynamic delegates are declared in the same manner as you [declare standard delegates](Programming/UnrealArchitecture/Delegates#DeclaringDelegates) except they use the macro variations specific to dynamic delegates.
+Dynamic delegates are declared in the same manner as you **[declare standard delegates](Programming/UnrealArchitecture/Delegates#DeclaringDelegates)** except they use the macro 
+variations specific to dynamic delegates.
 
 | Declaration Macro | Description |
 | --- | --- |
