TypoCorrection

diff --git a/Engine/Documentation/Source/Engine/QuickStart/4/LDQS_4.INT.udn b/Engine/Documentation/Source/Engine/QuickStart/4/LDQS_4.INT.udn
--- a/Engine/Documentation/Source/Engine/QuickStart/4/LDQS_4.INT.udn
+++ b/Engine/Documentation/Source/Engine/QuickStart/4/LDQS_4.INT.udn
@@ -59,7 +59,7 @@ After completing this step, you will know how to place Actors inside your own le
 
 	![](placeLight.png)
 	
-1. On the on the [Translation Tool](https://docs.unrealengine.com/latest/INT/Engine/UI/LevelEditor/Viewports/ViewportControls/index.html#transformcontrols), click and drag the Z Axis (Blue) gizmo up away from the surface of the floor.
+1. On the on the [Translation Tool](Engine/UI/LevelEditor/Viewports/ViewportControls), click and drag the Z Axis (Blue) gizmo up away from the surface of the floor.
 
 	![](moveLight.png)
 
@@ -81,7 +81,7 @@ After completing this step, you will know how to place Actors inside your own le
 
 	The **Lightmass Importance Volume** is used to control and concentrate lighting and shadowing affects within the volume. When placing the Lightmass Importance Volume in the level, the default size of the volume does not cover our playable area so we will need to scale it up. 
 	
-1. Inside the **Level Viewport**, press **R** to switch to the [Scale Tool](https://docs.unrealengine.com/latest/INT/Engine/UI/LevelEditor/Viewports/ViewportControls/index.html#transformcontrols). 
+1. Inside the **Level Viewport**, press **R** to switch to the [Scale Tool](Engine/UI/LevelEditor/Viewports/ViewportControls#transformcontrols). 
 
 1. Click and drag the white box in the center of the **Scale Tool** so that the **Lightmass Importance Volume** encapsulates the floor. 
 
@@ -97,7 +97,7 @@ After completing this step, you will know how to place Actors inside your own le
 
 	![](addChair.png)
 
-1. With the **SM_Chair** selected inside the **Level Viewport**, press **E** to access the [Rotation Tool](https://docs.unrealengine.com/latest/INT/Engine/UI/LevelEditor/Viewports/ViewportControls/index.html#transformcontrols). 
+1. With the **SM_Chair** selected inside the **Level Viewport**, press **E** to access the [Rotation Tool](Engine/UI/LevelEditor/Viewports/ViewportControls#transformcontrols). 
 	
 	![](Rotate1.png)(w:440)
 
