UEDOC-1810

All Gifs have been converted into YouTube Videos

diff --git a/Engine/Documentation/Source/Engine/Rendering/Materials/HowTo/Instancing/Instancing.INT.udn b/Engine/Documentation/Source/Engine/Rendering/Materials/HowTo/Instancing/Instancing.INT.udn
--- a/Engine/Documentation/Source/Engine/Rendering/Materials/HowTo/Instancing/Instancing.INT.udn
+++ b/Engine/Documentation/Source/Engine/Rendering/Materials/HowTo/Instancing/Instancing.INT.udn
@@ -50,7 +50,17 @@ You can add  **Parameter Nodes** to the Material Graph in two completely differe
 
 * **Converting a Node** - When inside of the **Material Graph**, you can **Right-Click** on a node and the choose the **Convert To Parameter** option like so. 
 
-	![](Materails_Instances_Convert_A_Node.gif)(convert:false)
+	[OBJECT:EmbeddedVideo]
+	[PARAMLITERAL:width]
+	640
+	[/PARAMLITERAL]
+	[PARAMLITERAL:height]
+	360
+	[/PARAMLITERAL]
+	[PARAMLITERAL:videoid]
+	ICjxcEzGAno
+	[/PARAMLITERAL]
+	[/OBJECT]
 
 [region:note]
 Keep in mind that you cannot do this for every Material node. If you do not see the **Convert To Parameter** option when you **Right-Click** on a node, then that node cannot be converted into a Parameter Node. 
@@ -74,7 +84,17 @@ In the following section, we will be creating our very own Material Instance usi
 	*  **Constant Node** x 1
 	*  **Scalar Parameter** x 1
 	
-	![](Materials_Instances_Adding_Nodes.gif)(convert:false)
+	[OBJECT:EmbeddedVideo]
+	[PARAMLITERAL:width]
+	640
+	[/PARAMLITERAL]
+	[PARAMLITERAL:height]
+	360
+	[/PARAMLITERAL]
+	[PARAMLITERAL:videoid]
+	COw7Q7o1sHo
+	[/PARAMLITERAL]
+	[/OBJECT]
 
 1. Now that the Material Nodes have been placed down, it is now time to wire them up to the their respective slots in the Main Material Node. For this example, we will be using the following slots.
 	
@@ -82,11 +102,31 @@ In the following section, we will be creating our very own Material Instance usi
 	* **Metallic**
 	* **Roughness**
 
-	![](Materials_Instances_Connecting_Nodes.gif)(convert:false)
+	[OBJECT:EmbeddedVideo]
+	[PARAMLITERAL:width]
+	640
+	[/PARAMLITERAL]
+	[PARAMLITERAL:height]
+	360
+	[/PARAMLITERAL]
+	[PARAMLITERAL:videoid]
+	aN3QPYokHf4
+	[/PARAMLITERAL]
+	[/OBJECT]
 
 1. Once everything has been hooked up, it is now time to convert some of the material nodes into **Parameter Nodes** so that they can be used in [Material Instances](Resources/ContentExamples/MaterialInstances/1_2). To do this, **Right-Click** on the Material node that you want to convert and then select **Convert to Parameter** from the **Right-Click** menu.
 	
-	![](Materails_Instances_Convert_Nodes_To_Prams.gif)(convert:false)
+	[OBJECT:EmbeddedVideo]
+	[PARAMLITERAL:width]
+	640
+	[/PARAMLITERAL]
+	[PARAMLITERAL:height]
+	360
+	[/PARAMLITERAL]
+	[PARAMLITERAL:videoid]
+	-tIO2lmJhVw
+	[/PARAMLITERAL]
+	[/OBJECT]
 
 1. Now that we have our nodes converted to Parameter nodes, it is now time to label them so that when we use them in a Material Instance, we will know what each node does. To give a Parameter node a name select the Parameter node in the Material Graph with the **Left Mouse Button** and then under the **Parameter Name** section type in a name that you want to use as the **Parameter Name**. For this example, we will be using the following names. 
 	
@@ -96,7 +136,17 @@ In the following section, we will be creating our very own Material Instance usi
 
 	* **Roughness **
 
-	![](Materails_Instances_Name_Nodes.gif)(convert:false)
+	[OBJECT:EmbeddedVideo]
+	[PARAMLITERAL:width]
+	640
+	[/PARAMLITERAL]
+	[PARAMLITERAL:height]
+	360
+	[/PARAMLITERAL]
+	[PARAMLITERAL:videoid]
+	tE8_Dju0ee0
+	[/PARAMLITERAL]
+	[/OBJECT]
 	
 	[region:tip]
 	Make sure that the name that you pick for your node matches very closely to what the node actually does. This way you do not have to keep going back and forth between the Material and Material Instance to see what the node actually does. 
@@ -108,7 +158,17 @@ In the following section, we will be creating our very own Material Instance usi
 
 1. Now that the Material has been compiled, it is now to time make a Material Instance out of it. To do this, **Right-Click** on the Material in the **Content Browser** and then select the **Create Material Instance** option.
 	
-	![](Materials_Instances_Creat_Instance.gif)(convert:false)
+	[OBJECT:EmbeddedVideo]
+	[PARAMLITERAL:width]
+	640
+	[/PARAMLITERAL]
+	[PARAMLITERAL:height]
+	360
+	[/PARAMLITERAL]
+	[PARAMLITERAL:videoid]
+	PwsuXdGpe7M
+	[/PARAMLITERAL]
+	[/OBJECT]
 
 	[region:tip]
 	Your Material Instance will derive its name from the Material that you are using to create it from. You can re-name the Material Instance during creation time or by selecting the Material and pressing **F2** on the keyboard and then entering the name you want it to have. 
@@ -120,11 +180,31 @@ In the following section, we will be creating our very own Material Instance usi
 
 1. Once you have the **Material Instance Editor** open, you can then start to adjust the parameters to your liking. To adjust a parameter, you first need to enable it by clicking on the box that is to the left of the parameter name to enable it. Once a parameter has been enabled, its name will no longer be grayed out and a check mark will be placed in the box to the left of the name.
 	
-	![](Materails_Instances_Enable_Prams.gif)(convert:false)
+	[OBJECT:EmbeddedVideo]
+	[PARAMLITERAL:width]
+	640
+	[/PARAMLITERAL]
+	[PARAMLITERAL:height]
+	360
+	[/PARAMLITERAL]
+	[PARAMLITERAL:videoid]
+	aJ4fJ2buAVo
+	[/PARAMLITERAL]
+	[/OBJECT]
 
 1. Now that the parameters have been enabled, you can now adjust any of them and see the results in real time, all without having to re-compile the Material. 
 	
-	![](Materails_Instances_Change_Parms.gif)(convert:false)
+	[OBJECT:EmbeddedVideo]
+	[PARAMLITERAL:width]
+	640
+	[/PARAMLITERAL]
+	[PARAMLITERAL:height]
+	360
+	[/PARAMLITERAL]
+	[PARAMLITERAL:videoid]
+	OKvC9e4FddU
+	[/PARAMLITERAL]
+	[/OBJECT]
 
 [region:note]
 The names that you set for parameters inside of the Material reflect the names that you see inside of your Material Instance. This is why it is very important that you pick names that correspond to what the parameter actually does. 
