#UE4doc #docs: ##Overview removed as according to the instructional guide.

diff --git a/Engine/Documentation/Source/Resources/Showcases/RealisticRendering/RealisticRendering.INT.udn b/Engine/Documentation/Source/Resources/Showcases/RealisticRendering/RealisticRendering.INT.udn
--- a/Engine/Documentation/Source/Resources/Showcases/RealisticRendering/RealisticRendering.INT.udn
+++ b/Engine/Documentation/Source/Resources/Showcases/RealisticRendering/RealisticRendering.INT.udn
@@ -60,7 +60,6 @@ Related: Engine/Rendering/Materials/LayeredMaterials
 
 [TOC(start:2 end:2)]
 
-## Overview
 
 This example shows off the realistic rendering possibilities of Unreal Engine 4. The goal was to create a small scene that had the look and feel of an architectural rendering, as those produced by CG rendering engines such as VRay and Mental Ray. However, the primary stipulation in this case was that all assets be game-ready, meaning that no high-poly modeling could be used. Everything in the scene had to be created such that one could make an actual and practical game using these methods. 
 
