UEDOC-1886

All gifs have been comverted into YouTube Videos

diff --git a/Engine/Documentation/Source/Engine/Rendering/Materials/HowTo/Main_Material_Node/Main_Material_Node.INT.udn b/Engine/Documentation/Source/Engine/Rendering/Materials/HowTo/Main_Material_Node/Main_Material_Node.INT.udn
--- a/Engine/Documentation/Source/Engine/Rendering/Materials/HowTo/Main_Material_Node/Main_Material_Node.INT.udn
+++ b/Engine/Documentation/Source/Engine/Rendering/Materials/HowTo/Main_Material_Node/Main_Material_Node.INT.udn
@@ -94,7 +94,17 @@ Setting up a Material to use the Main Material Node is something can be accompli
 
 1. First use the mouse to **Right-Click** in the **Content Browser** and then from the **Create Basic Asset** section of the pop-up Menu, select **Material**. Once the Material has been created, name the Material **MainMaterialNode**.
 
-	![](MMN_Create_New_Material.gif)(convert:false)
+	[OBJECT:EmbeddedVideo]
+	[PARAMLITERAL:width]
+	640
+	[/PARAMLITERAL]
+	[PARAMLITERAL:height]
+	360
+	[/PARAMLITERAL]
+	[PARAMLITERAL:videoid]
+	ukSfdOk0J_Y
+	[/PARAMLITERAL]
+	[/OBJECT]
 
 1. Open up the Material by **Double-clicking** on it with the **Left Mouse Button** inside of the **Content Browser**. 
 	 
@@ -110,7 +120,17 @@ Setting up a Material to use the Main Material Node is something can be accompli
 
 1. Next add a Constant Material Expression nodes to the Material Graph by first searching in the **Material Pallet** using the word Constant as a search term. When the Constant Material Expression is found select it by clicking on it with the **Left Mouse Button** and with the **Left Mouse Button** still held down, drag the Material Expression node into the Material Graph releasing the **Left Mouse Button** when your mouse is over the location in the Material Graph you want to place the Material Expression node. 
 
-	![](MMN_Add_Const_Pram.gif)(convert:false)
+	[OBJECT:EmbeddedVideo]
+	[PARAMLITERAL:width]
+	640
+	[/PARAMLITERAL]
+	[PARAMLITERAL:height]
+	360
+	[/PARAMLITERAL]
+	[PARAMLITERAL:videoid]
+	M2dX_kugXpg
+	[/PARAMLITERAL]
+	[/OBJECT]
 
 1. Now that we have a Constant node, duplicate it by first selecting it and then pressing Ctrl+W on the keyboard once. When completed, you should have something that looks like this. 
 	
@@ -118,11 +138,31 @@ Setting up a Material to use the Main Material Node is something can be accompli
 
 1. With all of the required Material Expression nodes now in place, it is now time to setup and connect them to the Main Material Node. First **Double-Click** on the Constant3Vector node using the **Left Mouse Button** to open up the Color Picker.  With the **Color Picker** open, select the color red and then press the **Ok** button.
 	 
-	 ![](Materials_Color_Picker.gif)(convert:false)
+	[OBJECT:EmbeddedVideo]
+	[PARAMLITERAL:width]
+	640
+	[/PARAMLITERAL]
+	[PARAMLITERAL:height]
+	360
+	[/PARAMLITERAL]
+	[PARAMLITERAL:videoid]
+	yztmMqXMyFs
+	[/PARAMLITERAL]
+	[/OBJECT]
 
 1. Now connect the **Constant3Vector** to the **Base Color** input on the Main Material node. To do this, use the **Left Mouse Button** to click on the circle output that is on the right of theConstant3Vector. Then with the **Left Mouse Button** still held down, drag to the right until you see a line coming from the circle output. Connect this line to the Base Color input on the Main Material node.
 	
-	![](Materials_Connect_Nodes_To_Graph.gif)(convert:false)
+	[OBJECT:EmbeddedVideo]
+	[PARAMLITERAL:width]
+	640
+	[/PARAMLITERAL]
+	[PARAMLITERAL:height]
+	360
+	[/PARAMLITERAL]
+	[PARAMLITERAL:videoid]
+	o3ANot4oEJU
+	[/PARAMLITERAL]
+	[/OBJECT]
 
 	[region:note]
 	As you make changes to the Main Material Node the viewport in the Material Editor will update to reflect your changes.  However this viewport can take a few seconds to update with changes depending on the complexity of your Material. For making slight tweaks to the Material a [Material Instance](Engine/Rendering/Materials/MaterialInstances) should be used as it shows changes instantly. 
@@ -134,7 +174,17 @@ Setting up a Material to use the Main Material Node is something can be accompli
 
 1. Once the values of the **Constant**  Material Expression are set, connect the **Constant** Material Expression node with the value of 1.0 to the **Metallic** input on the Main Material Node. Then connect the **Constant** Material Expression node with the value of 0.5 to the Roughness input on the Main Material node.  
 	
-	![](Materials_Connect_&_Compile.gif)(convert:false)
+	[OBJECT:EmbeddedVideo]
+	[PARAMLITERAL:width]
+	640
+	[/PARAMLITERAL]
+	[PARAMLITERAL:height]
+	360
+	[/PARAMLITERAL]
+	[PARAMLITERAL:videoid]
+	FuL5Dl4rPy8
+	[/PARAMLITERAL]
+	[/OBJECT]
 
 1. At this point, you can now begin to adjust both the values of the **Constant** Material Expression and the **Constant3Vector**  Material Expression to change the look of the Material. You can also connect different Material Expressions and combinations of Material Expression into the other inputs. For example, see what happens when you try and input something into the Emissive or Normal inputs on the Main Material node. Finally make sure to press both the **Apply** and **Save** buttons.  
 	
