#UE4doc #docs: ##Overview removed as according to the instructional guide.

diff --git a/Engine/Documentation/Source/Resources/Showcases/Reflections/Reflections.INT.udn b/Engine/Documentation/Source/Resources/Showcases/Reflections/Reflections.INT.udn
--- a/Engine/Documentation/Source/Resources/Showcases/Reflections/Reflections.INT.udn
+++ b/Engine/Documentation/Source/Resources/Showcases/Reflections/Reflections.INT.udn
@@ -61,10 +61,6 @@ Related: Resources/ContentExamples/Reflections/1_3
 
 [TOC(start:2 end:2)]
 
-
-
-## Overview
-
 This example shows off the real-time reflection capabilities of Unreal Engine 4. The scene shows a rundown subway terminal with a lot of water leaks, grimy ceramic tiles, dilapidated pipe work, and other worn-out environment details. In this document, we will give a high-level rundown of the techniques employed and the features involved in bringing the effect together.
 
 The following image shows the reflections in this scene blending on and off to demonstrate their effect.
