#UE4doc #docs: ##Overview removed as according to the instructional guide.

diff --git a/Engine/Documentation/Source/Shared/Editor/FbxErrors/FbxErrors.INT.udn b/Engine/Documentation/Source/Shared/Editor/FbxErrors/FbxErrors.INT.udn
--- a/Engine/Documentation/Source/Shared/Editor/FbxErrors/FbxErrors.INT.udn
+++ b/Engine/Documentation/Source/Shared/Editor/FbxErrors/FbxErrors.INT.udn
@@ -5,13 +5,10 @@ Description:Descriptions of errors generated during importing FBX files.
 
 [TOC(start:2 end:3)]
 
-
-## Overview
-
 This is a list of error or warning messages you may encounter during importing FBX files. 
 
 
-### Generic 
+## Generic 
 
 (#Generic_ImportingNewObjectFailed)
 [EXCERPT:Generic_ImportingNewObjectFailed]
