#UE4doc #docs:
Updated Engine version
Changed font settings on multiple terms
Added a note region where it was needed instead of leavingly a bolded statement as a note.

diff --git a/Engine/Documentation/Source/Programming/Rendering/ShaderDevelopment/ShaderDevelopment.INT.udn b/Engine/Documentation/Source/Programming/Rendering/ShaderDevelopment/ShaderDevelopment.INT.udn
--- a/Engine/Documentation/Source/Programming/Rendering/ShaderDevelopment/ShaderDevelopment.INT.udn
+++ b/Engine/Documentation/Source/Programming/Rendering/ShaderDevelopment/ShaderDevelopment.INT.udn
@@ -2,11 +2,11 @@ Availability:Public
 Title:Shader Development
 Crumbs:%ROOT%, Programming, Programming/Rendering
 Description:Information for graphics programmers writing shaders.
-Version: 4.5
+Version: 4.9
 
 ## Quick start
 
-**When working on shaders, be sure to enable r.ShaderDevelopmentMode by setting it to 1.**  The easiest way is to edit `ConsoleVariables.ini` so it is done every time you load.  This will enable retry-on-error and shader development related logs and warnings.
+When working on shaders, be sure to enable r.ShaderDevelopmentMode by setting it to 1.  The easiest way is to edit ConsoleVariables.ini so it is done every time you load.  This will enable retry-on-error and shader development related logs and warnings.
 
 Use **Ctrl+Shift+.** to recompile shaders that have changed, after you save changes to a .usf file.
 If you change a file that is included in many shaders (e.g. common.usf), this can take a while. If you want to iterate on a material, you can trigger recompile of the material by making a small change to the material (e.g. move node) and using the 'Apply' in the material editor.
@@ -21,7 +21,7 @@ Global shaders are shaders which operate on fixed geometry (like a full screen q
 Materials are defined by a set of states that control how the material is rendered (blend mode, two sided, etc) and a set of material inputs that control how the material interacts with the various rendering passes (BaseColor, Roughness, Normal, etc).
 #### Vertex Factories
 
-Materials have to support being applied to different mesh types, and this is accomplished with vertex factories. A **FVertexFactoryType** represents a unique mesh type, and a FVertexFactory instance stores the per-instance data to support that unique mesh type. For example, FGPUSkinVertexFactory stores the bone matrices needed for skinning, as well as references to the various vertex buffers that the GPU skin vertex factory shader code needs as input. The vertex factory shader code is an implicit interface which is used by the various pass shaders to abstract mesh type differences. Vertex factories consist of mainly vertex shader code, but some pixel shader code as well.  Some important components of the vertex factory shader code are:
+Materials have to support being applied to different mesh types, and this is accomplished with vertex factories. A `FVertexFactoryType` represents a unique mesh type, and a `FVertexFactory` instance stores the per-instance data to support that unique mesh type. For example, `FGPUSkinVertexFactory` stores the bone matrices needed for skinning, as well as references to the various vertex buffers that the GPU skin vertex factory shader code needs as input. The vertex factory shader code is an implicit interface which is used by the various pass shaders to abstract mesh type differences. Vertex factories consist of mainly vertex shader code, but some pixel shader code as well.  Some important components of the vertex factory shader code are:
 
 |Function| Description|
 |--|--|
@@ -34,12 +34,12 @@ Materials have to support being applied to different mesh types, and this is acc
 
 #### Material Shaders
 
-Shaders using **FMaterialShaderType** are pass specific shaders which need access to some of the material's attributes, and therefore must be compiled for each material, but do not need to access any mesh attributes. The light function pass shaders are an example of FMaterialShaderTypes.
+Shaders using `FMaterialShaderType` are pass specific shaders which need access to some of the material's attributes, and therefore must be compiled for each material, but do not need to access any mesh attributes. The light function pass shaders are an example of `FMaterialShaderType`'s.
 
 
-Shaders using **FMeshMaterialShaderType** are pass specific shaders which depend on the material's attributes AND the mesh type, and therefore must be compiled for each material/vertex factory combination. For example, TBasePassVS / TBasePassPS need to evaluate all of the material inputs in a forward rendering pass.
+Shaders using `FMeshMaterialShaderType` are pass specific shaders which depend on the material's attributes AND the mesh type, and therefore must be compiled for each material/vertex factory combination. For example, `TBasePassVS` / `TBasePassPS` need to evaluate all of the material inputs in a forward rendering pass.
 
-A material's set of required shaders is contained in a FMaterialShaderMap.  It looks like this:
+A material's set of required shaders is contained in a `FMaterialShaderMap`.  It looks like this:
 
     FMaterialShaderMap 
         FLightFunctionPixelShader - FMaterialShaderType
@@ -52,7 +52,7 @@ A material's set of required shaders is contained in a FMaterialShaderMap.  It l
         FGPUSkinVertexFactory - FVertexFactoryType
             Etc
 
-Vertex factories are included in this matrix based on their ShouldCache function, which depends on the material's usage. For example, bUsedWithSkeletalMesh being `true` will include the GPU skin vertex factories. FMeshMaterialShaderType's are included in this matrix based on their ShouldCache function, which depends on material and vertex factory attributes. This is a sparse matrix approach to caching shaders and it adds up to a large number of shaders pretty quick which takes up memory and increases compile times. The major advantage over storing a list of actually needed shaders is that no list has to be generated, so needed shaders have always already been compiled before run time on consoles. UE4 mitigates the shader memory problem with shader compression, and the compile time problem with multicore shader compilation.
+Vertex factories are included in this matrix based on their **ShouldCache** function, which depends on the material's usage. For example, bUsedWithSkeletalMesh being `true` will include the GPU skin vertex factories. `FMeshMaterialShaderType`'s are included in this matrix based on their ShouldCache function, which depends on material and vertex factory attributes. This is a sparse matrix approach to caching shaders and it adds up to a large number of shaders pretty quick which takes up memory and increases compile times. The major advantage over storing a list of actually needed shaders is that no list has to be generated, so needed shaders have always already been compiled before run time on consoles. UE4 mitigates the shader memory problem with shader compression, and the compile time problem with multicore shader compilation.
 #### Creating a material shader
 
 A material shader type is created with the DECLARE_SHADER_TYPE macro:
@@ -73,10 +73,10 @@ IMPLEMENT_MATERIAL_SHADER_TYPE(,FLightFunctionPixelShader,TEXT("LightFunctionPix
 
 This generates the material shader type's global metadata, which allows us to do things like iterate through all shaders using a given shader type at runtime.
 
-A typical material pixel shader type will first create a FMaterialPixelParameters struct by calling the GetMaterialPixelParameters vertex factory function. GetMaterialPixelParameters transforms the vertex factory specific inputs into properties like WorldPosition, TangentNormal, etc that any pass might want to access. Then a material shader will call CalcMaterialParameters, which writes out the rest of the members of FMaterialPixelParameters, after which FMaterialPixelParameters is fully initialized. The material shader will then access some of the material's inputs through functions in MaterialTemplate.usf (GetMaterialEmissive for the material's emissive input for example), do some shading and output a final color for that pass.
+A typical material pixel shader type will first create a `FMaterialPixelParameters` struct by calling the **GetMaterialPixelParameters** vertex factory function. **GetMaterialPixelParameters** transforms the vertex factory specific inputs into properties like WorldPosition, TangentNormal, etc that any pass might want to access. Then a material shader will call **CalcMaterialParameters**, which writes out the rest of the members of `FMaterialPixelParameters`, after which `FMaterialPixelParameters` is fully initialized. The material shader will then access some of the material's inputs through functions in MaterialTemplate.usf (**GetMaterialEmissive** for the material's emissive input for example), do some shading and output a final color for that pass.
 #### Special Engine Materials
 
-UMaterial has a setting called bUsedAsSpecialEngineMaterial that allows the material to be used with any vertex factory type. This means all vertex factories are compiled with the material, which will be a very large set. bUsedAsSpecialEngineMaterial is useful for:
+**UMaterial** has a setting called bUsedAsSpecialEngineMaterial that allows the material to be used with any vertex factory type. This means all vertex factories are compiled with the material, which will be a very large set. bUsedAsSpecialEngineMaterial is useful for:
 
  * Materials used with rendering viewmodes like lighting only.
  * Materials used as fallbacks when there is a compilation error (DefaultDecalMaterial, DefaultMaterial, etc).
@@ -103,7 +103,7 @@ If you want to step into the shader compiler DLL's directly from UE4 (CompileD3D
 
 With r.ShaderDevelopmentMode enabled, you will get the opportunity to retry on shader compile error.  This is especially important for global shaders since it is a fatal error if they do not compile successfully.
 
-In debug, with the debugger attached, you will hit a breakpoint and get the compile error in the Visual Studio output window.  **You can then double-click the error log to be taken directly to the offending line.**
+In debug, with the debugger attached, you will hit a breakpoint and get the compile error in the Visual Studio output window.  You can then **double-click** the error log to be taken directly to the offending line.
 
 ![](CompilerErrorDebug.png)
 
@@ -113,7 +113,7 @@ Otherwise you will get a Yes/No dialog
 
 ## Shader caching and cooking
 
-Once shaders are compiled, they are stored in the Derived Data Cache.  They contain in their key a hash of all the inputs to the compile, including shader source files.  That means that changes to shader source files are automatically picked up every time you re-launch the engine or do a 'recompileshaders changed'.
+Once shaders are compiled, they are stored in the Derived Data Cache.  They contain, in their key, a hash of all the inputs to the compile, including shader source files.  That means that changes to shader source files are automatically picked up every time you re-launch the engine or do a 'recompileshaders changed'.
 
 When you are modifying FShader Serialize functions, there is no need to handle backward compatibility, just add a space to a shader file that is included by that shader.
 
@@ -133,14 +133,14 @@ Then verify that the scale is right, and the result is not view dependent.  Howe
 
 #### Dumping debug info
 
-You can also use _r.DumpShaderDebugInfo=1_ to get files saved out to disk for all the shaders that get compiled.  It can be useful to set this in `ConsoleVariables.ini` just like r.ShaderDevelopmentMode.  Files are saved to _GameName/Saved/ShaderDebugInfo_, including
+You can also use r.DumpShaderDebugInfo=1 to get files saved out to disk for all the shaders that get compiled.  It can be useful to set this in ConsoleVariables.ini just like r.ShaderDevelopmentMode.  Files are saved to GameName/Saved/ShaderDebugInfo, including
 
  * Source files and includes
  * A preprocessed version of the shader
  * A batch file to compile the preprocessed version with equivalent commandline options to the compiler that were used
-
- **Note: If you leave this setting on, it can fill your HD with many tiny files and folders.**
-
+ [REGION:note]
+If you leave this setting on, it can fill your HD with many tiny files and folders.
+ [/REGION]
 ## Iteration best practices
 
 If you are working on a global shader, 'recompileshaders changed' or **Ctrl+Shift+.** is the fastest way to iterate.  If the shader takes a long time to compile, you might consider specifying CFLAG_StandardOptimization as a compile flag in the shader's ModifyCompilationEnvironment.
