UEDOC-1888

Gifs have been converted into YouTube Videos

diff --git a/Engine/Documentation/Source/Engine/Rendering/Materials/HowTo/Making_Functions/Making_Functions.INT.udn b/Engine/Documentation/Source/Engine/Rendering/Materials/HowTo/Making_Functions/Making_Functions.INT.udn
--- a/Engine/Documentation/Source/Engine/Rendering/Materials/HowTo/Making_Functions/Making_Functions.INT.udn
+++ b/Engine/Documentation/Source/Engine/Rendering/Materials/HowTo/Making_Functions/Making_Functions.INT.udn
@@ -33,14 +33,34 @@ Creating Material Functions can be accomplished in the few following steps.
 1. The first thing that we are going to need to do is to create a Material Function. 
 	This can be done within the **Content Browser** by **Right-Clicking** and choosing **New** > **Materials & Textures** > **Material Function**.
 
-	![](Material_Functions_Create_MF.gif)(convert:false)
+	[OBJECT:EmbeddedVideo]
+	[PARAMLITERAL:width]
+	640
+	[/PARAMLITERAL]
+	[PARAMLITERAL:height]
+	360
+	[/PARAMLITERAL]
+	[PARAMLITERAL:videoid]
+	XRlQsxUtnv8
+	[/PARAMLITERAL]
+	[/OBJECT]
 
 1. Once the Material Function has been created it needs to be named. 
 	Your name should be as descriptive as possible so that when someone goes to look at your Material Function they can easily understand what it does. 
 	For this example, we are going to call our node Custom_Fresnel. 
 	You can re-name your Material Function by selecting it in the **Content Browser**, pressing **F2** on the keyboard, and then typing in the name that you want. 
 	
-	![](Material_Functions_Rename.gif)(convert:false)
+	[OBJECT:EmbeddedVideo]
+	[PARAMLITERAL:width]
+	640
+	[/PARAMLITERAL]
+	[PARAMLITERAL:height]
+	360
+	[/PARAMLITERAL]
+	[PARAMLITERAL:videoid]
+	X2n6FtkrcXk
+	[/PARAMLITERAL]
+	[/OBJECT]
 
 ## Building the Material Function
 
@@ -115,7 +135,17 @@ Creating Material Functions can be accomplished in the few following steps.
 
 1. Now with all the nodes we need added to the Material Graph, it is now time to wire everything together. First, we need to wire the output of the OneMinus node to the Base input of the Power node. Then we need to connect our first Function Input node, Fresnel Exponent, into the Exp input of the Power node. Once that has been completed, we need to then connect the output of the Power node to the A input of the Multiply node and then connect our second Function Input, Fresnel Intensity, into the B input of the Multiply node. Finally we need to take the output of the Multiply node and connect that to the input of the Output Results node.  
 	
-	![](Material_Functions_Making_Connections.gif)(convert:false)                                                                                                                                                                                                      
+	[OBJECT:EmbeddedVideo]
+	[PARAMLITERAL:width]
+	640
+	[/PARAMLITERAL]
+	[PARAMLITERAL:height]
+	360
+	[/PARAMLITERAL]
+	[PARAMLITERAL:videoid]
+	0re4sjHvwrk
+	[/PARAMLITERAL]
+	[/OBJECT]                                                                                                                                                                                                     
 
 1. Once the Material nodes have all been connected, we need to now compile the Material Function by pressing the **Apply** button on the main tool bar. If the Material Function compiled correctly, the **Apply** button will be grayed out. When completed, you should have something that looks like the following image.
 
@@ -126,11 +156,31 @@ Creating Material Functions can be accomplished in the few following steps.
 
 1. With our Material Function now finished, close the Material Function Editor and create a new Material in the **Content Browser**. To create a new Material, **Right-Click** in the **Content Browser** and then select **Material** from **Basic Assets List** and name your newly created Material, **Fresnel_Material**.
 
-	![](Material_Functions_Create_New_Material.gif)(convert:false)
+	[OBJECT:EmbeddedVideo]
+	[PARAMLITERAL:width]
+	640
+	[/PARAMLITERAL]
+	[PARAMLITERAL:height]
+	360
+	[/PARAMLITERAL]
+	[PARAMLITERAL:videoid]
+	Av7GVfpBARw
+	[/PARAMLITERAL]
+	[/OBJECT]
 
 1. Now with the **Fresnel_Material** created, open it up by **Double Clicking** on it with the **Left Mouse Button** in the **Content Browser**. With the Material now open, let us first add our newly created **Material Function** by dragging it from the **Content Browser** into our material.  
 
-	![](Material_Functions_Add_Material_Function.gif)(convert:false)
+	[OBJECT:EmbeddedVideo]
+	[PARAMLITERAL:width]
+	640
+	[/PARAMLITERAL]
+	[PARAMLITERAL:height]
+	360
+	[/PARAMLITERAL]
+	[PARAMLITERAL:videoid]
+	S0wjEz-cO2o
+	[/PARAMLITERAL]
+	[/OBJECT]
 
 1. For this example, we are going to add the following nodes to our Material so that we have a bit better looking preview. Remember to find the nodes listed below search for the nodes' name in the **Palette** menu and then Drag and Drop the node from the **Palette** menu right into the Material Graph. When completed, you should have something that looks like this. 
 	
@@ -142,7 +192,17 @@ Creating Material Functions can be accomplished in the few following steps.
 
 1. Now that all the nodes have been added, it is now time to connect to the proper inputs and outputs. First, connect the Constant 3 Vector to the Base Color input on the Main Shader node. Then, connect one of the Constant One to the Metallic input on the Main Shader node. Then connect another of the Constant One to the Roughness input on the Main Shader node. Finally we are going to want to connect the final two Constant One nodes to the Fresnel Intensity and the Fresnel Exponent and then connect the output of the Custom_Fresnel node to the Emissive Color input on the Main Shader node. When completed, you should have something that looks like this.
 
-	![](Material_Functions_Connecting_Material_Function.gif)(convert:false)
+	[OBJECT:EmbeddedVideo]
+	[PARAMLITERAL:width]
+	640
+	[/PARAMLITERAL]
+	[PARAMLITERAL:height]
+	360
+	[/PARAMLITERAL]
+	[PARAMLITERAL:videoid]
+	g2yjwBwbWzw
+	[/PARAMLITERAL]
+	[/OBJECT]
 
 1. With the nodes now connected to the Main Material input, it is time to start adding values to the nodes. Fill each of the corresponding nodes in with the following values and when completed, make sure to **Compile your material**. When completed, you should have something that looks like this.
 
