#UE4doc #docs: ##Overview removed as according to the instructional guide.

diff --git a/Engine/Documentation/Source/Shared/Editor/MapErrors/MapErrors.INT.udn b/Engine/Documentation/Source/Shared/Editor/MapErrors/MapErrors.INT.udn
--- a/Engine/Documentation/Source/Shared/Editor/MapErrors/MapErrors.INT.udn
+++ b/Engine/Documentation/Source/Shared/Editor/MapErrors/MapErrors.INT.udn
@@ -5,9 +5,6 @@ Description:Descriptions of errors generated during map builds or running a map
 
 [TOC(start:2 end:3)]
 
-
-## Overview
-
 Unreal Engine 4 provides a mechanism for alerting you to potential problems in the form of map check errors and warnings instead of having to hunt down issues manually without knowing what is causing the problem or how to fix it. These errors are an automated way to find out where problems lie and allows you to get them fixed before they have a serious impact on your game. 
 
 
diff --git a/Engine/Documentation/Source/Shared/Glossary/C/Glossary_C.INT.udn b/Engine/Documentation/Source/Shared/Glossary/C/Glossary_C.INT.udn
--- a/Engine/Documentation/Source/Shared/Glossary/C/Glossary_C.INT.udn
+++ b/Engine/Documentation/Source/Shared/Glossary/C/Glossary_C.INT.udn
@@ -129,7 +129,7 @@ Components Mode
 In **Components Mode**, you can add components to your Blueprint, manage the details of existing components, and arrange components. The list of components in Components Mode is hierarchical.
 [/EXCERPT:Glossary_Components_Mode]
 
-**Components Mode** is one of three modes; the others are **Default Mode** and **Graph Mode**.
+Components Mode is one of three modes; the others are Default Mode and Graph Mode.
 	[/PARAM]
 [/OBJECT]
 
@@ -183,7 +183,7 @@ Controller Class
 	[PARAM:Definition]
 
 [EXCERPT:Glossary_Controller_Class]
-**Controllers** are non-physical Actors that can possess a Pawn (or Pawn-derived class like Character) to control its actions. A **PlayerController** is used by human players to control Pawns, while an **AIController** implements the artificial intelligence for the Pawns they control. Controllers take control of a Pawn with the **Possess** function, and give up control of the Pawn with the **Unpossess** function.
+**Controllers** are non-physical Actors that can possess a Pawn (or Pawn-derived class like Character) to control its actions. A PlayerController is used by human players to control Pawns, while an AIController implements the artificial intelligence for the Pawns they control. Controllers take control of a Pawn with the Possess function, and give up control of the Pawn with the Unpossess function.
 [/EXCERPT:Glossary_Controller_Class]
 
 	[/PARAM]
