#UE4doc #docs: ##Overview removed as according to the instructional guide.

diff --git a/Engine/Documentation/Source/Shared/Glossary/D/Glossary_D.INT.udn b/Engine/Documentation/Source/Shared/Glossary/D/Glossary_D.INT.udn
--- a/Engine/Documentation/Source/Shared/Glossary/D/Glossary_D.INT.udn
+++ b/Engine/Documentation/Source/Shared/Glossary/D/Glossary_D.INT.udn
@@ -37,11 +37,11 @@ Defaults Mode
 [EXCERPT:Glossary_Defaults_Mode]
 The **Class Defaults** or **Defaults** mode displays all of the properties of the _Blueprint_, including variables 
 contained within the _Blueprint_ and components defined in native code as well as replication, navigation, 
-and collision properties. The values set for these properties in the **Defaults** mode will be 
+and collision properties. The values set for these properties in the Defaults mode will be 
 used when new instances of the _Blueprint_ are created.
 [/EXCERPT:Glossary_Defaults_Mode]
 
-**Defaults Mode** is one of three Blueprint Editor modes; the others are **Components Mode** and **Graph Mode**.
+Defaults Mode is one of three Blueprint Editor modes; the others are Components Mode and Graph Mode.
 
 	[/PARAM]
 [/OBJECT]
@@ -53,7 +53,7 @@ Details panel
 	[PARAM:Definition]
 
 [EXCERPT:Glossary_Details_panel]
-The **Details panel** is a panel in the Level Editor that contains information about the properties for various elements. Many of the tools and editors in UE4 contain **Details** panels: the Level Editor, Matinee, Persona, and so on.
+The **Details panel** is a panel in the Level Editor that contains information about the properties for various elements. Many of the tools and editors in UE4 contain Details panels: the Level Editor, Matinee, Persona, and so on.
 [/EXCERPT:Glossary_Details_panel]
 
 	[/PARAM]
