#UE4doc #docs: Removed extraneous bolding as according to the instructional guide.

diff --git a/Engine/Documentation/Source/Shared/Glossary/G/Glossary_G.INT.udn b/Engine/Documentation/Source/Shared/Glossary/G/Glossary_G.INT.udn
--- a/Engine/Documentation/Source/Shared/Glossary/G/Glossary_G.INT.udn
+++ b/Engine/Documentation/Source/Shared/Glossary/G/Glossary_G.INT.udn
@@ -49,7 +49,7 @@ Graph
 The **Graph** is a visual representation of a particular graph of nodes as it shows 
 all of the nodes contained in the graph as well as the connections between them. It provides 
 editing capabilities for adding and removing nodes, arranging nodes, and creating links between 
-nodes. Breakpoints can also be set in the **Graph** to aid in [debugging](Engine/Blueprints\UserGuide\Debugging) _Blueprints_.
+nodes. Breakpoints can also be set in the Graph to aid in [debugging](Engine/Blueprints\UserGuide\Debugging) _Blueprints_.
 [/EXCERPT:Glossary_Graph_Mode]
 
 [INCLUDE:Globals#bpModesRemoved]
