#UE4doc #docs:
Changed font settings for multiple terms/lines
Expanded on multiple acronyms the first time they were mentioned.

diff --git a/Engine/Documentation/Source/Programming/Rendering/ShaderDevelopment/AsyncCompute/AsyncCompute.INT.udn b/Engine/Documentation/Source/Programming/Rendering/ShaderDevelopment/AsyncCompute/AsyncCompute.INT.udn
--- a/Engine/Documentation/Source/Programming/Rendering/ShaderDevelopment/AsyncCompute/AsyncCompute.INT.udn
+++ b/Engine/Documentation/Source/Programming/Rendering/ShaderDevelopment/AsyncCompute/AsyncCompute.INT.udn
@@ -6,8 +6,8 @@ Version: 4.9
 
 ## Introduction
 
-The RHI now supports asynchronous compute for Xbox One. This is a good way to utilize unused GPU resources (CUs, registers and bandwidth),
-by running dispatch() calls asynchronously with the rendering. Async compute uses a separate context, and we provide RHI functions to
+The Rendering Hardware Interface (RHI) now supports asynchronous compute (**AsyncCompute**) for Xbox One. This is a good way to utilize unused GPU resources (Compute Units (CUs), registers and bandwidth),
+by running **dispatch()** calls asynchronously with the rendering. Async compute uses a separate context, and we provide RHI functions to
 synchronize the rendering and compute contexts.
 Dr PIX is useful for identifying areas which might benefit from async compute. For example, if half the CUs are unused during a particular
 rendering pass, those CUs could potentially be utilized by an asynchronous compute job.
@@ -31,13 +31,13 @@ Async compute has some restrictions:
  * **RHIGraphicsWaitOnAsyncComputeJob**
    Inserts an command in the graphics immediate context to block until an async job is complete. Passing -1 causes this to early-out.
 
-Between calls to RHIBeginAsyncComputeJob_DrawThread and RHIEndAsyncComputeJob_DrawThread, the RHI will be in the async compute state.
+Between calls to **RHIBeginAsyncComputeJob_DrawThread** and **RHIEndAsyncComputeJob_DrawThread**, the RHI will be in the async compute state.
 During this time, supported RHI commands will be executed via the async compute context. Unsupported RHI functions will assert.
 
 ## Disabling/Enabling
 
 Async compute can be enabled or disabled at compile time with the #define USE_ASYNC_COMPUTE_CONTEXT. It can be disabled at runtime
-with the _r.AsyncCompute_ console variable. When async compute is disabled, dispatches within async compute blocks are executed synchronously on
+with the r.AsyncCompute console variable. When async compute is disabled, dispatches within async compute blocks are executed synchronously on
 the graphics command buffer. USE_ASYNC_COMPUTE_CONTEXT is defined to 0 on PC, since it's not supported in D3D11.1.
 
 ## PIX
