#UE4doc #docs: Altered images to match the instructional guide: green -> yellow.

diff --git a/Engine/Documentation/Source/Support/Builds/ReleaseNotes/2013/August/ReleaseNotesAugust2013.INT.udn b/Engine/Documentation/Source/Support/Builds/ReleaseNotes/2013/August/ReleaseNotesAugust2013.INT.udn
--- a/Engine/Documentation/Source/Support/Builds/ReleaseNotes/2013/August/ReleaseNotesAugust2013.INT.udn
+++ b/Engine/Documentation/Source/Support/Builds/ReleaseNotes/2013/August/ReleaseNotesAugust2013.INT.udn
@@ -54,7 +54,7 @@ Template:ReleaseNotes.html
 	![](safe1.png)(w:420)
 	* Normal selection still works with LMB, and click the RMB to bring the UI right back. 
 	![](safe2.png)(w:420)  ![](safe3.png)(w:420)
-		* You can even combine the HideUI function with Game View to completely clear the viewport, as seen above�
+		* You can even combine the HideUI function with Game View to completely clear the viewport, as seen above
 	* Safe Frames: If a camera actor is selected and the following criteria are met, the editor viewports now render safe frames based on the camera's aspect ratio setting:
 		* Game View must be off.
 		* The Safe Frames Show flag is on (**Show > Advanced**).
@@ -433,7 +433,7 @@ Template:ReleaseNotes.html
 * Added DLL export to various skeletal control functions.
 * Removed dead UJsonObject class. 
 * Got rid of deprecated USkeletalControl parallel hierarchy. 
-* Renamed box query parameter from Extents to HalfExtents, indicating that its the box radius instead of bounds. 
+* Renamed box query parameter from Extents to HalfExtents, indicating that it is the box radius instead of bounds. 
 * Added support for arbitrary nesting of Interp properties (e.g., Component.Struct.Prop), and redirection of Matinee track names (in the INI map MatineeTrackRedirects).
 * Removed a truncation from camera shake strength (legacy of integer FRotator). 
 
