UEDOC-1900

Converting gif files into YouTube Videos

diff --git a/Engine/Documentation/Source/Engine/Rendering/Materials/HowTo/Subsurface_Scattering/Subsurface_Scattering.INT.udn b/Engine/Documentation/Source/Engine/Rendering/Materials/HowTo/Subsurface_Scattering/Subsurface_Scattering.INT.udn
--- a/Engine/Documentation/Source/Engine/Rendering/Materials/HowTo/Subsurface_Scattering/Subsurface_Scattering.INT.udn
+++ b/Engine/Documentation/Source/Engine/Rendering/Materials/HowTo/Subsurface_Scattering/Subsurface_Scattering.INT.udn
@@ -29,7 +29,17 @@ When completed, your **Content Browser** should look like this.
 
 1. Now with the Material open, we need to change the Material's **Shading Model** from **Default Lit** to **Subsurface**. 
 
-	![](SSM_Setting_Shading_Model.gif)(convert:false)
+	[OBJECT:EmbeddedVideo]
+	[PARAMLITERAL:width]
+	640
+	[/PARAMLITERAL]
+	[PARAMLITERAL:height]
+	360
+	[/PARAMLITERAL]
+	[PARAMLITERAL:videoid]
+	_SnzU14R4Y4
+	[/PARAMLITERAL]
+	[/OBJECT]
 
 1. The Material is now ready to be used as a Subsurface Material. 
 	
@@ -65,20 +75,50 @@ Continuing with our Material from above that has Subsurface Scattering enabled,
 	Keep dragging to the Right with the  **Left Mouse Button** still held until you come to the input you want to use on the Main Material node. 
 	Once your mouse is over the input you want to use, release the  **Left Mouse Button** to connect.
 
-	![](SSM_Connecting_Nodes.gif)(convert:false)
+	[OBJECT:EmbeddedVideo]
+	[PARAMLITERAL:width]
+	640
+	[/PARAMLITERAL]
+	[PARAMLITERAL:height]
+	360
+	[/PARAMLITERAL]
+	[PARAMLITERAL:videoid]
+	44PbkN_uoP4
+	[/PARAMLITERAL]
+	[/OBJECT]
 
 1. Now that all the nodes have been connected, make sure that you compile the Material by pressing the **Apply** button. Once compiled, you should have something that looks like this. 
 	![](SSM_Finshed_Material.png)
 
 1. Once the Material has been compiled, you can now close the Material window. Then inside of the **Content Browser**, select the Material, **Right-Click** on it, and then select the **Create Material Instance** option from the menu.  
-	![](SSM_Create_New_Material_Instance.gif)(convert:false)
+	[OBJECT:EmbeddedVideo]
+	[PARAMLITERAL:width]
+	640
+	[/PARAMLITERAL]
+	[PARAMLITERAL:height]
+	360
+	[/PARAMLITERAL]
+	[PARAMLITERAL:videoid]
+	akSnwKBLMtw
+	[/PARAMLITERAL]
+	[/OBJECT]
 
 
 ## Applying a Subsurface Material
 1. Now that the Material Instance has been created, it is now time to test out our new Material. 
 	To do this, we first need to create a new blank level to work in by going to the main menu and under the **File** option select **New Level**. 
 	When prompted what type of level to select, select the **Empty Level**. 
-	![](SSM_Create_New_Level.gif)(convert:false)
+	[OBJECT:EmbeddedVideo]
+	[PARAMLITERAL:width]
+	640
+	[/PARAMLITERAL]
+	[PARAMLITERAL:height]
+	360
+	[/PARAMLITERAL]
+	[PARAMLITERAL:videoid]
+	bRuKiOW0i-w
+	[/PARAMLITERAL]
+	[/OBJECT]
 
 1. With the new level created, we need to now add some lights so that we can test how our Subsurface Material will look. The first light that we are going to add is a **Point Light**. To add a Point Light to the scene, we first need to find the Point Light which is in the **Modes** menu under the **Lights** option. Once the Point Light is located, you can add one to the screen by first selecting it with the **Left Mouse Button**  and then with the **Left Mouse Button** still held down dragging the Point Light into the scene and releasing the **Left Mouse Button** to add it. Once added to the scene, select the light by clicking on it with the **Left Mouse Button**. Then enter the following values.
 	| Property Name| Value|
@@ -114,7 +154,17 @@ Continuing with our Material from above that has Subsurface Scattering enabled,
 1. Now that the Static Mesh has been placed in the level, we can now apply our Subsurface Material to it. 
 	To do this, first find the Subsurface Material in the **Content Browser**. 
 	Once the Material has been located, Drag it from the **Content Browser** on to the mesh to apply it. 
-	![](SSM_Applying_Material.gif)(convert:false)
+	[OBJECT:EmbeddedVideo]
+	[PARAMLITERAL:width]
+	640
+	[/PARAMLITERAL]
+	[PARAMLITERAL:height]
+	360
+	[/PARAMLITERAL]
+	[PARAMLITERAL:videoid]
+	j0xDDl7J8z0
+	[/PARAMLITERAL]
+	[/OBJECT]
 
 ## Working with the Subsurface Material
 Now that the Subsurface Material has been applied, it is time to tweak the settings of the Material Instance. 
@@ -140,7 +190,17 @@ If you do the opposite of this, adjusting  the **Value** slider in the **Color P
  ![](SSM_Color_Value_Example_01.png)
 
 Here is an example of adjusting the value in real time. Notice how as the Value of the color is adjusted from Red to Black, the amount of Subsurface scattering influence is affected.  
-![](SSM_Value_On_SSS_Influance.gif)(convert:false)
+[OBJECT:EmbeddedVideo]
+[PARAMLITERAL:width]
+640
+[/PARAMLITERAL]
+[PARAMLITERAL:height]
+360
+[/PARAMLITERAL]
+[PARAMLITERAL:videoid]
+yKmyUwXmSJE
+[/PARAMLITERAL]
+[/OBJECT]
 ## Using a Texture as a Subsurface Influence Mask
 You can use a Texture as a Mask to have more control over what areas receive or do not receive  Subsurface scattering. 
 To do this, all you need to do is plug the Texture you want to use as the Mask into the **Opacity** Channel of your Material. 
