UEDOC-1899

gif files have been replaced with YouTube videos

diff --git a/Engine/Documentation/Source/Engine/Rendering/Materials/HowTo/ShinyMaterials/ShinyMaterials.INT.udn b/Engine/Documentation/Source/Engine/Rendering/Materials/HowTo/ShinyMaterials/ShinyMaterials.INT.udn
--- a/Engine/Documentation/Source/Engine/Rendering/Materials/HowTo/ShinyMaterials/ShinyMaterials.INT.udn
+++ b/Engine/Documentation/Source/Engine/Rendering/Materials/HowTo/ShinyMaterials/ShinyMaterials.INT.udn
@@ -112,7 +112,17 @@ content page for information about how to move content between projects. This wa
 	![](Engine\Rendering\Materials\HowTo\Refraction\R_Apply_&_Save.png)
 
 1. Now find the **Shiny_Material** in the **Content Browser** and  **Right-Click** on it with the mouse and from the menu that is displayed, select the **Create Material Instance** option.  
-	![](Tras_Create_MI.gif)(convert:false)
+	[OBJECT:EmbeddedVideo]
+	[PARAMLITERAL:width]
+	640
+	[/PARAMLITERAL]
+	[PARAMLITERAL:height]
+	360
+	[/PARAMLITERAL]
+	[PARAMLITERAL:videoid]
+	115AqnPtp9s
+	[/PARAMLITERAL]
+	[/OBJECT]
 
 1. Open the Material Instance up by **Double-Clicking** on it with the **Left Mouse Button** in the **Content Browser**. Once open enable all the **Scalar Parameter Values** by clicking on the grayed out check mark box with the **Left Mouse Button**. 
 	![](Shiny_Enabled_MI_Properties.png)
