#UE4doc #docs: Altered images to match the instructional guide: green -> yellow. Removed broken characters.

diff --git a/Engine/Documentation/Source/Support/Builds/ReleaseNotes/2013/Dec/ReleaseNotesDecember2013.INT.udn b/Engine/Documentation/Source/Support/Builds/ReleaseNotes/2013/Dec/ReleaseNotesDecember2013.INT.udn
--- a/Engine/Documentation/Source/Support/Builds/ReleaseNotes/2013/Dec/ReleaseNotesDecember2013.INT.udn
+++ b/Engine/Documentation/Source/Support/Builds/ReleaseNotes/2013/Dec/ReleaseNotesDecember2013.INT.udn
@@ -230,7 +230,7 @@ QA_APPROVED_UE4_BUILD_MINIMUM
 * Added detailed path finding debug in ANavigationTestingActor.
 * Added a generic mechanism for locking AI resources (like movement and brain) indicating source of lock. See IAIResourceInterface for details.
 * Animation Retargeting Improvement: Skeleton supports multiple base poses for retargeting.
-	* Base pose is different reference poses of the skeleton � i.e. from different size of meshes.
+	* Base pose is different reference poses of the skeleton i.e. from different size of meshes.
 	* When retargeting is used in your system, some animations need to have correct base pose (based on what mesh it was animated from) in order to retarget correctly.  
 
 
@@ -368,7 +368,7 @@ QA_APPROVED_UE4_BUILD_MINIMUM
 * Macros can now be dragged from My Blueprint into a Graph.
 * Added a call to reset the MyBlueprint search filter before adding any new element, so the new element is visible once created.
 * The Vector2D blueprint graph actions were moved into the Vector2D category splitting it out from the standard vector class.
-* The Alpha section of the colour picker has been enabled on blueprint graph color pins.
+* The Alpha section of the color picker has been enabled on blueprint graph color pins.
 * The ability to create a comment around selection has been added to the node context menu in the editors.
 * The Details view for the socket component was updated so the category is placed correctly and is always visible in the blueprint editor.
 * The ability to locate variable uses for any variable based node was added to the context menu in the blueprint editor.
@@ -550,7 +550,7 @@ QA_APPROVED_UE4_BUILD_MINIMUM
 * Added check to New Project wizard to make sure that the folder path is validly formatted for the platform.
 * Material Selection: Filtering broko when typing quickly so removed the bShouldNotifyNextAssetSync flag which seemed to exist only to cause it to automatically select an asset if it was still loading others.
 * The Asset Saving notification when clicking "Choose Files to Save" looked too much like an error, changed to Notification, and modified the text for clarity.
-* It was difficult to remove a level from Favorites. Entries in the favourite list could only be added/removed when the level list was open, removed that requirement and moved open/removing to a sub-menu, also fixed inconsistent internal entry naming bug.
+* It was difficult to remove a level from Favorites. Entries in the favorite list could only be added/removed when the level list was open, removed that requirement and moved open/removing to a sub-menu, also fixed inconsistent internal entry naming bug.
 * Normal Map Textures automatically reimported were corrupted due to combined use of Defer Compression & compatibility mode DXT5n normal maps, added swizzled XGXR8 mode for uncompressed deferred normal maps.
 * Window Management: On creating a new tab, or otherwise updating it by selecting another tab under a window, change the title to that tabs label unless it's the master/root window which always retains the "map - Unreal Editor" naming scheme.
 * Add an optional Cancel button to the auto-saving progress widget.
@@ -664,7 +664,7 @@ QA_APPROVED_UE4_BUILD_MINIMUM
 	* There have been several improvements and bug fixes to Oculus Rift support, including fixes to certain post process and screen-space effects that previously did not render when in stereoscopic view.
 * All new game template projects now have a simplified source file layout.
 * Moving Platformergame into Samples as a code-based SampleGames directory test.
-	* We'll make an effort to move all code-based example games from the main directory down into a new Games subdirectory under Samples. Games will be the current SampleGames folder � renamed. (e.g. CURRENTLY  //depot/UE4/PlatformerGame/ ? //depot/UE4/Samples/SampleGames/PlatformerGame/ )  (The final location being something like //depot/UE4/Samples/Games/PlatformerGame/)
+	* We'll make an effort to move all code-based example games from the main directory down into a new Games subdirectory under Samples. Games will be the current SampleGames folder renamed. (e.g. CURRENTLY  //depot/UE4/PlatformerGame/ ? //depot/UE4/Samples/SampleGames/PlatformerGame/ )  (The final location being something like //depot/UE4/Samples/Games/PlatformerGame/)
 * Trying to return a value from a RPC Client/Server function is now detected as a compile error.
 * **Component movement optimizations**
 	* Repeated calls to MoveComponent can defer most of the work until the sequence of moves is committed, by encapsulating the MoveComponent calls within a FScopedMovementUpdate() object.
@@ -690,7 +690,7 @@ QA_APPROVED_UE4_BUILD_MINIMUM
 * Added a new menu entry to manage notifies to allow deletion of custom notifies.
 * Behavior Trees core and tools are still WIP.
 * Animation System Improvement
-	* Physics Animation: Added BlueprintCallable Physics Weight functions � SetPhysicsWeight/AccumulatePhysicsWeight
+	* Physics Animation: Added BlueprintCallable Physics Weight functions SetPhysicsWeight/AccumulatePhysicsWeight
 	* BlendSpace Node allows to set Position/Looping via Pin.
 	* IK node displays effector/joint location.
 * Physics: Default setting has been tweaked to improve ragdoll.
