UEDOC-1889

gifs have been replaced with Youtube videos

diff --git a/Engine/Documentation/Source/Engine/Landscape/Creation/Creation.INT.udn b/Engine/Documentation/Source/Engine/Landscape/Creation/Creation.INT.udn
--- a/Engine/Documentation/Source/Engine/Landscape/Creation/Creation.INT.udn
+++ b/Engine/Documentation/Source/Engine/Landscape/Creation/Creation.INT.udn
@@ -113,7 +113,17 @@ In the main viewport, there should be a preview of the new Landscape you are set
 
 [REGION:tip]
 You can Move, Rotate, and Scale the Landscape Actor like any normal Actor. You can also drag the edges of the preview Landscape to easily size the Landscape to fit a space.
-![](Landscape_Move_In_Viewport.gif)(convert:false)
+[OBJECT:EmbeddedVideo]
+[PARAMLITERAL:width]
+640
+[/PARAMLITERAL]
+[PARAMLITERAL:height]
+360
+[/PARAMLITERAL]
+[PARAMLITERAL:videoid]
+og7fnKaqYdI
+[/PARAMLITERAL]
+[/OBJECT]
 [/REGION]
 
 
