Updates for Jira UEDOC-1943

Updated this whole section:
Updated Main Using Cameras page for images/minor content fixes.
New Page added for Using Static Cameras.
New Page added for Using Multiple Static Cameras and switching.
New Page added for Using Camera Components.
New Page added for Using Spring Arm Components.

diff --git a/Engine/Documentation/Source/Gameplay/HowTo/UsingCameras/Blueprints/UsingCameras_BP.INT.udn b/Engine/Documentation/Source/Gameplay/HowTo/UsingCameras/Blueprints/UsingCameras_BP.INT.udn
--- a/Engine/Documentation/Source/Gameplay/HowTo/UsingCameras/Blueprints/UsingCameras_BP.INT.udn
+++ b/Engine/Documentation/Source/Gameplay/HowTo/UsingCameras/Blueprints/UsingCameras_BP.INT.udn
@@ -1,55 +1,59 @@
 Availability: Public
-Title:Using Cameras in Blueprints
+Title:Using a Static Camera
 Crumbs: %ROOT%, Gameplay, Gameplay\HowTo\UsingCameras
 Description:A How To Guide for Finding Actors in Blueprints. 
+Related: Gameplay\HowTo\UsingCameras
+Related: Gameplay\Framework\Camera
+Related: Engine/Matinee
+Related: Engine/Blueprints
+SkillLevel: Beginner
+Version: 4.8
 
-[TOC (start:2 end:2)]
-
-This page covers several examples on how to use **Cameras** within Blueprints for gameplay purposes. Whether you are looking to create a fixed camera angle for your game or cycle between multiple cameras, in order for the camera to be used during gameplay, you will need to know how to set them up first in Blueprints. 
+In this How-to we take a look at creating a static (or fixed) camera angle that is used for a player's perspective during gameplay. Once completing this tutorial, you can take the process used here and apply it to your own game to set up a fixed perspective for a player to view your game from. 
 
 [REGION:note]
-For information on using Cameras for cinematic purposes, refer to [](Engine/Matinee/HowTo) documentation and the Camera sections. The [](Resources/Showcases/MatineeFightScene) is also an excellent resource for a sample cutscene and how it was constructed.
+ For this example, we are using the **Blueprint Third Person Template** project with **Starter Content** included.
 [/REGION]
 
- _For these examples, we are using the **Blueprint Third Person Template with Starter Content**._
-
-## Using a Static Camera
-
-This section covers how to place a **Camera** Actor in your level and assign it for use as the default Camera for your game. You can also use the method to switch to a static Camera during gameplay (the next section will cover how to switch between multiple cameras). 
-
-1.	From the **Modes** menu, under the **Basic** tab, drag a **Camera** Actor into your level. 
+1.	From the **Modes** menu while in [](Engine/UI/LevelEditor/Modes/PlaceMode), search for **Camera** then drag a Camera into the level. 
 
 	![](camera_bp1.png)
 
-	Place your Camera in the level where you want the player to view the game from. 
+	You can place the Camera Actor where you would like the view perspective to be for your players. To get a better idea of the view the player will have, you can take a First Person perspective from the Camera's Point of View by locking the Viewport to the Camera Actor and entering **Pilot** mode. 
 	
-	To get a better idea of the view the player will have, you can take a First Person perspective from the Camera's Point of View by locking the Viewport to the Camera Actor.	With the Camera selected, click the **Viewport Options** button and select **Lock Viewport to Actor** then **Camera Actor**. 
+1.	With the Camera selected in the level, **Right-click** on it, then from the context menu select **Pilot 'CameraActor'**. 
 
 	[REGION:raw]
 	![](camera_bp2.png)
 	[/REGION]
 
-	You can now move around the Viewport using the WASD keys and the Camera's position will move along with it. While the Camera is locked to the Viewport, you will notice that a lock icon appears in the upper left corner of the Viewport. 
+	You can now move around the Viewport using the **WASD** keys while holding the **Left** or **Right Mouse Button** down. While you fly around the level, the Camera's position will move along with your movement allowing get an idea of the perspective the camera will take during gameplay. 
+
+1.	To unlock the Camera, click the **Unlock** button. 
 
 	![](camera_bp3.png)
 
-	To unlock the Camera, click the **Viewport Options** button again and select **Lock Viewport to Actor** then select the unlock option. 
+	The camera will remain at the resulting position when you select to unlock it. 
 
+	[REGION:note]
+	The icon next to Unlock allows you to toggle between showing the in-game camera view or level editor view. 
+	[/REGION]
+	
 1.	With the **Camera Actor** selected in the level, click the **Blueprints** button from the Main Toolbar then select **Open Level Blueprint**. 
 
 	![](camera_bp4.png)
 
-1.	Inside the **Level Blueprint**, **Right-click** anywhere in the graph window and select **Add Reference to CameraActor**. 
+1.	Inside the **Level Blueprint**, **Right-click** anywhere in the graph window and select **Create a Reference to CameraActor**. 
 
 	![](camera_bp5.png)
 
-	By adding a reference to the **Camera Actor**, the **Level Blueprint** now knows which camera we are talking about when we want to do something (in our case we want to set this Camera Actor to be the default camera view when the game is started). 
+	By adding a reference to the **Camera Actor**, the **Level Blueprint** now knows which camera we are talking about in the level (as we could have multiple cameras set up inside our level). 
 
 1.	**Right-click** in the graph again and search for and add the **Event Begin Play** node. 
 
 	![](camera_bp6.png)
 
-	This node only executes once and executes anything following it when the game is started, which is what we want. 
+	This node only executes once and executes anything following it when the game is started. 
 
 1.	**Right-click** in the graph again and search for **Set View Target** and select the **Set View Target with Blend** node. 
 
@@ -77,179 +81,23 @@ This section covers how to place a **Camera** Actor in your level and assign it
 
 1.	From the Main Editor View, click the **Play** button to play in the Editor. 
 
-You should now see that when the game starts, the player's view is assigned to and using the **Camera Actor** that you have placed in your level. You may have also noticed that the view is letterboxed; you can easily disable this by un-checking the **Constrain Aspect Ratio** option inside the **Details** panel for the **Camera Actor**. 
-
-![](camera_bp12.png)
-
-
-## Switching Between Multiple Cameras
-
-This section shows an example of how you can implement multiple camera angles in your game by placing multiple **Camera Actors** then use **TriggerBoxes** to tell the **Level Blueprint** when to switch between each camera. 
-
-1.	From the **Modes** menu, under **Basic**, drag in the number of **Camera Actors** you wish to cycle between. 
-
-	![](camera_bp13.png)(w:800)
-	
-	Here we have placed three cameras using different angles. 
-
-1.	Also in the **Modes** menu, under **Basic**, drag in the number of **Box Triggers** to match the number of cameras you placed. 
-
-	![](camera_bp14.png)(w:800)
-
-	Move the **Box Triggers** to where you want the player to trigger the changing of cameras. You can move/re-size the **Box Triggers** using the **W** key to move it around and the **R** key to scale the box size. Here we have set a trigger at the base of the stairs, one at the top of the stairs, and one when the player walks around the platform.
-
-1.	With each of your cameras selected, click the **Blueprints** button then select **Open Level Blueprint**. 
-
-	![](camera_bp4.png)
-
-1.	Inside the **Level Blueprint**, **Right-click** in the graph and choose **Add references to selected actors**. 
-
-	![](camera_bp15.png)
-
-	This will add references to each of the Camera Actors you have placed in your level. 
+## End Result
 
-1.	Minimize the **Level Blueprint**, and in your level, select each of your **Box Triggers** then return to the **Level Blueprint**. 
+[OBJECT:EmbeddedVideo]
+[PARAMLITERAL:width]
+640
+[/PARAMLITERAL]
+[PARAMLITERAL:height]
+360
+[/PARAMLITERAL]
+[PARAMLITERAL:videoid]
+mi4YXtSVyPc
+[/PARAMLITERAL]
+[/OBJECT]
 
-1.	**Right-click** in the graph, then under **Add Event for Selected Actors** and **Collision**, select **Add OnActorBeginOverlap**. 
-
-	![](camera_bp16.png)
-
-	This will add **Events** for each **Box Trigger** that will fire off when the player enters them. 
-
-1.	**Right-click** in the graph and add the **Get Player Controller** node. 
-
-	![](camera_bp17.png)
-
-1.	**Right-click** in the graph and add the **Set View Target with Blend** node (un-check **Context Sensitive** if you do not see it). 
-
-1.	Copy the **Set View Target with Blend** node twice, then connect the nodes as shown below. 
-
-	![](camera_bp18.png)
-
-1.	Click **Compile** then close the Blueprint and click **Play** to play in the Editor. 
-
-You should now see that when you enter each of the Box Trigger's you have placed, it updates the game view to the Camera Actor you have specified. This a quick example of how you could start working on a multi-camera based game, you would probably want to also account for situations when the player is outside of your camera's viewing angle and adjust player input based on camera relativity. 	
-
-## Using a Camera Component
-
-This section will show you how to add a **Camera Component** to a **Blueprint** which can be used to create third person perspectives for playable or non-playable Actors in your levels. 
-
-1.	In the **Content Browser**, click the **New** button then click the **Blueprint** option. 
-
-	[REGION:raw]
-	![](Engine\Blueprints\UserGuide\Types\ClassBlueprint\Creation\CreateBP_mainFlow.png)
-	[/REGION]
-	
-1.	Click the **Character** option when the ***Class** window appears, then give the Blueprint a name and open it. 
-
-1.	On the **Components** panel, click the **Add Component** button then select **Camera**. 
-
-	[REGION:raw]
-	![](component3.png)
-	[/REGION]
-
-	This will add a **Camera Component** to the Blueprint Actor. 
-
-1.	When the Camera is added, it will appear inside the Character Movement Capsule which we do not want. 
-
-	![](component4.png)
-
-	[REGION:note]
-	The **Viewport** is docked to the Graph area of the Blueprint Editor. Switch to it by clicking its tab so you can see where your Components are located.
-	[/REGION]
-
-1.	Using the **Translation** widget, move the Camera behind the Capsule to a desired third person perspective. 
-
-	![](component5.png)
-
-	You can angle the Camera by entering **Rotation** mode (press **E**) to get your desired angle. 
-
-1.	Click the **Add Component** button again, and select **Static Mesh**. 
-
-1.	In the **Details** panel for the **Static Mesh**, click the drop-down menu under **Static Mesh** and select a mesh to use. 
-
-	![](component11.png)
-	
-	We have selected a Capsule to use as our Character for this example. 
-
-	![](component6.png)
-
-1.	Minimize the Blueprint, then under the **Game/Blueprints** folder, in the **Content Browser**, open the **MyCharacter** Blueprint. 
-
-1.	Copy the Mouse, Movement, and Jump Input sections (we will paste these in the Character Blueprint we just created). 
-
-	![](component7.png)
-
-1.	Back in your new Character Blueprint, paste the contents that were copied from the **MyCharacter** Blueprint. 
-
-1.	Click the **Compile** button, then close the Blueprint. 
-
-1.	In the **Game/Blueprints** folder, open the **MyGame** Blueprint and assign your new Character to the **Default Pawn Class**. 
-
-	![](component8.png)
-
-1.	Click the **Compile** button, then close the Blueprint and from the Main Editor window and click the **Play** button. 
-
-When you play in the Editor, you should have a character using the Static Mesh you specified that can move around the level and the Camera Component we assigned will follow that Character around. 
-
-All of the Camera Options that exist on a **Camera Actor** can be accessed and set via the **Camera** Component inside the Blueprint we created above (for example if you wanted to add tint or color shifts when the player is damaged or receives a power up), you can also set the location of the Camera Component to create a method by which a player can toggle between first or third person or even an "over-the-shoulder" type of Camera as seen below. 
-
-![](component9.png)
-
-In the above image, we placed two **Arrow** Components to act as markers to where we want our **Camera** Component to move to. In the **Graph** tab of the Blueprint, we then can add a script similar to below to toggle between the two camera locations. 
-
- ![](component10.png)
-
- Here we are saying when **Right Mouse Button** is pressed, set the location of **Camera1** (which is our **Camera** Component) to the location of our **Arrow2** (which is an over-the-shoulder) view. When releasing **Right Mouse Button**, return the camera to its default location (denoted by the location of **Arrow1**). 
-
- 
-## Using a Spring Arm Component
-
-When using a **Camera Component** to create a third person perspective for a character, you will typically want to also include a **Spring Arm Component** and parent it to the Camera. The **Spring Arm** Component tries to maintain its children at a fixed distance from the parent, but will retract the children if there is a collision, and will spring back when there is no collision. 
-
-What this does is prevent or minimize instances when the camera can penetrate solid objects, for example backing up against a wall as shown in the image below where the camera penetrates the wall next to the Character allowing them to view outside the level's play space. When using a **Spring Arm**, the camera will retract to avoid penetrating the wall and will spring back to its normal position when moving away from the wall. 
-
-![](springArm1.png)
-
-The steps below will show you how to add a **Spring Arm** Component to your Camera setup. 
-
-1.	Continuing from the **Using a Camera Component** section, open your **Character Blueprint**. 
-
-1.	On the **Components** tab, click the **Add Component** button then select the **Spring Arm** Component. 
-
-	[REGION:raw]
-	![](springArm2.png)
-	[/REGION]
-
-1.	Inside the **Components** window, drag your **Camera** Component onto the **Spring Arm** Component to attach it to the Spring Arm. 
-
-	[REGION:raw]
-	![](springArm3.png)
-	[/REGION]
-
-1.	There are multiple settings that can be used with the **Spring Arm** which are outlined below. 
-
-	![](springArm4.png)
-
-	| Section | Description |
-	| ------- | ------------ |
-	| **Sockets** | This allows you to parent the Spring Arm to a Bone on a Skeletal Mesh or a Socket. |
-	| **Events** | Allows you to add Events for the Spring Arm or to apply Component Tags. |
-	| **Camera** | Here you can set the length of the Spring Arm when there are no collisions as well as set Offsets. |
-	| **Camera Collision** | Set whether or not to check for collisions and define how big the query probe is to check for collisions. |
-	| **Camera Settings** | Whether or not to inherit Pitch, Yaw, or Roll from the Parent or Use Pawn Control Rotation. |
-	| **Lag** | Set whether or not the camera should lag behind the target position and settings associated with lag. |
-	
-1.	Once you are content with your settings (we have left everything as default), click **Compile** then close the Blueprint. 
-
-1.	Click the **Play** button from the Main Tool Bar to play in the Editor. 
-
-Move the character near a wall or position the character so that an object comes between it and the camera and notice how the camera will automatically adjust and move so that it does not intersect the object in the way.
-
-[REGION:note]
-For more on Cameras in Unreal Engine 4, refer to the [](Gameplay\Framework\Camera) documentation in the **Framework** section. 
-[/REGION]
+You will now see that when the game starts, the player's view is assigned to and using the **Camera Actor** that you have placed in your level. You may have also noticed that the view is letterboxed; you can disable this by un-checking the **Constrain Aspect Ratio** option inside the **Details** panel for the **Camera Actor**. 
 
+![](camera_bp12.png)
 
 
 
diff --git a/Engine/Documentation/Source/Gameplay/HowTo/UsingCameras/CameraComponents/CameraComponents_BP.INT.udn b/Engine/Documentation/Source/Gameplay/HowTo/UsingCameras/CameraComponents/CameraComponents_BP.INT.udn
new file mode 100644
--- /dev/null
+++ b/Engine/Documentation/Source/Gameplay/HowTo/UsingCameras/CameraComponents/CameraComponents_BP.INT.udn
@@ -0,0 +1,162 @@
+Availability: Public
+Title:Working with Camera Components
+Crumbs: %ROOT%, Gameplay, Gameplay\HowTo\UsingCameras
+Description:Illustrates how to add a Camera Component to a Blueprint which can be used as a camera persepctive for a object. 
+Related: Gameplay\HowTo\UsingCameras
+Related: Gameplay\HowTo\UsingCameras\SpringArmComponents
+Related: Gameplay\Framework\Camera
+Related: Engine/Matinee
+Related: Engine/Blueprints
+SkillLevel: Beginner
+Version: 4.8
+
+In this How-to we take a look at adding a **Camera** Component to a **Blueprint** which can be used to create third person perspectives for playable or non-playable characters in your levels. By the end of this tutorial you will have gone through the process of adding a Camera Component to a Character Blueprint and set up the ability to toggle between two different perspectives for a character. 
+
+[REGION:note]
+ For this example, we are using the **Blueprint Third Person Template** project with **Starter Content** included.
+[/REGION]
+
+1.	With your project open, select the character in the level and press **Delete** to remove it. 
+
+	![](Camera1.png)(w:720)
+
+	This is the default Character Blueprint that comes with this template. Notice that it has a **Camera** component attached to it already and you can see a preview window of the perspective that will be used when playing as this character. We are going to recreate this setup from scratch on a new type of playable character. 
+	
+1.	In the **Content Browser**, click the **Add New** button then select [](Engine/Blueprints/UserGuide/Types/ClassBlueprint) from the menu. 
+
+	![](Camera2.png)
+
+1.	In the **Pick Parent Class** window, select **Character** then name the Blueprint **MyCharacter**. 
+
+	![](Camera3.png)
+
+	After you choose a save location and a name for the new Blueprint, the [Blueprint Editor Interface](Engine/Blueprints/Editor/UIBreakdowns/ClassBPUI) will open automatically. 
+
+1.	In the **Components** window select the **CapsuleComponent** then in the **Details** panel, change the **Capsule Half Height** to **64**. 	
+
+	![](Camera4.png)
+
+	The [Capsule Component](Engine/Components/Shapes) is used to provide [Collision Detection](Engine/Physics/Collision) for the character. 
+
+1.	In the **Components** window, click the **Add Component** button and search for and add a **Static Mesh** component. 
+
+	![](Camera5.png)
+
+1.	In the **Details** panel for the **Static Mesh**, set its **Z** location to **-60** and assign the **Shape_NarrowCapsule** as the Static Mesh. 
+
+	![](Camera6.png)
+
+	By adjusting the Z value for the capsule location, you will see that the Static Mesh drops down inside of our CapsuleComponent.
+
+1.	In the **Components** window, click the **Add Component** button and search for and add a **Camera** component. 
+
+	![](Camera7.png)
+
+	This is the camera that will be tied to this character providing us with a third person perspective. While this example has us focusing on a playable character, you can add a Camera component to any Blueprint, for instance you could create a security camera that produces a view for a player to look through. 
+
+1.	In the **Details** panel for the **Camera**, set its location for **X** to **-120** and it’s **Z** to **20**. 
+
+	![](Camera8.png)
+
+	We've used these values to position the camera behind the capsule, however you can adjust the position as you see fit. 
+
+1.	Click the **Add Component** button again and search for and add an **Arrow** component. 
+
+	![](Camera9.png)
+
+	Not only are we going to allow this character to use a third person perspective, but we are going to allow the player press a button and move the camera to an over-the-shoulder type view. We will use this Arrow component to tell the camera where to move when the button is pressed or released. 
+
+1.	In the **Details** panel for the **Arrow** component, set its location the same as **Step 8**. 
+
+	![](Camera10.png)
+
+	This is a reference to the default camera location that we will want to return to once the button is released. 
+
+1.	Click the **Add Component** button and add another **Arrow** component, set its **Location** as shown below. 
+
+	![](Camera11.png)
+
+	We've set the **X** to **-40** and the **Y** to **40** to produce the example below. 
+
+	![](Camera11b.png)(w:800)
+
+	You can adjust the position of the Arrow components as you would like. 
+
+1.	Click the **Event Graph** tab, then **Right-click** in the graph and search for and add a **Right Mouse Button** Event. 
+
+	![](Camera12.png)
+
+	The [Event Graph](Engine/Blueprints/UserGuide/EventGraph) is where will provide the scripted functionality to move the camera. 
+
+1.	In the **My Blueprint** window, hold **Ctrl** and **Left-click** and drag the **Camera** component into the graph. 
+
+	![](Camera13.png)(w:800)
+
+	This will allow us to set properties or affect the camera component as we are getting a reference to it above. 
+
+1.	In the **My Blueprint** window, **Ctrl-Left-click** and drag in both **Arrow** components. 
+
+1.	In the graph, **Left-click** and drag off the out pin of the **Camera** node and search for and add a **SetRelativeLocation** node. 
+
+	![](Camera14.png)
+
+	This will allow us to move the camera to the desired location. 
+
+1.	Drag off **Arrow 2** and search for and add the **Get Relative Location** node, then connect the nodes as shown below. 
+
+	![](Camera15.png)(w:720)
+
+	When the **Right Mouse Button** is pressed, we **Set Relative Location** of the camera to the **Relative Location** of the **Arrow 2** component (which is our over-the-shoulder camera position). Next we need to set the camera back to its default location (which is the position of Arrow 1) when the Right Mouse Button is released. 
+
+1.	Copy/Paste the **SetRelativeLocation** and repeat the previous step for **Arrow 1** as shown below. 
+
+	![](Camera16.png)(w:720)
+
+	Our script is complete and we can now toggle between the two camera locations when the Right Mouse Button is pressed/released. 
+
+1.	Click **Compile** from the toolbar, **Save** and close the Blueprint. 
+
+	![](Camera17.png)
+
+	Our character is setup, however we still need to tell our project that the character we created should be the default character to use when playing the game. We can do that from the [Level World Settings](Engine/Levels) by changing the settings in the [Game Mode](Gameplay/HowTo/SettingUpAGameMode/Blueprints) that is being used with this template.  
+
+1.	From the main Level Editor window, click the **Settings** button from the toolbar and select **World Settings**. 
+
+	![](Camera18.png)
+
+1.	In the **World Settings** window, under **Selected GameMode**, click the **Default Pawn Class** and change it to **MyCharacter**. 
+
+	![](Camera19.png)
+
+1.	From the main toolbar, click the **Play** button to play in the editor. 
+
+
+## End Result
+
+[OBJECT:EmbeddedVideo]
+[PARAMLITERAL:width]
+640
+[/PARAMLITERAL]
+[PARAMLITERAL:height]
+360
+[/PARAMLITERAL]
+[PARAMLITERAL:videoid]
+vf2ef-Z490U
+[/PARAMLITERAL]
+[/OBJECT]
+
+When you play in the editor, you can now toggle between two different camera positions by moving the **Camera** component that we assigned. Our character in the video above also includes the ability to move around which we have taken the script from the **ThirdPersonCharacter** Blueprint (located under the **Content/ThirdPersonBP/Blueprints** folder) and copied it into our **MyCharacter** Blueprint. 
+
+If you would like to add movement functionality, copy the following script from **ThirdPersonCharacter** into your **MyCharacter** Blueprint. 
+
+[REGION:lightbox]
+[![](Camera20.png)(w:720)](Camera20.png)
+[/REGION]
+
+_Click image for a larger view._
+
+The script above is used to move the character around, jump and manipulate the viewing angle of the camera. 
+
+[REGION:note]
+After adding a **Camera** component, typically you will also want to include a [Spring Arm](Gameplay\HowTo\UsingCameras\SpringArmComponents) component to attach the camera to. This will help prevent the camera from being obstructed when moving into tight spaces or when an object moves between the character and the camera. 
+[/REGION]
diff --git a/Engine/Documentation/Source/Gameplay/HowTo/UsingCameras/SpringArmComponents/SpringArm.INT.udn b/Engine/Documentation/Source/Gameplay/HowTo/UsingCameras/SpringArmComponents/SpringArm.INT.udn
new file mode 100644
--- /dev/null
+++ b/Engine/Documentation/Source/Gameplay/HowTo/UsingCameras/SpringArmComponents/SpringArm.INT.udn
@@ -0,0 +1,72 @@
+Availability: Public
+Title:Using Spring Arm Components
+Crumbs: %ROOT%, Gameplay, Gameplay\HowTo\UsingCameras
+Description:The Spring Arm component is used to automatically control how the camera handles situations where the it becomes obstructed. 
+Related: Gameplay\HowTo\UsingCameras
+Related: Gameplay\Framework\Camera
+Related: Engine/Matinee
+Related: Engine/Blueprints
+SkillLevel: Beginner
+Version: 4.8
+
+[REGION:warning]
+You can use your own character for this tutorial provided it has a **Camera** component assigned to it. If you do not have a character, it is recommended that you complete the [](Gameplay\HowTo\UsingCameras\CameraComponents) tutorial in order to create the basic character being used in this guide first. 
+[/REGION]
+
+In this How-to, we add a **Spring Arm** component and parent a **Camera** component to it so that our camera can expand/retract based on gameplay situations. Typically when you add a Camera component to a character for the purposes of creating a third person perspective, you will also want to include a Spring Arm to automatically control how the camera handles situations where the it becomes obstructed by level geometry or other objects in the level. 
+
+![](springArm1.png)
+
+While it may not eliminate instances of the camera penetrating objects in your level, depending on your Spring Arm's settings, it will help reduce them by moving the camera's position when it becomes obstructed automatically and then returning to its default position when it is free from obstruction. 
+
+1.	Continuing from the [](Gameplay\HowTo\UsingCameras\CameraComponents) How-to, open your **Character Blueprint**. 
+
+1.	In the **Components** window, click the **Add Component** button then select the **Spring Arm** Component. 
+
+	![](Spring1.png)
+
+1.	Inside the **Components** window, drag your **Camera** Component onto the **Spring Arm** Component to attach it to the Spring Arm. 
+
+	![](Spring2.png)
+
+1.	There are multiple settings that can be used with the **Spring Arm** which are outlined below. 
+
+	![](Spring4.png)
+
+	| Section | Description |
+	| ------- | ------------ |
+	| **Sockets** | This allows you to parent the Spring Arm to a Bone on a Skeletal Mesh or a Socket. |
+	| **Camera** | Here you can set the length of the Spring Arm when there are no collisions as well as set Offsets. |
+	| **Camera Collision** | Set whether or not to check for collisions and define how big the query probe is to check for collisions. |
+	| **Camera Settings** | Whether or not to inherit Pitch, Yaw, or Roll from the Parent or Use Pawn Control Rotation. |
+	| **Lag** | Set whether or not the camera should lag behind the target position and settings associated with lag. |
+	
+1.	We set the **Target Arm Length** to **100** and **Enabled Camera lag**. 
+
+	![](spring3.png)
+
+	This gives us a shorter length than the default settings and also adds a bit of lag to the camera when moving around. 
+
+1.	Click **Compile** and **Play** to play in the editor. 
+
+## End Result
+
+[OBJECT:EmbeddedVideo]
+[PARAMLITERAL:width]
+640
+[/PARAMLITERAL]
+[PARAMLITERAL:height]
+360
+[/PARAMLITERAL]
+[PARAMLITERAL:videoid]
+jjf9P6u052E
+[/PARAMLITERAL]
+[/OBJECT]
+
+When moving the character, you will notice that the camera now lags slightly behind the characters movement. Also, when we move near a wall and rotate the camera, it attempts to reposition itself and moves in closer. When we are free from obstruction, it returns to its default length that we have specified in the settings. 
+
+
+
+
+
+
diff --git a/Engine/Documentation/Source/Gameplay/HowTo/UsingCameras/SwitchingCameras/SwitchingCameras_BP.INT.udn b/Engine/Documentation/Source/Gameplay/HowTo/UsingCameras/SwitchingCameras/SwitchingCameras_BP.INT.udn
new file mode 100644
--- /dev/null
+++ b/Engine/Documentation/Source/Gameplay/HowTo/UsingCameras/SwitchingCameras/SwitchingCameras_BP.INT.udn
@@ -0,0 +1,99 @@
+Availability: Public
+Title:Switching Between Multiple Fixed Cameras
+Crumbs: %ROOT%, Gameplay, Gameplay\HowTo\UsingCameras
+Description:An example of how to setup and trigger the use of multiple fixed camera angles. 
+Related: Gameplay\HowTo\UsingCameras
+Related: Gameplay\Framework\Camera
+Related: Engine/Matinee
+Related: Engine/Blueprints
+SkillLevel: Beginner
+Version: 4.8
+
+In this How-to we add multiple **Camera** Actors to our level and switch between them as the player enters **Trigger** Volumes that we have placed in the level. After completing this tutorial, you will have the ability to set up your own multi-camera based game and define when each camera is used. 
+
+[REGION:note]
+ For this example, we are using the **Blueprint Third Person Template** project with **Starter Content** included.
+[/REGION]
+
+
+1.	From the **Modes** menu while in **Place**, under **Basic**, drag in the number of **Camera Actors** you wish to cycle between. 
+
+	[REGION:lightbox]
+	[![](camera_1.png)(w:720)](camera_1.png)
+	[/REGION]
+	_Click image for a larger view._
+	
+	Here we have placed three cameras using different angles. 
+
+1.	Also in the **Modes** menu, under **Basic**, drag in the number of **Box Triggers** to match the number of cameras you placed. 
+
+	[REGION:lightbox]
+	[![](camera_2.png)(w:720)](camera_2.png)
+	[/REGION]
+	_Click image for a larger view._
+
+	Move the **Box Triggers** to where you want the player to trigger the changing of cameras. 
+
+	[REGION:tip]
+	With a Box Trigger selected, press **W** (to enter Translation Mode), **E** (to enter Rotation Mode) or **R** (to enter Scale Mode) to manipulate the Actors to your desired location, rotation or scale. Please see [](GettingStarted/HowTo/ManipulatingActors) for more information. 
+	[/REGION]
+
+1.	With each of your cameras selected, click the **Blueprints** button then select **Open Level Blueprint**. 
+
+	![](camera_3.png)
+
+1.	Inside the **Level Blueprint**, **Right-click** in the graph and choose **Create References to selected Actors**. 
+
+	![](camera_4.png)
+
+	This will add references to each of the Camera Actors you have placed in your level. 
+
+1.	Return to the Level Editor viewport, then in the [World Outliner](Engine/UI/LevelEditor/SceneOutliner/Interface), select each of your **Box Triggers** and return to the Level Blueprint. 
+
+1.	**Right-click** in the graph and under **Add Event for Selected Actors** and **Collision**, select **AddOnActorBeginOverlap**. 
+
+	![](camera_5.png)
+	
+	This will add **Events** for each **Box Trigger** that will fire off when the player enters them. 
+
+1.	**Right-click** in the graph and add the **Get Player Controller** node. 
+
+	![](camera_6.png)
+
+1.	**Right-click** in the graph and add the **Set View Target with Blend** node (un-check **Context Sensitive** if you do not see it). 
+
+	![](camera_7.png)
+
+1.	Connect the (TriggerBox), **Get Player Controller** and **CameraActor** nodes as shown below. 
+
+	![](camera_8.png)
+
+	We have now assigned the first camera to be used when entering the first trigger box. 
+
+1.	Repeat the process above for the rest of your cameras/trigger boxes. 
+
+	![](camera_9.png)
+
+	Each of our cameras are now set to be switched to when entering their respective trigger boxes. 
+
+1.	Click **Compile** then close the Blueprint and click **Play** to play in the Editor. 
+
+	![](camera_10.png)
+
+## End Result
+
+[OBJECT:EmbeddedVideo]
+[PARAMLITERAL:width]
+640
+[/PARAMLITERAL]
+[PARAMLITERAL:height]
+360
+[/PARAMLITERAL]
+[PARAMLITERAL:videoid]
+RQz0V-UuLCw
+[/PARAMLITERAL]
+[/OBJECT]
+
+You will see that when you enter each of the Box Trigger's you have placed, it updates the game view to the Camera Actor you have specified. This a quick example of how you could start working on a multi-camera based game, you would probably want to also account for situations when the player is outside of your camera's viewing angle and adjust player input based on camera relativity. 
+
+
diff --git a/Engine/Documentation/Source/Gameplay/HowTo/UsingCameras/UsingCamerasOverview.INT.udn b/Engine/Documentation/Source/Gameplay/HowTo/UsingCameras/UsingCamerasOverview.INT.udn
--- a/Engine/Documentation/Source/Gameplay/HowTo/UsingCameras/UsingCamerasOverview.INT.udn
+++ b/Engine/Documentation/Source/Gameplay/HowTo/UsingCameras/UsingCamerasOverview.INT.udn
@@ -5,11 +5,13 @@ Description:A How To Guide for Finding Actors in your Scenes in Unreal Engine 4.
 
 This "How To" series covers working with and placing **Cameras** in Unreal Engine 4. A Camera can be used by itself and placed directly into a level or it can be a part of a Blueprint (for instance to provide a perspective for the player when flying a plane, driving a car, or controlling a character). 
 
-Placing a Camera Actor can be done by dragging a **Camera** Actor into your level from the **Basic** tab of the **Modes** Menu. 
+Placing a Camera Actor can be done by dragging a **Camera** Actor into your level from **Modes** Menu. 
 
 ![](camera1.png)(w:820)
 
-Or by **Right-clicking** in the Level Viewport and using the Context Menu and selecting **Place Actor** then selecting the **Camera Actor**. 
+You can use the Search Bar to find the Camera Actor or look under the **All Classess** section. 
+
+You can also **Right-click** in the Level Viewport and use the context menu by selecting **Place Actor** then selecting the **Camera Actor**. 
 
 ![](camera4.png)
 
@@ -51,18 +53,18 @@ Below is an overview of each section in the **Details** panel for a Camera Actor
 For information on using Cameras for cinematic purposes, refer to [](Engine/Matinee/HowTo) documentation and the Camera sections. The [](Resources/Showcases/MatineeFightScene) is also an excellent resource for a sample cutscene and how it was constructed.
 [/REGION]
 
-Now that you have an understanding of how to place a Camera Actor in your level and some of the settings used to alter its look, head to the Setup section below to learn how to work with a Camera including how to use a Camera as a viewpoint for a player, how to use a Camera Component as part of a Blueprint, how to use a Spring Arm Component with a Camera Component (typically used for creating Third Person perspectives), as well as how to switch between Cameras during gameplay. 
+Now that you have an understanding of how to place a Camera Actor in your level and some of the settings used to alter its look, the examples below provide some ways in which you can use Camera including how to use a Camera as a viewpoint for a player, how to use a Camera Component as part of a Blueprint, how to use a Spring Arm Component with a Camera Component (typically used for creating Third Person perspectives), as well as how to switch between Cameras during gameplay. 
 
 
-## Setup
+## Usage Examples
 
-To learn more about Using Cameras in Blueprints, use the link below:
+Refer to each of the links below for step-by-step usage case examples. 
 
 [REGION:buttonlist]
 * [](Gameplay\HowTo\UsingCameras\Blueprints)
-<!-- When Code is Ready we will re-enable this link.
-* [](Gameplay\HowTo\FindingActors\Code)
--->
+* [](Gameplay\HowTo\UsingCameras\SwitchingCameras)
+* [](Gameplay\HowTo\UsingCameras\CameraComponents)
+* [](Gameplay\HowTo\UsingCameras\SpringArmComponents)
 [/REGION]
 
 
