#UE4doc #docs: removed broken characters: ?

diff --git a/Engine/Documentation/Source/Support/Builds/ReleaseNotes/2013/July/ReleaseNotesJuly2013.INT.udn b/Engine/Documentation/Source/Support/Builds/ReleaseNotes/2013/July/ReleaseNotesJuly2013.INT.udn
--- a/Engine/Documentation/Source/Support/Builds/ReleaseNotes/2013/July/ReleaseNotesJuly2013.INT.udn
+++ b/Engine/Documentation/Source/Support/Builds/ReleaseNotes/2013/July/ReleaseNotesJuly2013.INT.udn
@@ -30,14 +30,14 @@ Template:ReleaseNotes.html
 
 ####Editor Tools
 * **Plugins Editor**
-	* There�s now a new interface that lets you see, browse and maintain Plugins -  both built-in plugins as well as plugins installed into your game project     
+	* There is now a new interface that lets you see, browse and maintain Plugins -  both built-in plugins as well as plugins installed into your game project     
 		![](pluginsEd.png)
 	* You can browse categories and search/filter, and enable/disable or uninstall plugins, all from this interface
 * **New PlayerController and Mouse Interface Events**
 	* The new mouse interface /cursor system is implemented through the PlayerController.  
 		![](nc1.png)
 	* ShowMouseCursor enables the system and will begin displaying the default mouse cursor when in-game.  NOTE: this does not work with Play in Viewport, you must use Play in New Window
-	* The interface uses a very simple set of input touch events and index in order to �touch� an object:  
+	* The interface uses a very simple set of input touch events and index in order to "touch" an object:  
 		![](mc2bpnode.png) ![](mc2bpnode2.png) ![](mc2bpnode3.png)
 * **Add Volumes Through Placement Browser**
 	* Different types of volumes can now be dropped in through the Placement Browser!  
@@ -328,7 +328,7 @@ Template:ReleaseNotes.html
 
 
 ####Core
-* Added command line `�AllowStdOutLogVerbosity`. This will print every log with a verbosity >= Log to stdout if you run the `* -Cmd.exe`
+* Added command line `AllowStdOutLogVerbosity`. This will print every log with a verbosity >= Log to stdout if you run the `* -Cmd.exe`
 * New Plugin Descriptor changes 
 	* Many new fields were added to plugin descriptors.  (Documentation update coming soon!)
 	* Plugins can now have a Resources sub-folder (for Icon128.png thumbnail)
