Updates for Jira UEDOC-1916.
Added icon, add entry to the table.

diff --git a/Engine/Documentation/Source/Engine/UI/LevelEditor/Toolbar/LevelEditorToolbar.INT.udn b/Engine/Documentation/Source/Engine/UI/LevelEditor/Toolbar/LevelEditorToolbar.INT.udn
--- a/Engine/Documentation/Source/Engine/UI/LevelEditor/Toolbar/LevelEditorToolbar.INT.udn
+++ b/Engine/Documentation/Source/Engine/UI/LevelEditor/Toolbar/LevelEditorToolbar.INT.udn
@@ -32,6 +32,7 @@ The **Toolbar** panel, like in most applications, is a group of commands providi
 | Command | Description |
 | ------- | ----------- |
 | ![](button_Toolbar_Save.png)| Saves the current level.|
+| ![](SourceControl.png)| Allows you to connect to or assign your [Source Control](Engine/UI/SourceControl) solution. |
 | ![](button_Toolbar_Content.png)| Opens the [Content Browser](Engine/Content/Browser) displaying all of the content assets contained in your game. This is where you go to create, import, and edit all content. |
 | ![](button_Toolbar_marketplace.png) | Opens the Unreal Engine Launcher to the Marketplace section.  |
 | ![](button_Toolbar_View.png) | Displays the [Quick Settings](Engine/UI/LevelEditor/QuickSettings) menu providing easy access to commonly used options controlling selection, editing, and previewing aspects of the Level Editor. |
