UEDOC-1897

Replacing Gif files with Youtube videos

diff --git a/Engine/Documentation/Source/Engine/Rendering/Materials/HowTo/Previewing_Applying/Previewing_Applying.INT.udn b/Engine/Documentation/Source/Engine/Rendering/Materials/HowTo/Previewing_Applying/Previewing_Applying.INT.udn
--- a/Engine/Documentation/Source/Engine/Rendering/Materials/HowTo/Previewing_Applying/Previewing_Applying.INT.udn
+++ b/Engine/Documentation/Source/Engine/Rendering/Materials/HowTo/Previewing_Applying/Previewing_Applying.INT.udn
@@ -51,7 +51,17 @@ While inside of the Material Editor, if you use your mouse cursor to hover over
 While you are creating your materials in the Material Editor, the **Viewport** window will update with your changes in real time if you have real time enabled (Real time is enabled by default). 
 To see the changes happen, all you have to do is adjust any parameter that is currently hooked up to any of the Main Material Nodes like in the following example.
 
-![](Previewing_Applying_Changing_Color_Realtime.gif)(convert:false)
+[OBJECT:EmbeddedVideo]
+[PARAMLITERAL:width]
+640
+[/PARAMLITERAL]
+[PARAMLITERAL:height]
+360
+[/PARAMLITERAL]
+[PARAMLITERAL:videoid]
+TDn_9Dg1s-Q
+[/PARAMLITERAL]
+[/OBJECT]
 
 [region:note]
  When making any changes to your Material Network, it could take some time for the changes to be rendered correctly. The more complex your material becomes, the longer it could take for the preview window to update. If you find that you really need a quicker update, you should look into making your Material into a [Material Instance](Engine/Rendering/Materials/MaterialInstances). 
@@ -69,13 +79,33 @@ To enable this functionality, you need to make sure that the Material meets the
 
 1. First make sure that the Material that you want to preview is applied to something in the level. 
 
-	![](MMN_Applying_Material.gif)(convert:false)
+[OBJECT:EmbeddedVideo]
+[PARAMLITERAL:width]
+640
+[/PARAMLITERAL]
+[PARAMLITERAL:height]
+360
+[/PARAMLITERAL]
+[PARAMLITERAL:videoid]
+SvCjqU1XDzY
+[/PARAMLITERAL]
+[/OBJECT]
 
 1. Inside of the Material, make sure that you are using **Scalar** or **Vector Parameters** as inputs for the properties you want to change. 
 	Previewing like this will not work with **Constant** Material Expression nodes, you have to use or convert the Material Expression nodes you want to tweak into **Parameter** nodes.
 	It is very important that the Material Expression inputs that you convert to preview have unique names or this will not work.
 
-	![](MMN_Converting_Nodes.gif)(convert:false)
+	[OBJECT:EmbeddedVideo]
+	[PARAMLITERAL:width]
+	640
+	[/PARAMLITERAL]
+	[PARAMLITERAL:height]
+	360
+	[/PARAMLITERAL]
+	[PARAMLITERAL:videoid]
+	mnlMbWkxz9w
+	[/PARAMLITERAL]
+	[/OBJECT]
 
 1.   To see the changes happen in real time in the level viewport, adjust the value of the Scalar or Vector parameters inside the Material. 
 	The adjustments that you make should then be displayed, in real time, in the level viewport.
@@ -96,10 +126,30 @@ To enable this functionality, you need to make sure that the Material meets the
 
 There are two main ways that you can apply Materials to objects inside of UE4.
 * **Drag and Drop** - You can drag any Material from the **Content Browser** and then drop it on any object in the level to change its material like in the following example:
-![](Previewing_Applying_Drag_&_Drop.gif)(convert:false)
+[OBJECT:EmbeddedVideo]
+[PARAMLITERAL:width]
+640
+[/PARAMLITERAL]
+[PARAMLITERAL:height]
+360
+[/PARAMLITERAL]
+[PARAMLITERAL:videoid]
+57tC0KaB2BU
+[/PARAMLITERAL]
+[/OBJECT]
 
 * **Apply via Object's Details Panel** - You can also apply materials via the object's **Details** panel by using the **Use selected asset from Content Browser** button like in the following example.
-![](Previewing_Applying_Assing_From_CB.gif)(convert:false)
+[OBJECT:EmbeddedVideo]
+[PARAMLITERAL:width]
+640
+[/PARAMLITERAL]
+[PARAMLITERAL:height]
+360
+[/PARAMLITERAL]
+[PARAMLITERAL:videoid]
+fSLRV947n7A
+[/PARAMLITERAL]
+[/OBJECT]
 
 
 ## Conclusion
