Jira UEDOC-1834
Gif to YouTube video.

diff --git a/Engine/Documentation/Source/Engine/Rendering/LightingAndShadows/QuickStart/2/LQS_2.INT.udn b/Engine/Documentation/Source/Engine/Rendering/LightingAndShadows/QuickStart/2/LQS_2.INT.udn
--- a/Engine/Documentation/Source/Engine/Rendering/LightingAndShadows/QuickStart/2/LQS_2.INT.udn
+++ b/Engine/Documentation/Source/Engine/Rendering/LightingAndShadows/QuickStart/2/LQS_2.INT.udn
@@ -50,7 +50,17 @@ With our basic apartment setup created in the previous step, we will now add the
 
 	With **Atmosphere Sun Light** checked, we can control the location of the sun using **Rotation** mode (**E**) and rotating our light. 
 
-	![](SunControlGif.gif)(convert:false)
+	[OBJECT:EmbeddedVideo]
+	[PARAMLITERAL:width]
+	640
+	[/PARAMLITERAL]
+	[PARAMLITERAL:height]
+	360
+	[/PARAMLITERAL]
+	[PARAMLITERAL:videoid]
+	d5PdhSHwpWs
+	[/PARAMLITERAL]
+	[/OBJECT]
 
 	The lighting in our level changes based on the location of the sun allowing us to easily switch between various times of day. 
 
