#UE4doc #docs:
Updated Engine version

diff --git a/Engine/Documentation/Source/Programming/Slate/SlateUI.INT.udn b/Engine/Documentation/Source/Programming/Slate/SlateUI.INT.udn
--- a/Engine/Documentation/Source/Programming/Slate/SlateUI.INT.udn
+++ b/Engine/Documentation/Source/Programming/Slate/SlateUI.INT.udn
@@ -3,7 +3,7 @@ Title:Slate UI Framework
 Crumbs:%ROOT%, Programming
 Description:Cross-platform user interface framework for creating tool and in-game UIs.
 Navigation:topic
-Version:4.4
+Version:4.9
 
 [VAR:Topic]
 [OBJECT:Topic]
