#UE4doc #docs:
Updated Engine version
Changed font settings for multiple terms

diff --git a/Engine/Documentation/Source/Programming/Slate/Architecture/SlateArchitecture.INT.udn b/Engine/Documentation/Source/Programming/Slate/Architecture/SlateArchitecture.INT.udn
--- a/Engine/Documentation/Source/Programming/Slate/Architecture/SlateArchitecture.INT.udn
+++ b/Engine/Documentation/Source/Programming/Slate/Architecture/SlateArchitecture.INT.udn
@@ -2,12 +2,12 @@ Availability:Public
 Title:Slate Architecture
 Crumbs:%ROOT%, Programming, Programming/Slate
 Description: Core ideas guiding Slate design.
-Version:4.4
+Version:4.9
 
 [TOC (start:2)]
 
 ## How to Read This Page
-Core ideas guiding Slate design are found here. They are found in no particular order.
+Core ideas guiding **Slate** design are found here. They are found in no particular order.
 There is no rigid structure or grand theory here; it is just a collection of principles that
 arise from UI tinkering experience.
 
@@ -47,7 +47,7 @@ Some follow here:
 
 Design for developer efficiency where possible. A programmer's time is expensive; CPUs are fast and cheap.
 
-* Avoid opaque caches and duplicated state! Historically, UIs cache state and require explicit invalidation. Slate uses following approaches (from preferred to least preferred):
+* Avoid opaque caches and duplicated state. Historically, UIs cache state and require explicit invalidation. Slate uses following approaches (from preferred to least preferred):
 	1. Polling
 	1. Transparent caches
 	1. Opaque caches with low-grain invalidation
@@ -64,10 +64,10 @@ UIs visualize and manipulate Models. Slate uses delegates as a flexible conduit
 When a user performs some actions, Slate widgets invoke the write delegate in order to modify the data.
 
 
-Consider STextBlock - a Slate widget that displays text.
+Consider **STextBlock** - a Slate widget that displays text.
 The STextBlock must be told where to get the data to display.
 The data could be set statically. However, a more flexible way to accomplish this is a delegate (a function specified by the user).
-The STextBlock uses a delegate called Text for this purpose.
+The STextBlock uses a delegate called **Text** for this purpose.
 
 [REGION:figure]
 
@@ -86,13 +86,13 @@ The STextBlock uses a delegate called Text for this purpose.
 	[/REGION]
 [/REGION]
 
-Consider also that in the example above the Text is a string, while the `framerate` is most likely stored as a `float` or an `integer`. The use of delegates offers us the flexibility to perform the conversion whenever the value is read.
+Consider also that in the example above the Text is a string, while the framerate is most likely stored as a `float` or an `integer`. The use of delegates offers us the flexibility to perform the conversion whenever the value is read.
 This immediately conjures up performance concerns which are addressed in the [Performance Considerations](#PerformanceConsiderations) section below.
 
 
 
-SEditableText is a Slate widget responsible for both input and output. Just like STextBlock, it uses the Text delegate for visualizing data.
-The user types some text into the editable text field, and when they hit enter, SEditableText invokes the OnTextChanged delegate.
+**SEditableText** is a Slate widget responsible for both input and output. Just like STextBlock, it uses the Text delegate for visualizing data.
+The user types some text into the editable text field, and when they hit enter, SEditableText invokes the **OnTextChanged** delegate.
 The assumption is that the programmer has attached the functionality appropriate for validating the input and mutating the Model's data to OnTextChanged.
 
 [REGION:figure]
@@ -118,7 +118,7 @@ During the next frame, the SEditableText will read from the Model's data. In the
 ### Attributes and Arguments
 
 Using a delegate is not always desirable. Depending on the use case, the arguments to Slate widgets may need to be constant values or functions.
-We encapsulate this notion via the `TAttribute<T>` class. An attribute can be set to a constant or to a delegate.
+We encapsulate this notion via the **TAttribute< T >** class. An attribute can be set to a constant or to a delegate.
 
 ## Performance Considerations
 
@@ -146,9 +146,9 @@ This is preferable to fine-grain invalidation because it is simpler and more mai
 ## Child Slots
 
 All Slate widgets store children in child slots. (As opposed to storing a plain array of child widgets.)
-Child slots always store a valid widget; by default they store the SNullWidget, which is a widget with no visualization or interaction.
+Child slots always store a valid widget; by default they store the **SNullWidget**, which is a widget with no visualization or interaction.
 Each type of widget can declare its own type of child slot, which caters to its specific needs.
-Consider that SVerticalSlot arranges its children completely differently than an SCanvas, which is quite different from SUniformGridPanel.
+Consider that **SVerticalSlot** arranges its children completely differently than an **SCanvas**, which is quite different from SUniformGridPanel.
 The Slots allow each type of panel to ask for a set of per-child settings that affect the arrangement of the children.
 
 
@@ -156,17 +156,17 @@ The Slots allow each type of panel to ask for a set of per-child settings that a
 
 Widgets come in three flavors.
 
-* Leaf Widgets - widgets with no child slots. e.g. STextBlock displays a piece of text. It has native knowledge of how to draw text.
-* Panels - widgets with a dynamic number of child slots. e.g. SVerticalBox arranges any number of children vertically given some layout rules.
-* Compound Widgets - widgets with a fixed number of explicitly named child slots. e.g. SButton has one slot called Content which contains any widgets inside the button.
+* **Leaf Widgets** - widgets with no child slots. e.g. STextBlock displays a piece of text. It has native knowledge of how to draw text.
+* **Panels** - widgets with a dynamic number of child slots. e.g. **SVerticalBox** arranges any number of children vertically given some layout rules.
+* **Compound Widgets** - widgets with a fixed number of explicitly named child slots. e.g. **SButton** has one slot called Content which contains any widgets inside the button.
 
 
 ## Layout
 
 Slate layout is accomplished in two passes. The two-pass breakdown is an optimization, but unfortunately not a transparent one.
 
-1. Pass 1: Cache Desired Size - the relevant functions are `SWidget::CacheDesiredSize` and `SWidget::ComputeDesiredSize`
-1. Pass 2: ArrangeChildren - the relevant function is `SWidget::ArrangeChildren`
+1. Pass 1: **Cache Desired Size** - the relevant functions are **SWidget::CacheDesiredSize** and **SWidget::ComputeDesiredSize**
+1. Pass 2: **ArrangeChildren** - the relevant function is **SWidget::ArrangeChildren**
 
 
 In more detail:
@@ -176,7 +176,7 @@ In more detail:
 The goal of this pass is to figure out how much space each widget wants to occupy.
 Widgets with no children (i.e. leaf widgets) are asked to compute and cache their desired size based on their intrinsic properties.
 Widgets that combine other widgets (i.e. compound widgets and panels) use special logic to determine their desired size as a function of the size of their children.
-Note that each type of widget is only required to implement ComputeDesiredSize(); the caching and traversal logic are implemented by Slate.
+Note that each type of widget is only required to implement **ComputeDesiredSize();** the caching and traversal logic are implemented by Slate.
 Slate guarantees that when ComputeDesiredSize() is called on a widget, its children have already computed and cached their desired size.
 Thus, this is a bottom-up pass.
 
@@ -248,7 +248,7 @@ In the example above, the panel was allotted 25 units by its parent.
 The first slot indicates that it wants to use the desired size of the child, and is allotted 14 units of space.
 The second slot indicates that it wants to fill the available width, and is allotted the balance (11 units of space).
 Note that in the actual SHorizontalBox widget, the alignment of the SImage within its slot would be driven by
-the HAlign property, which can be Left, Center, Right, or Fill.
+the **HAlign** property, which can be Left, Center, Right, or Fill.
 
 
 In practice, Slate never performs a full ArrangeChildren pass. Instead, this functionality is used to implement other functionality.
@@ -262,7 +262,7 @@ This list is produced anew for every frame.
 
 We begin at the top level windows and recur down the hierarchy, appending the draw elements of every widget to the draw list.
 Widgets tend to do two things during paint: they output actual draw elements or figure out where a child widget should exist
-and ask the child widget to paint itself. Thus, we can think of a simplified general-case OnPaint function as being:
+and ask the child widget to paint itself. Thus, we can think of a simplified general-case **OnPaint** function as being:
 
 ~~~
 .	// An arranged child is a widget and its allotted geometry
