#UE4doc #docs:
Fixed capitalization on a few words
Changed font settings for multiple terms / lines

diff --git a/Engine/Documentation/Source/Programming/Slate/DetailsCustomization/DetailsCustomization.INT.udn b/Engine/Documentation/Source/Programming/Slate/DetailsCustomization/DetailsCustomization.INT.udn
--- a/Engine/Documentation/Source/Programming/Slate/DetailsCustomization/DetailsCustomization.INT.udn
+++ b/Engine/Documentation/Source/Programming/Slate/DetailsCustomization/DetailsCustomization.INT.udn
@@ -9,15 +9,15 @@ The **Details** panel is now fully customizable.  You can rearrange properties v
 
 ## Setup instructions
 
-1. Create a class for customizing properties in.  This must inherit from `ILayoutDetails`.
+1. Create a class for customizing properties in.  This must inherit from **ILayoutDetails**.
     
-	* You implement one function: `void LayoutDetails( IDetailLayoutBuilder& )`.
+	* You implement one function: **void LayoutDetails( IDetailLayoutBuilder& )**.
 	* The purpose of this class is to encapsulate customization for a classes properties.  One instance of the class will be created for each **Details** panel that requires it. 
 
 2. Set up a delegate that will be called when the **Details** panel recognizes properties for a specific class.
     
-	* The sole purpose of this delegate is to create an instance of your customization class for a specific `UObject` that has properties.  Remember that there often are multiple details views up at any point and each instance of the details view gets its own customization class instance.  This allows you to store per detail instance data on your layout class.  
-	* Example (more examples are located in `DetailCustomizations.cpp`):
+	* The sole purpose of this delegate is to create an instance of your customization class for a specific **UObject** that has properties.  Remember that there often are multiple details views up at any point and each instance of the details view gets its own customization class instance.  This allows you to store per detail instance data on your layout class.  
+	* Example (more examples are located in DetailCustomizations.cpp):
 	
 			FPropertyEditorModule& PropertyModule = FModuleManager::LoadModuleChecked<FPropertyEditorModule>("PropertyEditor");
 			PropertyModule.RegisterCustomPropertyLayout( ABrush::StaticClass(), FOnGetDetailLayoutInstance::CreateRaw( &FBrushDetails::MakeInstance ) );
@@ -33,14 +33,14 @@ The **Details** panel is now fully customizable.  You can rearrange properties v
 	* If this is an engine class, you should add your customization class (if it does not already exist) to the `DetailCustomizations` module.  This module can be recompiled and reloaded without restarting the editor, making it useful for  fast tweaking of properties.
 	* Bind your delegate in `FDetailCustomizationsModule.StartupModule` and unbind it in `FDetailCustomizationsModule.ShutdownModule`.
 	* Game specific classes should use a game specific module.
-	* See examples in this document and the `DetailCustomizations` module (such as `PrimitiveComponentDetails.cpp` and `StaticMeshComponentDetails.cpp`).
+	* See examples in this document and the `DetailCustomizations` module (such as PrimitiveComponentDetails.cpp and StaticMeshComponentDetails.cpp).
 
 
 ## Customizing
 
-You handle all customization inside the `LayoutDetails` function of your customization class.  This function accepts an `IDetailLayoutBuilder`, which is your interface to properties and how you pass back customization widgets.
+You handle all customization inside the LayoutDetails function of your customization class.  This function accepts an **IDetailLayoutBuilder**, which is your interface to properties and how you pass back customization widgets.
 
-The primary function of `IDetailLayoutBuilder` is to create categories where properties and other details reside.  There are some other minor functions on this class which are self explanatory (most of them you will not need). The documentation for those can be found in `DetailLayoutBuilder.h`.
+The primary function of IDetailLayoutBuilder is to create categories where properties and other details reside.  There are some other minor functions on this class which are self explanatory (most of them you will not need). The documentation for those can be found in DetailLayoutBuilder.h.
 
 The first step in customizing is to edit a category:
 
@@ -50,9 +50,9 @@ The first step in customizing is to edit a category:
 		IDetailCategory& LightingCategory = DetailBuilder.EditCategory("Lighting", TEXT("OptionalLocalizedDisplayName") );
 	}
 
-The `EditCategory` function takes an `FName` for the category where properties will reside and an optional localized display name.  If the display name is specified, it will override any existing display name.  The category name does **NOT** have to be the same category name specified in the `UPROPERTY` macro, although it will reuse the `UPROPERTY` category if the names are the same.  The macro category name is used as the default category if the property is not customized and in the tree view.
+The **EditCategory** function takes an **FName** for the category where properties will reside and an optional localized display name.  If the display name is specified, it will override any existing display name.  The category name does not have to be the same category name specified in the `UPROPERTY` macro, although it will reuse the `UPROPERTY` category if the names are the same.  The macro category name is used as the default category if the property is not customized and in the tree view.
 
-`EditCategory` returns an `IDetailCategoryBuilder&` which is what you use to add properties to a category.  There are a two ways to do this:
+EditCategory returns an **IDetailCategoryBuilder&** which is what you use to add properties to a category.  There are a two ways to do this:
 
 * Using simple [multibox style layout](#MultiboxStyleLayout), which is quick for rearranging properties.
 * Using [Slate syntax](#SlateLayout) which offers the most robust customization options.
@@ -60,7 +60,7 @@ The `EditCategory` function takes an `FName` for the category where properties w
 
 ### Multibox Style Layout
 
-An easy example using the `LightingCategory` created above:
+An easy example using the **LightingCategory** created above:
 		
 	// Add a property to the category.  The first param is the name of the property and the second is an optional display name override.
 	LightingCategory.AddProperty("bCastStaticShadow", TEXT("Static") );
@@ -122,7 +122,7 @@ Now you have a handle to the bool property `bOverrideLightmapRes`.
 
 From here, you can read and write the value of that property and/or pass it to a slate widget for customization.  
 
-**Useful property handle functions** (For the complete documented list see `PropertyHandle.h`):
+Useful property handle functions (For the complete documented list see PropertyHandle.h):
 
 | Function | Description |
 | -------- | ----------- |
@@ -340,7 +340,7 @@ The result:
 
 ### Accessing complicated properties.
 
-A _complicated property_ is defined as anything that cannot be resolved with just a property name.  Usually this involves properties inside structs.
+A complicated property is defined as anything that cannot be resolved with just a property name.  Usually this involves properties inside structs.
 
 There are two ways to access complicated properties:
 
@@ -386,7 +386,7 @@ There are two ways to access complicated properties:
 
 ### Accessing Arrays
 
-You can access arrays via `IPropertyHandle::AsArray`.  If the property handle is an array, this will return an `IPropertyHandleArray` which has functions for _Adding_, _Removing_, _Inserting_, _Duplicating_, and _getting the number of elements_ of an array.
+You can access arrays via `IPropertyHandle::AsArray`.  If the property handle is an array, this will return an `IPropertyHandleArray` which has functions for adding, removing, inserting, duplicating, and getting the number of elements of an array.
 
 ### Hiding properties
 
