UEDOC-1878

Gifs have been replaced with YouTube videos

diff --git a/Engine/Documentation/Source/Engine/Rendering/Materials/HowTo/BumpOffset/BumpOffset.INT.udn b/Engine/Documentation/Source/Engine/Rendering/Materials/HowTo/BumpOffset/BumpOffset.INT.udn
--- a/Engine/Documentation/Source/Engine/Rendering/Materials/HowTo/BumpOffset/BumpOffset.INT.udn
+++ b/Engine/Documentation/Source/Engine/Rendering/Materials/HowTo/BumpOffset/BumpOffset.INT.udn
@@ -32,7 +32,17 @@ When searching for the Bump Offset Material Expression, you will see another Mat
 [/Region]
 
 Once you locate the Bump Offset Material Expression, you can bring it into the Material Graph like any other Material Expression by first selecting it with your **Left Mouse Button** and then dragging it into the Material Graph. When you release your **Left Mouse Button**, the Bump Offset Material Expression will be placed in the Material Graph.
-![](Bump_Offest_Adding_BO_Node.gif)(convert:false)
+[OBJECT:EmbeddedVideo]
+[PARAMLITERAL:width]
+640
+[/PARAMLITERAL]
+[PARAMLITERAL:height]
+360
+[/PARAMLITERAL]
+[PARAMLITERAL:videoid]
+XUq0zSiUIX4
+[/PARAMLITERAL]
+[/OBJECT]
 
 ## Using the Bump Offset Material Expression 
 
@@ -46,7 +56,17 @@ content page for information about how to move content between projects. This wa
 1. First use the mouse to **Right-Click** in the **Content Browser** and then from the **Create Basic Asset** section of the pop-up Menu, select **Material** and name the Material **Bump_Offset_Example**.
 Once named, **Double-Click** on the Material in the **Content Browser** to open it up.
 
-	![](Bump_Offest_Creating_Material.gif)(convert:false)
+	[OBJECT:EmbeddedVideo]
+	[PARAMLITERAL:width]
+	640
+	[/PARAMLITERAL]
+	[PARAMLITERAL:height]
+	360
+	[/PARAMLITERAL]
+	[PARAMLITERAL:videoid]
+	THFueQLwS6Q
+	[/PARAMLITERAL]
+	[/OBJECT]
 
 1. Because the Bump Offset Material Expression works using UV coordinates, we need to have a Material that utilizes Textures for the Base Color, Roughness, Ambient Occlusion, and Normal inputs. 
    For this example, we are going to be using the following textures that can be found in the Textures folder if you include the Starter Content in your project. 
@@ -79,13 +99,33 @@ Once named, **Double-Click** on the Material in the **Content Browser** to open
 	[/Region]
 
 1. Once the Bump Offset Material Expression has been connected to the UV input of all the Textures, it now needs to be setup so that it can work. The first thing that needs to be done is that we need to give it a Height input. To do this, grab the Texture we are using for the Base Color input, T_Brick_Clay_New_D, and duplicate it by first selecting it with the mouse and then pressing **CRTL + W** on the keyboard. Once the node is duplicated, move it so that it is in front of the Bump Offset Material Expression and the connect the Alpha output of the Texture into the Height input of the Bump Offset Material Expression. When completed, you should have something that looks like this. 
-	![](Bump_Offest_Duplicate_For_HM.gif)(convert:false)
+	[OBJECT:EmbeddedVideo]
+	[PARAMLITERAL:width]
+	640
+	[/PARAMLITERAL]
+	[PARAMLITERAL:height]
+	360
+	[/PARAMLITERAL]
+	[PARAMLITERAL:videoid]
+	dIRc_h7jX-s
+	[/PARAMLITERAL]
+	[/OBJECT]
 
 1. Now the settings for the Bump Offset Material Expression need to be set so that we do not get what is commonly called UV swimming. To adjust the settings of the Bump Offset Material Expression, select the Bump Offset Material Expression in the Material Graph by clicking on it with the **Left Mouse Button** and  then in the **Details** panel, set the **Height Ratio** to **0.005**. When completed, you should have something that looks like this. 
 	![](BumpOffset_Setting_Up_BO_Node.png)
 
 1. Once the Height Ratio is set, make sure to compile your Material by pressing the green **Apply** button. After the Material has been compiled, it is now able to be used for in your game.  
-	![](Bump_Offest_Press_Apply.gif)(convert:false)
+	[OBJECT:EmbeddedVideo]
+	[PARAMLITERAL:width]
+	640
+	[/PARAMLITERAL]
+	[PARAMLITERAL:height]
+	360
+	[/PARAMLITERAL]
+	[PARAMLITERAL:videoid]
+	e-DLy1G7lts
+	[/PARAMLITERAL]
+	[/OBJECT]
 	
 	Here is an example of the new created Material in action. 
 	![](BumpOffset_Setting_Results.png)
