#UE4doc #docs:
SlateInGameUI:
Removed Overview header
Updated Engine version
Expanded on a few acronyms upon their first uses
Changes capitalization on a few terms.

SlateProjectSetup:
Removed Overview header
Updated Engine version
Changed font settings for multiple terms/lines
Changed some grammar/wording to meet our instructional guide's standards

diff --git a/Engine/Documentation/Source/Programming/Slate/InGameUI/SlateInGameUI.INT.udn b/Engine/Documentation/Source/Programming/Slate/InGameUI/SlateInGameUI.INT.udn
--- a/Engine/Documentation/Source/Programming/Slate/InGameUI/SlateInGameUI.INT.udn
+++ b/Engine/Documentation/Source/Programming/Slate/InGameUI/SlateInGameUI.INT.udn
@@ -2,12 +2,11 @@ Availability:Public
 Title: Using Slate In-Game
 Crumbs:%ROOT%, Programming, Programming/Slate
 Description:Using Slate UI widgets for in-game user interfaces.
-Version:4.4
+Version:4.9
 
-## Overview
 
-Slate widgets can be used in-game to create heads up displays (HUDs) or other user interface (UI) elements, such 
-as menus. You will generally create one or more _container_ widgets, each of which may be composed of several 
+Slate widgets can be used in-game to create Heads Up Displays (HUDs) or other User Interface (UI) elements, such 
+as menus. You will generally create one or more container widgets, each of which may be composed of several 
 other types of widgets that are responsible for a particular aspect of the UI. 
 
 <!-- HUD widget w/ menu overlay -->
@@ -37,7 +36,7 @@ In order to display a Slate widget in your game, it must be added to the game's
 are ordered according to the Z-order specified when adding them, with larger Z-order values appearing on top 
 of smaller Z-orders.
 
-**Accessing the Game Viewport**
+### Accessing the Game Viewport
 
 The game's viewport is an instance of the `GameViewportClient` class. A reference to the current game viewport 
 is accessible via the `GameViewport` member of UEngine, which can be accessed using the `GEngine` global pointer 
@@ -52,7 +51,7 @@ For example:
 or any of their members.
 [/REGION]
 
-**Adding the Widget to the Viewport**
+### Adding the Widget to the Viewport
 
 Slate widgets are added to the viewport by passing a reference - a `TSharedref<SWidget>` to be exact - to the widget 
 to `GameViewportClient::AddViewportWidgetContent()`. This function takes in both a widget and a Z-order, which 
diff --git a/Engine/Documentation/Source/Programming/Slate/Setup/SlateProjectSetup.INT.udn b/Engine/Documentation/Source/Programming/Slate/Setup/SlateProjectSetup.INT.udn
--- a/Engine/Documentation/Source/Programming/Slate/Setup/SlateProjectSetup.INT.udn
+++ b/Engine/Documentation/Source/Programming/Slate/Setup/SlateProjectSetup.INT.udn
@@ -2,13 +2,11 @@ Availability:Public
 Title: Using Slate in a Project
 Crumbs:%ROOT%, Programming, Programming/Slate
 Description:Setting up your project to use the Slate UI Framework
-Version:4.4
-
-## Overview
+Version:4.9
 
 [EXCERPT:main]
 In order to use the Slate User Interface (UI) Framework, your project must be set up properly so that it is aware of the 
-framework. This allows you to include the `Slate.h` header and reference the various elements of the 
+framework. This allows you to include the Slate.h header and reference the various elements of the 
 framework necessary for building a UI with slate.
 
 ## Module Dependencies
@@ -26,7 +24,7 @@ The modules your project needs access to are:
 
 **To set up the Slate module dependencies:**
 
-1. Open the `[ProjectName].build.cs` file for your project. It is located in the `[ProjectDir]/[ProjectName]/Source/[ProjectName]` directory.
+1. Open the [ProjectName].build.cs file for your project. It is located in the [ProjectDir]/[ProjectName]/Source/[ProjectName] directory.
 1. Add the InputCore public dependency by adding `"InputCore"` to the `PublicDependencyModuleNames`.
 
 		PublicDependencyModuleNames.AddRange(new string[] { "Core", "CoreUObject", "Engine", "InputCore" });
@@ -39,14 +37,14 @@ The modules your project needs access to are:
 
 		PrivateDependencyModuleNames.AddRange(new string[] {  });
 
-	All you need to do is add the SlateCore, and Slate modules to that line:
+	Add the SlateCore and Slate modules to that line:
 
 		PrivateDependencyModuleNames.AddRange(new string[] { "Slate", "SlateCore" });
 
 	[REGION:tip]
 	Depending on when you created your project, and with what version of the engine, it may already have the Slate 
-	dependencies set up in the *.build.cs files but commented out. You can simply uncomment the appropriate lines 
-	to set up the dependencies!
+	dependencies set up in the *.build.cs files but commented out. You can uncomment the appropriate lines 
+	to set up the dependencies.
 
 		// Uncomment if you are using Slate UI
 		// PrivateDependencyModuleNames.AddRange(new string[] { "Slate", "SlateCore" });
