Fix for Jira UEDOC-1847
GIF to YouTube video.

diff --git a/Engine/Documentation/Source/Engine/Blueprints/UserGuide/Functions/BlueprintFunctions.INT.udn b/Engine/Documentation/Source/Engine/Blueprints/UserGuide/Functions/BlueprintFunctions.INT.udn
--- a/Engine/Documentation/Source/Engine/Blueprints/UserGuide/Functions/BlueprintFunctions.INT.udn
+++ b/Engine/Documentation/Source/Engine/Blueprints/UserGuide/Functions/BlueprintFunctions.INT.udn
@@ -273,7 +273,17 @@ We could then hook up the rest of the script to destroy the projectile after it
 
 And if we were to play in the editor, we would see something similar to below. 
 
-![](ExampleGIF.gif)(convert:false)
+[OBJECT:EmbeddedVideo]
+[PARAMLITERAL:width]
+640
+[/PARAMLITERAL]
+[PARAMLITERAL:height]
+360
+[/PARAMLITERAL]
+[PARAMLITERAL:videoid]
+T8dbefQDqQs
+[/PARAMLITERAL]
+[/OBJECT]
 
 We have hooked the **PlayerHealth** variable up to a **PrintString** node so that we can show its current value. By default its set to 100 and when the player shoots into the wall and it bounces back and hits the player, you can see that the **Take Damage** function is being called and reduces the **PlayerHealth** variable by 10 each time they are hit. 
 
