Jira UEDOCQA-213

Added description for Self Moved on Event Hit.

diff --git a/Engine/Documentation/Source/Engine/Blueprints/UserGuide/Events/Events.INT.udn b/Engine/Documentation/Source/Engine/Blueprints/UserGuide/Events/Events.INT.udn
--- a/Engine/Documentation/Source/Engine/Blueprints/UserGuide/Events/Events.INT.udn
+++ b/Engine/Documentation/Source/Engine/Blueprints/UserGuide/Events/Events.INT.udn
@@ -153,7 +153,7 @@ This event will only execute if the collision settings on both Actors involved i
 | **My Comp** | **PrimitiveComponent** - The Component on the executing Actor that was hit. |
 | **Other** | **Actor** - The other Actor involved in the collision. |
 | **Other Comp** | **PrimitiveComponent** - The component on the other Actor involved in the collision that was hit. |
-| **Self Moved** | **Boolean** - |
+| **Self Moved** | **Boolean** - Used when receiving a hit from another object's movement (if false), the directions of Hit Normal and Hit Impact Normal will be adjusted to indicate force from the other object against the object being hit. |
 | **Hit Location** | **Vector** - The location of contact between the two colliding Actors. |
 | **Hit Normal** | **Vector** - The direction of the collision. |
 | **Normal Impulse** | **Vector** - The force that the Actors collided with. |
