#UE4doc #docs - bolded link

diff --git a/Engine/Documentation/Source/Programming/UnrealArchitecture/Delegates/Multicast/MulticastDelegates.INT.udn b/Engine/Documentation/Source/Programming/UnrealArchitecture/Delegates/Multicast/MulticastDelegates.INT.udn
--- a/Engine/Documentation/Source/Programming/UnrealArchitecture/Delegates/Multicast/MulticastDelegates.INT.udn
+++ b/Engine/Documentation/Source/Programming/UnrealArchitecture/Delegates/Multicast/MulticastDelegates.INT.udn
@@ -8,11 +8,11 @@ Multi-cast delegates have most of the same features as single-cast delegates.  T
 Just like regular delegates, multi-cast delegates can be loaded/saved, and triggered remotely; however, multi-cast delegate functions cannot use return values. They are best used to
 easily pass a collection of delegates around.
 
-[](Programming/UnrealArchitecture/Delegates\Events) are a specialized type of multi-cast delegate with restricted access to the `Broadcast()`, `IsBound()`, and `Clear()` functions.
+**[](Programming/UnrealArchitecture/Delegates\Events)** are a specialized type of multi-cast delegate with restricted access to the `Broadcast()`, `IsBound()`, and `Clear()` functions.
 
 ## Declaring Multi-Cast Delegates
 
-Multi-cast delegates are declared in the same manner as you [declare standard delegates](Programming/UnrealArchitecture/Delegates#DeclaringDelegates) except they use the macro variations specific to multi-cast delegates.
+Multi-cast delegates are declared in the same manner as you **[declare standard delegates](Programming/UnrealArchitecture/Delegates#DeclaringDelegates)** except they use the macro variations specific to multi-cast delegates.
 
 | Declaration Macro | Description |
 | --- | --- |
