#UE4doc #docs - bolded links

diff --git a/Engine/Documentation/Source/Programming/UnrealArchitecture/Delegates/Events/Events.INT.udn b/Engine/Documentation/Source/Programming/UnrealArchitecture/Delegates/Events/Events.INT.udn
--- a/Engine/Documentation/Source/Programming/UnrealArchitecture/Delegates/Events/Events.INT.udn
+++ b/Engine/Documentation/Source/Programming/UnrealArchitecture/Delegates/Events/Events.INT.udn
@@ -4,13 +4,13 @@ Crumbs:%ROOT%, Programming, Programming/UnrealArchitecture, Programming/UnrealAr
 Description:Delegates that can be bound to multiple functions and execute them all at once.
 Version: 4.5
 
-Events are very similar to [multi-cast delegates](Programming/UnrealArchitecture/Delegates\Multicast). However, while any class can bind events, only the class that declares the event may invoke the event's 
+Events are very similar to **[multi-cast delegates](Programming/UnrealArchitecture/Delegates\Multicast)**. However, while any class can bind events, only the class that declares the event may invoke the event's 
 `Broadcast`, `IsBound`, and `Clear` functions.  This means event objects can be exposed in a public interface without worrying about giving external classes access to these sensitive functions. 
 Event use cases include including callbacks in purely abstract classes, and restricting external classes from invoking the `Broadcast`, `IsBound`, and `Clear` functions.
 
 ## Declaring Events
 
-Events are declared in almost the same manner as you [declare multi-cast delegates](Programming/UnrealArchitecture/Delegates/Multicast#DeclaringMulti-castDelegates) except they use the macro variations specific to events.
+Events are declared in almost the same manner as you **[declare multi-cast delegates](Programming/UnrealArchitecture/Delegates/Multicast#DeclaringMulti-castDelegates)** except they use the macro variations specific to events.
 
 | Declaration Macro | Description |
 | --- | --- |
@@ -25,7 +25,7 @@ The first parameter of the `DECLARE_EVENT` macro is the class which will "own" t
 
 ## Binding Events
 
-Binding of events is performed in the same manner as [binding multi-cast delegates](Programming/UnrealArchitecture/Delegates/Multicast#BindingMulti-CastDelegates).
+Binding of events is performed in the same manner as **[binding multi-cast delegates](Programming/UnrealArchitecture/Delegates/Multicast#BindingMulti-CastDelegates)**.
 
 ## Event Execution
 
