#UE4doc #docs - adjusted bold lettering and captilizations

diff --git a/Engine/Documentation/Source/Programming/UnrealArchitecture/Objects/Objects.INT.udn b/Engine/Documentation/Source/Programming/UnrealArchitecture/Objects/Objects.INT.udn
--- a/Engine/Documentation/Source/Programming/UnrealArchitecture/Objects/Objects.INT.udn
+++ b/Engine/Documentation/Source/Programming/UnrealArchitecture/Objects/Objects.INT.udn
@@ -6,10 +6,10 @@ Version: 4.5
 
 [TOC (start:2 end:3)]
 
-## Overview
+
 
 [EXCERPT:ObjectIntro]
-Unreal has a robust system for handling game Objects. The base class for Objects in Unreal is UObject. The UCLASS macro can be used to tag classes derived from UObject so that the UObject handling
+Unreal has a robust system for handling game Objects. The base class for **Objects** in Unreal is `UObject`. The UCLASS macro can be used to tag classes derived from `UObject` so that the UObject handling
 system is aware of them. 
 [/EXCERPT:ObjectIntro]
 It is not required or even appropriate to use this system in all cases, but there are many benefits to doing so, including:
@@ -24,16 +24,16 @@ It is not required or even appropriate to use this system in all cases, but ther
 *	Type information available at runtime
 *	Network replication
 
-For more detail on each of these systems, see the [](Programming\UnrealArchitecture\Objects\Optimizations) documentation.
+For more detail on each of these systems, see the **[](Programming\UnrealArchitecture\Objects\Optimizations)** documentation.
 
-The UClass macro gives the UObject a reference to a UClass that describes its Unreal-based type. Each UClass maintains one Object called the "class default Object", or CDO for short. The CDO is 
-essentially a default "template" Object, generated by the class constructor and unmodified thereafter. Both the UClass and the CDO can be retrieved for a given Object instance, though they should 
-generally be considered read-only. The UClass for an Object instance can get accessed at any time using the GetClass() function.
+The **UCLASS* macro gives the UObject a reference to a UCLASS that describes its Unreal-based type. Each UCLASS maintains one Object called the "class default Object", or CDO for short. The CDO is 
+essentially a default "template" Object, generated by the class constructor and unmodified thereafter. Both the UCLASS and the CDO can be retrieved for a given Object instance, though they should 
+generally be considered read-only. The UCLASS for an Object instance can get accessed at any time using the GetClass() function.
 
-A UClass contains a set of properties and functions that define the class. These are normal C++ functions and variables available to native code, but tagged with Unreal-specific metadata that controls 
-how they behave within the Object system. For more details about the tagging syntax, see the [Programming Reference](Programming\UnrealArchitecture\Reference). 
+A UCLASS contains a set of properties and functions that define the class. These are normal C++ functions and variables available to native code, but tagged with Unreal-specific metadata that controls 
+how they behave within the Object system. For more details about the tagging syntax, see the **[Programming Reference](Programming\UnrealArchitecture\Reference)**. 
 
-Note that a UObject class can also include native-only properties that do not exist within the corresponding UClass. 
+Note that a UObject class can also include native-only properties that do not exist within the corresponding UCLASS. 
 
 
 ## Object Creation
