Minor Updates.

diff --git a/Engine/Documentation/Source/Engine/Blueprints/HowTo/BPHT_1/BPHT_1.INT.udn b/Engine/Documentation/Source/Engine/Blueprints/HowTo/BPHT_1/BPHT_1.INT.udn
--- a/Engine/Documentation/Source/Engine/Blueprints/HowTo/BPHT_1/BPHT_1.INT.udn
+++ b/Engine/Documentation/Source/Engine/Blueprints/HowTo/BPHT_1/BPHT_1.INT.udn
@@ -81,7 +81,11 @@ We will now create some spawn points in our level to launch the shootable rock f
 
 	![](TargetPointTransform.png)
 
-	**Location** for **X** is **7.0m** (the others are 0) and **Rotation** for **Y** is **40.0** (the others are 0). 
+	**Location** for **X** is **7.0m, 0.0m, 0.0m** and **Rotation** is **0.0, 40.0, 0.0**. 
+
+	[REGION:note]
+	By entering 7.0m, it may convert the value to 700cm which is the equivalent value. Check out the [World Scale](Platforms/VR/ContentSetup) section on the VR Best Practies page as it gives a breakdown of how Unreal Units are calculated as well as how to change the way measurements are displayed. 
+	[/REGION]
 
 1.	In the level, with the **Target Point** selected, press **Ctrl+W** to duplicate it and set its **Details** as shown below. 
 
@@ -93,10 +97,12 @@ We will now create some spawn points in our level to launch the shootable rock f
 
 	![](NewDuplicate.png)
 
-	You will now have three spawn points across from the Player Start (controller icon and direction indicated by the Blue Arrow). 
+1.	Select the **PlayerStart** in the level, then press **E** to enter **Rotation** mode and rotate it so it faces the Target Points. 
 
 	![](TargetPointsAdded.png)
 
+	You will now have three spawn points across from the Player Start (controller icon and direction indicated by the Blue Arrow). 
+
 
 ## Creating a HUD
 
