updated docs on CircleDOF

diff --git a/Engine/Documentation/Source/Engine/Rendering/PostProcessEffects/DepthOfField/DepthOfField.INT.udn b/Engine/Documentation/Source/Engine/Rendering/PostProcessEffects/DepthOfField/DepthOfField.INT.udn
--- a/Engine/Documentation/Source/Engine/Rendering/PostProcessEffects/DepthOfField/DepthOfField.INT.udn
+++ b/Engine/Documentation/Source/Engine/Rendering/PostProcessEffects/DepthOfField/DepthOfField.INT.udn
@@ -14,12 +14,12 @@ make the rendering appear more like a photograph or like a movie.
 
 ## Depth of Field Types
 
-There are two methods available for performing depth of field effects in Unreal Engine 4: Bokeh and Gaussian.
+There are three methods available for performing depth of field effects in Unreal Engine 4: Bokeh, Gaussian and Circle.
 
 [REGION:imagetable]
-| ![Bokeh DoF](bokeh.png)(w:455) | ![Gaussian DoF](gaussian.png)(w:455) |
-| :----: | :----: |
-| Bokeh DoF | Gaussian DoF |
+| ![Bokeh DoF](bokeh.png)(w:300) | ![Gaussian DoF](gaussian.png)(w:300) | ![Circle DoF](Circle.png)(w:300) |
+| :----: | :----: | :----: |
+| Bokeh DoF | Gaussian DoF | Circle DoF |
 [/REGION]
 
 
@@ -38,6 +38,14 @@ The actual implementation does this performance intensive effect only in half re
 The Bokeh method is far more costly performance-wise, making cinematics and showcases prime candidates for the Bokeh method as the appealing visuals often outweigh any performance concerns in those situations.
 [/EXCERPT:DepthOfField_Bokeh]
 
+### Circle
+
+[EXCERPT:DepthOfField_Circle]
+The Circle DoF is the newest addition and matches real world cameras more closely. Simialar to the Bokeh DoF you can see circular shaped Bokeh with sharp HDR content. Very large Bokeh are a weakness of this
+algorithm, they appear noisy and not as clear as the BokehDOF can do them.
+The method has good performance (quite better than Bokeh DoF) and good quality, especially in the transition areas.
+[/EXCERPT:DepthOfField_Circle]
+
 ## Implementation
 
 Depth of Field is broken up into three layers - Near, Far, and the Focal Region - each of which are processed separately and then later composited together to achieve the final effect. Objects in the near and far layers - that is, objects not in the focal region - are always fully blurred. Then, these layers are blended with the non-blurred scene. 
@@ -59,7 +67,10 @@ These layers, including the transition regions, can be visualized using the **De
 
 | Property | Description |
 | -------- | ----------- |
-| **Method** | Selects the method, either Bokeh or Gaussian, used for blurring the scene. |
+| **Method** | Selects the method, Bokeh, Gaussian or Circular, used for blurring the scene. |
+| **Depth Blur km for 50%** | todo |
+| **Depth Blur Radius** | Depth blur radius in pixels at 1920x |
+| **Aperture** | todo |
 | **Focal Distance** | The distance in Unreal Units from the camera that acts as the center of the region where the scene is perfectly in focus and no blurring occurs. [INCLUDE:#FocalDistance] |
 | **Focal Region** | The distance in Unreal Units beyond the focal distance where the scene is perfectly in focus and no blurring occurs. [INCLUDE:#FocalRegion] |
 | **Near Transition Range** | The distance in Unreal Units from the focal region on the side nearer to the camera over which the scene transitions from focused to blurred when using Gaussian-based DoF. [INCLUDE:#NearTransitionRange] |
