#UE4 DOCS: fixing UEDOC-1441

diff --git a/Engine/Documentation/Source/Engine/Blueprints/Editor/BlueprintEditorUserGuide.INT.udn b/Engine/Documentation/Source/Engine/Blueprints/Editor/BlueprintEditorUserGuide.INT.udn
--- a/Engine/Documentation/Source/Engine/Blueprints/Editor/BlueprintEditorUserGuide.INT.udn
+++ b/Engine/Documentation/Source/Engine/Blueprints/Editor/BlueprintEditorUserGuide.INT.udn
@@ -72,10 +72,10 @@ If you are completely new to Blueprint editing, or are unsure what type of Bluep
 		Animation Blueprint UI
 	[/PARAM]
 	[PARAM:description]
-		%Engine/Animation/Persona:description%
+		%Engine/Animation/AnimBlueprints/Interface:description%
 	[/PARAM]
 	[PARAM:path]
-		[RELATIVE:Engine/Animation/Persona]
+		[RELATIVE:Engine/Animation/AnimBlueprints/Interface]
 	[/PARAM]
 [/OBJECT]
 [/REGION]
