#UE4doc #docs - removed overview header

diff --git a/Engine/Documentation/Source/Programming/UnrealArchitecture/Objects/Optimizations/UnrealObjectHandling.INT.udn b/Engine/Documentation/Source/Programming/UnrealArchitecture/Objects/Optimizations/UnrealObjectHandling.INT.udn
--- a/Engine/Documentation/Source/Programming/UnrealArchitecture/Objects/Optimizations/UnrealObjectHandling.INT.udn
+++ b/Engine/Documentation/Source/Programming/UnrealArchitecture/Objects/Optimizations/UnrealObjectHandling.INT.udn
@@ -6,7 +6,6 @@ Version: 4.5
 
 [TOC (start:2 end:2)]
 
-## Overview
 Marking classes, properties, and functions with the appropriate macros turns them into UClasses, UProperties, and UFunctions. This gives Unreal Engine access to them, which allows for a 
 number of under-the-hood handling features to be implemented.
 
