#UE4doc #docs: Altered images to match the instructional guide: green -> yellow.

diff --git a/Engine/Documentation/Source/Support/Builds/ReleaseNotes/2013/Oct/ReleaseNotesOctober2013.INT.udn b/Engine/Documentation/Source/Support/Builds/ReleaseNotes/2013/Oct/ReleaseNotesOctober2013.INT.udn
--- a/Engine/Documentation/Source/Support/Builds/ReleaseNotes/2013/Oct/ReleaseNotesOctober2013.INT.udn
+++ b/Engine/Documentation/Source/Support/Builds/ReleaseNotes/2013/Oct/ReleaseNotesOctober2013.INT.udn
@@ -2,7 +2,6 @@ Title: October 2013 Release Notes
 Crumbs:%ROOT%, Support/Builds
 Description:
 Availability:Licensee
-Template:ReleaseNotes.html
 
 [TOC(start:2 end:2)]
 
diff --git a/Engine/Documentation/Source/Support/Builds/ReleaseNotes/2013/Sept/ReleaseNotesSeptember2013.INT.udn b/Engine/Documentation/Source/Support/Builds/ReleaseNotes/2013/Sept/ReleaseNotesSeptember2013.INT.udn
--- a/Engine/Documentation/Source/Support/Builds/ReleaseNotes/2013/Sept/ReleaseNotesSeptember2013.INT.udn
+++ b/Engine/Documentation/Source/Support/Builds/ReleaseNotes/2013/Sept/ReleaseNotesSeptember2013.INT.udn
@@ -2,7 +2,6 @@ Title: September 2013 Release Notes
 Crumbs:%ROOT%, Support/Builds
 Description:
 Availability:Licensee
-Template:ReleaseNotes.html
 
 [TOC(start:2 end:2)]
 
