#UE4doc #docs: Removed extraneous bolding as according to the instructional guide.

diff --git a/Engine/Documentation/Source/Engine/Blueprints/HowTo/BPHT_7/BPHT_7.INT.udn b/Engine/Documentation/Source/Engine/Blueprints/HowTo/BPHT_7/BPHT_7.INT.udn
--- a/Engine/Documentation/Source/Engine/Blueprints/HowTo/BPHT_7/BPHT_7.INT.udn
+++ b/Engine/Documentation/Source/Engine/Blueprints/HowTo/BPHT_7/BPHT_7.INT.udn
@@ -30,19 +30,19 @@ Using what you have learned over the course of this project, try to do the follo
 * In the level, add two more spawn points (perhaps at the corners, or on the left/right side of the player). 
 * In the HUD, add a new category to keep track of the number of "first shot" hits by a player. 
 
-For more information on the topics covered in this guide, and across the entire editor, see the [**Unreal Editor Manual**](Engine).
+For more information on the topics covered in this guide, and across the entire editor, see the [Unreal Editor Manual](Engine).
 
 As for coverage related to the topics in this guide:
 
-* For more information on **Blueprints** in general, see: [**Blueprints Visual Scripting**](Engine/Blueprints)
+* For more information on **Blueprints** in general, see: [Blueprints Visual Scripting](Engine/Blueprints)
 
-* For more information on using Blueprint Communication and when to use it, see: [**Blueprint Communication Usage**](Engine/Blueprints/UserGuide/BlueprintCommsUsage)
+* For more information on using Blueprint Communication and when to use it, see: [Blueprint Communication Usage](Engine/Blueprints/UserGuide/BlueprintCommsUsage)
 
-* For more information on **Event Dispatchers**, see: [**Event Dispatchers**](Engine/Blueprints/UserGuide/EventDispatcher)
+* For more information on **Event Dispatchers**, see: [Event Dispatchers](Engine/Blueprints/UserGuide/EventDispatcher)
 
-* For more information on **Casting**, see: [**Casting**](Engine/Blueprints/UserGuide/CastNodes)
+* For more information on **Casting**, see: [Casting](Engine/Blueprints/UserGuide/CastNodes)
 
-* For more information on creating HUD elements with **Unreal Motion Graphics** (UMG), see: [**Unreal Motion Graphics**](Engine/UMG)
+* For more information on creating HUD elements with **Unreal Motion Graphics** (UMG), see: [Unreal Motion Graphics](Engine/UMG)
 
 
 %Steps%
