#UE4doc #docs - adjusted bold lettering.

diff --git a/Engine/Documentation/Source/Resources/ContentExamples/ImportOptions/1_6/1_6.INT.udn b/Engine/Documentation/Source/Resources/ContentExamples/ImportOptions/1_6/1_6.INT.udn
--- a/Engine/Documentation/Source/Resources/ContentExamples/ImportOptions/1_6/1_6.INT.udn
+++ b/Engine/Documentation/Source/Resources/ContentExamples/ImportOptions/1_6/1_6.INT.udn
@@ -59,7 +59,7 @@ Related: Engine/Content/FBX/Materials
 
 <!-- This information isn't represented anywhere but in this document, it needs to be moved to the content pipeline docs -->
 
-Need a **Skeletal Mesh** that just loops an animation while staying in the same location in the world? You can do this simply by dragging your **Skeletal Mesh** into the world and then dragging the animation you want to play on to that **Skeletal Mesh**.
+Need a **Skeletal Mesh** that just loops an animation while staying in the same location in the world? You can do this simply by dragging your Skeletal Mesh into the world and then dragging the animation you want to play on to that Skeletal Mesh.
 
 
 
