#UE4doc #docs - adjusted bold lettering.

diff --git a/Engine/Documentation/Source/Resources/ContentExamples/ImportOptions/ImportOptions.INT.udn b/Engine/Documentation/Source/Resources/ContentExamples/ImportOptions/ImportOptions.INT.udn
--- a/Engine/Documentation/Source/Resources/ContentExamples/ImportOptions/ImportOptions.INT.udn
+++ b/Engine/Documentation/Source/Resources/ContentExamples/ImportOptions/ImportOptions.INT.udn
@@ -52,14 +52,15 @@ Related: Engine/Content/FBX/Materials
 [TOC(start:2 end:2)]
 
 
-## Overview
+
 The **FBX_Import_Options** map shows some of the options that can be set when importing an FBX asset into Unreal Engine 4. 
 
 When importing a **Static Mesh**, **Skeletal Mesh**, or **Animation** asset in FBX format, you will be presented with the **Import Options** dialogue box as seen below. 
 
 ![](Import_DialogBox.png)
 
-Depending on the type of asset being imported, additional options will become available such as importing an asset along with its materials/textures, LOD options as well as animation and morph target options.
+Depending on the type of asset being imported, additional options will become available such as importing an asset along with its materials/textures, LOD 
+options as well as animation and morph target options.
 
 [region:note]
 You can find additional resources for each example by clicking on the example's name.
@@ -77,9 +78,9 @@ Listed below are the examples provided inside the **FBX_Import_Options** map:
 | **[1.1 Static Mesh](Resources\ContentExamples\ImportOptions\1_1)** | A simiple 200 x 200 box imported with default settings and imported along with its **Materials** and **Textures** from the modeling program it was created with automatically applied. |
 | **[1.2 Import Mesh LODs](Resources\ContentExamples\ImportOptions\1_2)** | An asset which has a cube as its base LOD (level of detail) and a sphere set as LOD1 where the asset changes between the two based on distance to it. |
 | **[1.3 Import Normals](Resources\ContentExamples\ImportOptions\1_3)** | A cube with one smoothing group and custom vertex normals is imported twice, once using the option to **Calculate Normals** and the other using the **Import Normals** option. |
-| **[1.4 Skeletal Mesh](Resources\ContentExamples\ImportOptions\1_4)** | A simple **Skeletal Mesh** with 3 bones that has been imported and creates the mesh asset and the **Skeleton** asset.  |
-| **[1.5 Animation](Resources\ContentExamples\ImportOptions\1_5)** | The process of importing an **Animation** and assigning the **Skeleton** asset that matches the animation. |
-| **[1.6 Drag and drop animation](Resources\ContentExamples\ImportOptions\1_6)** | Demonstrates how an **Animation Sequence** can be dragged-and-dropped onto a matching **Skeleton** asset inside a level in order to have it applied to and play on the asset. |
+| **[1.4 Skeletal Mesh](Resources\ContentExamples\ImportOptions\1_4)** | A simple Skeletal Mesh with 3 bones that has been imported and creates the mesh asset and the **Skeleton** asset.  |
+| **[1.5 Animation](Resources\ContentExamples\ImportOptions\1_5)** | The process of importing an Animation and assigning the Skeleton asset that matches the animation. |
+| **[1.6 Drag and drop animation](Resources\ContentExamples\ImportOptions\1_6)** | Demonstrates how an **Animation Sequence** can be dragged-and-dropped onto a matching Skeleton asset inside a level in order to have it applied to and play on the asset. |
 | **[1.7 Morph Target](Resources\ContentExamples\ImportOptions\1_7)** | A mesh that was imported with **Morph Targets** and a simple animation which shows the mesh morphing between a cube mesh and a circular mesh.  |
 
 
