Updated for Jira UEDOC-1867

diff --git a/Engine/Documentation/Source/Engine/Animation/AnimHowTo/AdditiveAnimations/AdditiveHowTo.INT.udn b/Engine/Documentation/Source/Engine/Animation/AnimHowTo/AdditiveAnimations/AdditiveHowTo.INT.udn
--- a/Engine/Documentation/Source/Engine/Animation/AnimHowTo/AdditiveAnimations/AdditiveHowTo.INT.udn
+++ b/Engine/Documentation/Source/Engine/Animation/AnimHowTo/AdditiveAnimations/AdditiveHowTo.INT.udn
@@ -191,7 +191,17 @@ Finally, we update the way we blend between animations so that our upperbody pla
 
 1.	**Compile** and **Save**, close the **Blueprint** and play in the editor. 
 
-	![](AdditiveGIF.gif)(convert:false)
+	[OBJECT:EmbeddedVideo]
+	[PARAMLITERAL:width]
+	640
+	[/PARAMLITERAL]
+	[PARAMLITERAL:height]
+	360
+	[/PARAMLITERAL]
+	[PARAMLITERAL:videoid]
+	N96shYlF0Aw
+	[/PARAMLITERAL]
+	[/OBJECT]
 
 You should now see that when you hold the **Left Mouse Button** down, the character will play the shooting AnimMontage. You can walk around while shooting, you can even crouch (by holding **C**) and shoot as well. 
 
