#UE4doc #docs - added link to How to create LODs

diff --git a/Engine/Documentation/Source/Resources/ContentExamples/ImportOptions/1_2/1_2.INT.udn b/Engine/Documentation/Source/Resources/ContentExamples/ImportOptions/1_2/1_2.INT.udn
--- a/Engine/Documentation/Source/Resources/ContentExamples/ImportOptions/1_2/1_2.INT.udn
+++ b/Engine/Documentation/Source/Resources/ContentExamples/ImportOptions/1_2/1_2.INT.udn
@@ -11,6 +11,7 @@ Related: Engine/Content/FBX
 Related: Engine/Content/FBX/MorphTargets
 Related: Shared/Editor/FbxErrors
 Related: Engine/Content/FBX/Materials
+Related: Engine/Content/Types/StaticMeshes/HowTo/LODs
 
 
 [VAR:TopicCompact]
