Updates for Jira UEDOC-1928
GIF to YouTube video.

diff --git a/Engine/Documentation/Source/Gameplay/HowTo/Networking/ReplicateFunction/Blueprints/ReplicateFunction_BP.INT.udn b/Engine/Documentation/Source/Gameplay/HowTo/Networking/ReplicateFunction/Blueprints/ReplicateFunction_BP.INT.udn
--- a/Engine/Documentation/Source/Gameplay/HowTo/Networking/ReplicateFunction/Blueprints/ReplicateFunction_BP.INT.udn
+++ b/Engine/Documentation/Source/Gameplay/HowTo/Networking/ReplicateFunction/Blueprints/ReplicateFunction_BP.INT.udn
@@ -212,9 +212,17 @@ What we will do in this example is create a variable that only gets updated on a
 
 1.	**Compile** and **Save**, close the **Level Blueprint** then play in the editor. 
 
-	[![](ExampleImage.png)](MultiplayerExampleVideo.mp4)
-
-	_Click to play video example._
+	[OBJECT:EmbeddedVideo]
+	[PARAMLITERAL:width]
+	640
+	[/PARAMLITERAL]
+	[PARAMLITERAL:height]
+	360
+	[/PARAMLITERAL]
+	[PARAMLITERAL:videoid]
+	KiBrcLbXbYQ
+	[/PARAMLITERAL]
+	[/OBJECT]
 
 	In the video above when the game is launched, pressing P to print the text variable to the screen displays it as "empty" for each of the characters. However when one of the characters enters the trigger box, the text "item added" is displayed. When that character presses P, the text changes to say "has the item" while each of the other characters still displays "empty". When the character leaves the box, the text "item removed" is displayed and when they press P again, the text now changes back to "empty". 
 
