Marking as NoPublish 2 out-of-date UI reference pages that fall below the visit threshold and all links/references

diff --git a/Engine/Documentation/Source/Engine/Animation/AnimBlueprints/Interface/AnimationBlueprintUI.INT.udn b/Engine/Documentation/Source/Engine/Animation/AnimBlueprints/Interface/AnimationBlueprintUI.INT.udn
--- a/Engine/Documentation/Source/Engine/Animation/AnimBlueprints/Interface/AnimationBlueprintUI.INT.udn
+++ b/Engine/Documentation/Source/Engine/Animation/AnimBlueprints/Interface/AnimationBlueprintUI.INT.udn
@@ -12,7 +12,7 @@ The following sections describe the areas of the Animation Blueprint UI.
 
 ![](AnimGraphUI.png)(w:900)
 
-1. **Menu bar**
+1. **Menu Bar**
 1. **Toolbar**
 1. **Viewport**
 1. **Class Defaults panel**
@@ -24,68 +24,8 @@ The following sections describe the areas of the Animation Blueprint UI.
 
 ## Menu Bar
 
-### File
-
-[INCLUDE:Engine/Blueprints/Editor/UIComponents/Menu/#FileMenu]
-
-### Edit
-
-[INCLUDE:Engine/Blueprints/Editor/UIComponents/Menu/#EditMenu]
-
-### Asset
-| Command | Description |
-| --- | --- |
-| [REGION:tablesection]**Actions**[/REGION] ||
-|Find in Content Browser... | Summons the **Content Browser** and navigates to this asset. |
-|View References | Launces the reference viewer showing the selected assets' references. |
-|Reimport... | Reimports the current asset. |
-| [REGION:tablesection]**Animation**[/REGION] ||
-|Set Preview Mesh as Default | Changes the skeletons default preview mesh to the current open preview mesh. The skeleton will require saving after this action. |
-|Remove Unused Bones from Skeleton | Removes any bones from the skeleton that are not used by any of its meshes. The skeleton and associated animations will require saving after this action. |
-|Record to new Animation | Create new animation from currently playing animation. |
-
-### View
-
-[INCLUDE:Engine/Blueprints/Editor/UIComponents/Menu/#ViewMenu]
-
-### Debug
-
-[INCLUDE:Engine/Blueprints/Editor/UIComponents/Menu/#DebugMenu]
-
-### Window
-
-| Command | Description |
-| --- | --- |
-| Toolbar | Shows the Toolbar if it is currently not visible. |
-| Details | Shows the [Details](Engine/Blueprints\Editor/UIComponents\Details) panel if it is currently not visible. |
-| Debug | Shows the **Debug** panel if it is currently not visible. |
-| Palette | Shows the [](Engine/Blueprints\Editor/UIComponents\Palette) panel if it is not currently visible. |
-| My Blueprint | Shows the [](Engine/Blueprints\Editor/UIComponents\MyBlueprint) panel if it is not currently visible. |
-| Compiler Results | Shows the **Compiler Results** panel if it is not currently visible.  |
-| Find Results | Shows the **Find Results** panel if it is not currently visible.  | 
-| Skeleton Tree | Shows the **Skeleton Tree** panel, which displays all the bones in the Skeletal Mesh associated with the Animation Blueprint. |
-| Asset Browser | Shows the **Asset Browser** panel if it is currently closed, and changes focus to it if it is open. |
-| Viewport | Shows the **Preview Viewport** panel if it is not currently visible.  |
-| Animation Notifies | Opens the **Animation Notifies** window, which lists all Animation Notifies currently associated with the Skeletal Mesh's animation assets. |
-| Animation Blueprint Editor | Opens the **Animation Blueprint Editor** window if it is currently not visible. |
-| Asset Override Editor | Opens the **Asset Override Editor** window if it is currently not visible. |
-| Animation SlotName | Opens the **Animation SlotName** window if it is currently not visible. |
-| Animation Slot Group Name | Opens the **Animation Slot Group Name** window if it is currently not visible. |
-| Content Browser | Opens a sub-menu with access to all four **Content Browsers**. |
-| Developer Tools | Opens a sub-menu with access to **Camera Debugger**, **Collision Analyzer**, **Debug Tools**, **Modules**, and **Widget Reflector**. |
-| Blueprint Debugger | Opens the **Blueprint Debugger**. |
-| Class Viewer | Opens the **Class Viewer**.|
-| Device Profile Editor | Opens the **Device Profile Editor**.|
-| Output Log | Opens the **Output Log**. |
-| Undo History | Opens the **Undo History** window. |
-| [REGION:tablesection]**Unreal Frontend**[/REGION] ||
-| Device Manager | Opens the **Device Manager**.|
-| Session Frontend | Opens the **Session Frontend**. |
-| [REGION:tablesection]**Layout**[/REGION] ||
-| Reset Layout | Resets the layout to the default setting (requires automatic restart of the editor). |
-| Save Layout | Saves the current panel layout. |
-| Enable Fullscreen | Enables/Disables the display of the editor in Fullscreen mode. |
-
+The **Menu Bar** gives you access to common editor actions such as saving, undoing and redoing, as well as several Persona-specific menus. For instance, the **Asset** menu allows
+you to change the default preview mesh, remove unused bones from the skeleton, and record animations. 
 
 
 ## Toolbar
diff --git a/Engine/Documentation/Source/Engine/Animation/Persona/MenuBar/PersonaMenuBar.INT.udn b/Engine/Documentation/Source/Engine/Animation/Persona/MenuBar/PersonaMenuBar.INT.udn
--- a/Engine/Documentation/Source/Engine/Animation/Persona/MenuBar/PersonaMenuBar.INT.udn
+++ b/Engine/Documentation/Source/Engine/Animation/Persona/MenuBar/PersonaMenuBar.INT.udn
@@ -1,4 +1,4 @@
-Availability:Public
+Availability:NoPublish
 Title:Persona Menu Bar
 Crumbs: %ROOT%, Engine, Engine/Animation, Engine/Animation/Persona
 Description:User guide for the Menu Bar in the Persona animation editor.
diff --git a/Engine/Documentation/Source/Engine/Animation/Persona/PersonaEditor.INT.udn b/Engine/Documentation/Source/Engine/Animation/Persona/PersonaEditor.INT.udn
--- a/Engine/Documentation/Source/Engine/Animation/Persona/PersonaEditor.INT.udn
+++ b/Engine/Documentation/Source/Engine/Animation/Persona/PersonaEditor.INT.udn
@@ -2,7 +2,7 @@ Availability:Public
 Title:Persona Animation Editor Reference
 Crumbs: %ROOT%, Engine, Engine/Animation
 Description:Editor for previewing and editing AnimBlueprints, AnimationSequences and other animation assets.
-Related: Engine/Animation/Persona/MenuBar
+
 
 [VAR:ModeSize]w:450 convert:false[/VAR]
 
diff --git a/Engine/Documentation/Source/Engine/Blueprints/Editor/UIBreakdowns/ClassBPUI/ClassBPUI.INT.udn b/Engine/Documentation/Source/Engine/Blueprints/Editor/UIBreakdowns/ClassBPUI/ClassBPUI.INT.udn
--- a/Engine/Documentation/Source/Engine/Blueprints/Editor/UIBreakdowns/ClassBPUI/ClassBPUI.INT.udn
+++ b/Engine/Documentation/Source/Engine/Blueprints/Editor/UIBreakdowns/ClassBPUI/ClassBPUI.INT.udn
@@ -36,7 +36,7 @@ The **Blueprint Class** editor contains several tabs by default:
 
 <!----------
 [EXCERPT:Visible]
-1. [Menu](Engine/Blueprints/Editor/UIComponents/Menu)
+1. Menu
 1. [Toolbar](Engine/Blueprints/Editor/UIComponents/Toolbar/ToolbarOptions)
 1. [Components](Engine/Blueprints/UserGuide/Components)
 1. [My Blueprint](Engine/Blueprints/Editor/UIComponents/MyBlueprint)
diff --git a/Engine/Documentation/Source/Engine/Blueprints/Editor/UIComponents/Menu/BlueprintEditorMenu.INT.udn b/Engine/Documentation/Source/Engine/Blueprints/Editor/UIComponents/Menu/BlueprintEditorMenu.INT.udn
--- a/Engine/Documentation/Source/Engine/Blueprints/Editor/UIComponents/Menu/BlueprintEditorMenu.INT.udn
+++ b/Engine/Documentation/Source/Engine/Blueprints/Editor/UIComponents/Menu/BlueprintEditorMenu.INT.udn
@@ -1,4 +1,4 @@
-Availability:Public
+Availability:NoPublish
 Title: Menu
 Crumbs:%ROOT%, Engine/Blueprints, Engine/Blueprints/Editor
 Description:Menu
