#UE4doc #docs - adjusted bold lettering and removed overview header

diff --git a/Engine/Documentation/Source/Resources/ContentExamples/LevelDesign/LevelDesign.INT.udn b/Engine/Documentation/Source/Resources/ContentExamples/LevelDesign/LevelDesign.INT.udn
--- a/Engine/Documentation/Source/Resources/ContentExamples/LevelDesign/LevelDesign.INT.udn
+++ b/Engine/Documentation/Source/Resources/ContentExamples/LevelDesign/LevelDesign.INT.udn
@@ -47,12 +47,12 @@ Related: Engine/Actors/Volumes
 [TOC(start:2 end:2)]
 
 
-## Overview
+
 The **Leveldesign_Workflow** map gives aspiring Level Designers insight into the iterative level design process which, in this example, is broken up into four phases: **Prototype**, **Meshing**, **Lighting**, and **Polish**.  
 
 During each phase, different goals are set where the initial phases concentrate on getting the foundation in place (think quick and dirty to get your level playable) and the later phases look to build upon that foundation (improving what is in place and fine tuning the details).
 
-While inside the editor on this map, you can cycle through each of the phases quickly by pressing the **1 key** for Prototype, **2 key** for Meshing, **3 key** for Lighting, or **4 key** for Polish to jump to each section. You can also press the **5 key** to zoom out to get a clear picture of all four phases along with a brief description of what each phase consists of. 
+While inside the editor on this map, you can cycle through each of the phases quickly by **pressing** the **1 key** for Prototype, **2 key** for Meshing, **3 key** for Lighting, or **4 key** for Polish to jump to each section. You can also press the **5 key** to zoom out to get a clear picture of all four phases along with a brief description of what each phase consists of. 
 [region:note]
 You can find additional resources for each example by clicking on the example's name.
 [/region]
