#UE4doc #docs - adjusted bold lettering.

diff --git a/Engine/Documentation/Source/Resources/ContentExamples/LevelDesign/2/2.INT.udn b/Engine/Documentation/Source/Resources/ContentExamples/LevelDesign/2/2.INT.udn
--- a/Engine/Documentation/Source/Resources/ContentExamples/LevelDesign/2/2.INT.udn
+++ b/Engine/Documentation/Source/Resources/ContentExamples/LevelDesign/2/2.INT.udn
@@ -30,6 +30,6 @@ Related: Engine/UI/LevelEditor/Modes
 
 ![](leveldesign_2.png)
 
-In the meshing pass, the geometries set up with Brushes or simple Static Meshes from the prototyping pass are replaced with finished or nearly finished mesh assets.  This pass can be refined by reimporting art assets as the
+In the meshing pass, the geometries set up with **Brushes** or simple **Static Meshes** from the prototyping pass are replaced with finished or nearly finished mesh assets.  This pass can be refined by reimporting art assets as the
 assets themselves are refined.  The lighting is still simple at this point in the level design process. 
 [/REGION]
