#UE4 Docs: Fixing Blueprint Quick Start descriptions

diff --git a/Engine/Documentation/Source/Engine/Blueprints/QuickStart/2/BPQS_2.INT.udn b/Engine/Documentation/Source/Engine/Blueprints/QuickStart/2/BPQS_2.INT.udn
--- a/Engine/Documentation/Source/Engine/Blueprints/QuickStart/2/BPQS_2.INT.udn
+++ b/Engine/Documentation/Source/Engine/Blueprints/QuickStart/2/BPQS_2.INT.udn
@@ -1,7 +1,7 @@
 Availability: Public
-Title:2. Construct your Launch Pad
+Title:2. Construct your Launchpad
 Crumbs:%ROOT%, Engine, Engine/Blueprints, Engine/Blueprints/QuickStart
-Description:An overview on the most commonly used viewport navigation controls.
+Description:Build your first launchpad in the level.
 
 
 [VAR:Steps]
@@ -58,7 +58,7 @@ If at any time you need to change your Actor's properties, such as to scale up t
 
 ![](LaunchPadActor.png)
 
-Now that you have your Actor the way you want it, we will turn it into a Blueprint. You can add more components inside the Blueprint Editor, and tweak them just like you can here in the level.
+Now that you have your Actor the way you want it, we will turn it into a Blueprint Class. You can add more components inside the Blueprint Editor, and tweak them just like you can here in the level.
 
 
 ---
diff --git a/Engine/Documentation/Source/Engine/Blueprints/QuickStart/3/QS_3.INT.udn b/Engine/Documentation/Source/Engine/Blueprints/QuickStart/3/QS_3.INT.udn
--- a/Engine/Documentation/Source/Engine/Blueprints/QuickStart/3/QS_3.INT.udn
+++ b/Engine/Documentation/Source/Engine/Blueprints/QuickStart/3/QS_3.INT.udn
@@ -1,7 +1,7 @@
 Availability: Public
-Title:3. Convert Your Actor to a Class Blueprint
+Title:3. Convert Your Actor to a Blueprint Class
 Crumbs:%ROOT%, Engine, Engine/Blueprints, Engine/Blueprints/QuickStart
-Description:Creating an empty level via the New Level dialog window.
+Description:Make your launchpad reusable and scriptable by converting it to a Blueprint Class
 
 
 [VAR:Steps]
diff --git a/Engine/Documentation/Source/Engine/Blueprints/QuickStart/4/BPQS_4.INT.udn b/Engine/Documentation/Source/Engine/Blueprints/QuickStart/4/BPQS_4.INT.udn
--- a/Engine/Documentation/Source/Engine/Blueprints/QuickStart/4/BPQS_4.INT.udn
+++ b/Engine/Documentation/Source/Engine/Blueprints/QuickStart/4/BPQS_4.INT.udn
@@ -1,7 +1,7 @@
 Availability: Public
 Title:4. Create Your Starting Point
 Crumbs:%ROOT%, Engine, Engine/Blueprints, Engine/Blueprints/QuickStart
-Description:Placing and Manipulating Actors is the foundation of building any level.
+Description:Your launchpad will launch the Character when you overlap it, so start with creating a Begin Overlap node.
 
 [VAR:Steps]
 [OBJECT:Navigation]
diff --git a/Engine/Documentation/Source/Engine/Blueprints/QuickStart/5/BPQS_5.INT.udn b/Engine/Documentation/Source/Engine/Blueprints/QuickStart/5/BPQS_5.INT.udn
--- a/Engine/Documentation/Source/Engine/Blueprints/QuickStart/5/BPQS_5.INT.udn
+++ b/Engine/Documentation/Source/Engine/Blueprints/QuickStart/5/BPQS_5.INT.udn
@@ -1,7 +1,7 @@
 Availability: Public
 Title:5. Test the Overlapping Actor
 Crumbs:%ROOT%, Engine, Engine/Blueprints, Engine/Blueprints/QuickStart
-Description:With some Actors in the Level, it's now time to add some variation.
+Description:Make sure it is your Character overlapping the launchpad, so you only launch the Character.
 
 [VAR:Steps]
 [OBJECT:Navigation]
diff --git a/Engine/Documentation/Source/Engine/Blueprints/QuickStart/6/BPQS_6.INT.udn b/Engine/Documentation/Source/Engine/Blueprints/QuickStart/6/BPQS_6.INT.udn
--- a/Engine/Documentation/Source/Engine/Blueprints/QuickStart/6/BPQS_6.INT.udn
+++ b/Engine/Documentation/Source/Engine/Blueprints/QuickStart/6/BPQS_6.INT.udn
@@ -1,7 +1,7 @@
 Availability: Public
 Title:6. Launch Your Character
 Crumbs:%ROOT%, Engine, Engine/Blueprints, Engine/Blueprints/QuickStart
-Description:With a number of Actors placed and customized, one of the final steps for creating any level is the Build Process.
+Description:Set up the actual launching behavior with a Launch Character node.
 
 
 
diff --git a/Engine/Documentation/Source/Engine/Blueprints/QuickStart/7/BPQS_7.INT.udn b/Engine/Documentation/Source/Engine/Blueprints/QuickStart/7/BPQS_7.INT.udn
--- a/Engine/Documentation/Source/Engine/Blueprints/QuickStart/7/BPQS_7.INT.udn
+++ b/Engine/Documentation/Source/Engine/Blueprints/QuickStart/7/BPQS_7.INT.udn
@@ -1,7 +1,7 @@
 Availability: Public
 Title:7. On Your Own!
 Crumbs:%ROOT%, Engine, Engine/Blueprints, Engine/Blueprints/QuickStart
-Description:Further reading on the topics covered in the Content Quick Start Guide.
+Description:Add even more to your launchpad, and read more about Blueprints.
 
 
 
