#UE4doc #docs - adjusted bold lettering.

diff --git a/Engine/Documentation/Source/Resources/ContentExamples/Lighting/Lighting.INT.udn b/Engine/Documentation/Source/Resources/ContentExamples/Lighting/Lighting.INT.udn
--- a/Engine/Documentation/Source/Resources/ContentExamples/Lighting/Lighting.INT.udn
+++ b/Engine/Documentation/Source/Resources/ContentExamples/Lighting/Lighting.INT.udn
@@ -50,8 +50,6 @@ Related: Engine/Rendering/LightingAndShadows/LightMobility
 [TOC(start:2 end:2)]
 
 
-## Overview
-
 The **Lighting** map inside the Content Examples project shows the different types of lights that can be used in Unreal Engine 4 as well as the settings that can be adjusted to help create a more realistically lit environment. 
 The map is broken up into several sections where the first two sections cover light types and light mobility and the remaining cover shadowing, falloff and advanced techniques. 
 
@@ -109,7 +107,7 @@ Listed below are each of the sections provided inside the **Lighting** map:
 
 | Example | What is Demonstrated |
 | ------- | ------------ |
-| **[5.1 Source Radius (Static Lights)](Resources\ContentExamples\Lighting\5_1)** | Shows how soft or hard shadows can be created by adjusting the **Source Radius** of a **Static Light**. |
+| **[5.1 Source Radius (Static Lights)](Resources\ContentExamples\Lighting\5_1)** | Shows how soft or hard shadows can be created by adjusting the **Source Radius** of a Static Light. |
 | **[5.2 Lightsource Radius](Resources\ContentExamples\Lighting\5_2)** | The effects of **Source Radius** on **Stationary** or **Movable** lights and how they contribute to specular reflections. |
 | **[5.3 Lightsource Length](Resources\ContentExamples\Lighting\5_3)** | How the shape of a light's reflection can be altered using **Source Length** (e.g. the reflection from a fluorescent light fixture). |
 | **[5.4 Min Roughness](Resources\ContentExamples\Lighting\5_4)** | Demonstrates a way to blur out a highlight that may be too strong in a reflection by adjusting the **Min Roughness** setting. |
@@ -138,7 +136,7 @@ If you cannot see the light radius or shape indicators, be sure to click the **S
 
 
 [EXCERPT:LightMobility]
-There are 3 types of mobility available to objects: **Static**, **Stationary**, and **Movable**. Mobility can be set at the top of the **Details** panel when the object is selected.
+There are 3 types of mobility available to objects: Static, Stationary, and Movable. Mobility can be set at the top of the **Details** panel when the object is selected.
 
 ![](MobilitySetting.png)
 [/EXCERPT:LightMobility]
