TypoCorrection

diff --git a/Engine/Documentation/Source/GettingStarted/DownloadingUnrealEngine/GitHub/GitHub.INT.udn b/Engine/Documentation/Source/GettingStarted/DownloadingUnrealEngine/GitHub/GitHub.INT.udn
--- a/Engine/Documentation/Source/GettingStarted/DownloadingUnrealEngine/GitHub/GitHub.INT.udn
+++ b/Engine/Documentation/Source/GettingStarted/DownloadingUnrealEngine/GitHub/GitHub.INT.udn
@@ -1,4 +1,4 @@
-Availability:NoPublish
+﻿Availability:NoPublish
 Title:Downloading Unreal Engine from GitHub
 Crumbs:%ROOT%, GettingStarted, GettingStarted/DownloadingUnrealEngine
 Description:Downloading the Unreal Engine source code from GitHub.
@@ -94,4 +94,4 @@ In version 4.5.1 and prior, the `Setup.bat` does not exist and the depencies mus
 
 ## Next Steps
 
-After installing the dependencies, you're ready to [build the engine](https://docs.unrealengine.com/latest/INT/Programming/Development/BuildingUnrealEngine/index.html).
+After installing the dependencies, you're ready to [build the engine](Programming/Development/BuildingUnrealEngine).
