TypoCorrection

diff --git a/Engine/Documentation/Source/Engine/AI/EnvironmentQuerySystem/NodeReference/NodeReference.INT.udn b/Engine/Documentation/Source/Engine/AI/EnvironmentQuerySystem/NodeReference/NodeReference.INT.udn
--- a/Engine/Documentation/Source/Engine/AI/EnvironmentQuerySystem/NodeReference/NodeReference.INT.udn
+++ b/Engine/Documentation/Source/Engine/AI/EnvironmentQuerySystem/NodeReference/NodeReference.INT.udn
@@ -1,4 +1,4 @@
-Availability:Public
+﻿Availability:Public
 Title: The  Environment Query System Node Reference
 Crumbs:%ROOT%, Engine, Gameplay/AI, Engine/AI/EnvironmentQuerySystem
 Description:Break downs and descriptions of the various nodes available in the Environment Query System editor.
@@ -279,7 +279,7 @@ This will trace to or from an Item to a Context and return if it hit it or not.
 | Property | %spacer% | Description |
 | --- | --- | :--- |
 | [REGION:tablesection]**Trace Data**[/REGION] |||
-| **Trance Channel** || The channel to trace against, by default Unreal Engine 4 has Visibility and Camera. Trace Channels can be added and adjusted in the **Edit Menu** -> **Project Settings** -> **Physics** -> **Trace Channels** |
+| **Trace Channel** || The channel to trace against, by default Unreal Engine 4 has Visibility and Camera. Trace Channels can be added and adjusted in the **Edit Menu** -> **Project Settings** -> **Physics** -> **Trace Channels** |
 | **Trace Shape** || Line, Sphere, Box, Capsule |
 | **Trace Complex** || If the trace should be against the mesh (complex) or just the simple collision |
 | **Only Blocking Hits** || Whether the trace uses blocking or non-blocking traces in its results. | 
