Fixing title and image

diff --git a/Engine/Documentation/Source/Gameplay/Analytics/Flurry/Flurry.INT.udn b/Engine/Documentation/Source/Gameplay/Analytics/Flurry/Flurry.INT.udn
--- a/Engine/Documentation/Source/Gameplay/Analytics/Flurry/Flurry.INT.udn
+++ b/Engine/Documentation/Source/Gameplay/Analytics/Flurry/Flurry.INT.udn
@@ -1,9 +1,9 @@
 Availability:Public
-Title:Blueprint Analytics Plugin
+Title:Flurry Analytics Provider
 Crumbs: %ROOT%, Gameplay, Gameplay/Analytics
 Description:A set of Blueprint nodes provided to allow you to communicate with analytics services
 
-![](image00.png)
+![](flurry.png)
 
 [Flurry](http://www.flurry.com) is a widely used free analytics service. Because it is so widely used, it can compare your app's data to other applications in the same category. This lets you quickly get a feeling of 
 how your game is doing and areas you might need to do some work. To use their service, you must register at their site, get an application key that uniquely identifies your app, and 
