#UE4doc #docs:
Removed Overview header
Updated Engine version
Bolded a term

diff --git a/Engine/Documentation/Source/Programming/Slate/Overview/SlateOverview.INT.udn b/Engine/Documentation/Source/Programming/Slate/Overview/SlateOverview.INT.udn
--- a/Engine/Documentation/Source/Programming/Slate/Overview/SlateOverview.INT.udn
+++ b/Engine/Documentation/Source/Programming/Slate/Overview/SlateOverview.INT.udn
@@ -2,7 +2,7 @@ Availability:Public
 Title: Slate Overview
 Crumbs:%ROOT%, Programming, Programming/Slate
 Description:Overview of the cross-platform Slate UI framework.
-Version:4.4
+Version:4.9
 
 [REGION:fullwidth]
 ![Slate UI](editor.png)
@@ -10,7 +10,6 @@ Version:4.4
 
 [TOC (start:2)]
 
-## Overview
 
 [EXCERPT:Intro]
 **Slate** is a completely custom and platform agnostic user interface framework that is designed to make building the user interfaces for tools and applications, such as Unreal Editor, or in-game user interfaces fun and efficient. It combines a declarative syntax with the ability to easily design, layout, and style components that allows for easily creating and iterating on UIs.
@@ -20,7 +19,7 @@ The Slate UI solution makes it extremely easy to put together graphical user int
 
 ## Declarative Syntax
 
-Slate's **declarative syntax** makes building UIs accessible to programmers without adding layers of indirection. A complete set of macros are provided to simplify the process of declaring and creating new widgets.
+Slate's declarative syntax makes building UIs accessible to programmers without adding layers of indirection. A complete set of macros are provided to simplify the process of declaring and creating new widgets.
 
 	SLATE_BEGIN_ARGS( SSubMenuButton )
 		: _ShouldAppearHovered( false )
@@ -146,7 +145,7 @@ Slate has support for accepting mouse and keyboard input. It provides a flexible
 
 ## Output
 
-Slate uses target agnostic rendering primitives allowing it to potentially run on any platform. It currently targets the Unreal Engine 4 RHI so it can run on any platform that Unreal Engine 4 runs on.
+Slate uses target agnostic rendering primitives allowing it to potentially run on any platform. It currently targets the Unreal Engine 4 (UE4) Rendering Hardware Interface (RHI) so it can run on any platform that UE4 runs on.
 
 ## Layout Primitives
 
@@ -179,17 +178,17 @@ The layout above creates the UI below:
 
 ## User Driven Layout
 
-Slate's **docking framework** puts the power into the user's hands by allowing its tabbed panels to be rearranged and docked into virtually any layout imaginable. The ability to customize their environment lets the user work with the tools they want to work with the way they want to work with them.
+Slate's docking framework puts the power into the user's hands by allowing its tabbed panels to be rearranged and docked into virtually any layout imaginable. The ability to customize their environment lets the user work with the tools they want to work with the way they want to work with them.
 
-**Floating Tab**
+Floating tab:
 
 ![Docking](docking_before.png)
 
-**Drag Tab to Docking Target**
+Drag tab to docking target:
 
 ![Drag to Docking Area](docking_place.png)
 
-**Docked Tab**
+Docked tab:
 
 ![Docked](docking_after.png)
 
