#UE4doc #docs:
Updated Engine version
Bolded a few terms

diff --git a/Engine/Documentation/Source/Programming/Slate/Widgets/SlateWidgets.INT.udn b/Engine/Documentation/Source/Programming/Slate/Widgets/SlateWidgets.INT.udn
--- a/Engine/Documentation/Source/Programming/Slate/Widgets/SlateWidgets.INT.udn
+++ b/Engine/Documentation/Source/Programming/Slate/Widgets/SlateWidgets.INT.udn
@@ -2,7 +2,7 @@ Availability:Public
 Title:Slate Widget Examples
 Crumbs:%ROOT%, Programming, Programming/Slate, Programming/Slate/Widgets
 Description:Layout and widget complexities not demonstrated in the Slate Viewer Widget Gallery.
-Version:4.4
+Version:4.9
 
 [TOC (start:2)]
 
@@ -56,18 +56,18 @@ The possible alignments are listed below:
 
 ## Box Panels
 
-SHorizontalBox and SVerticalBox are the most common widgets for arranging your layout. These box panels are declared as widgets, but then they have slots inserted into them. SHorizontalBoxes have these slots arranged with the first widget on the left and the last to the right, while SVerticalBoxes have these slots arranged from top to bottom.
+**SHorizontalBox** and **SVerticalBox** are the most common widgets for arranging your layout. These box panels are declared as widgets, but then they have slots inserted into them. SHorizontalBoxes have these slots arranged with the first widget on the left and the last to the right, while SVerticalBoxes have these slots arranged from top to bottom.
 
-SScrollBox functions similar to SVerticalBox, with the exception that it allows vertical scrolling of the child slots.
+**SScrollBox** functions similar to SVerticalBox, with the exception that it allows vertical scrolling of the child slots.
 
 ### Slot Attributes
 
 * Width or Height Settings (The following options are mutually exclusive):
-	* Auto Size (Default) - This will specify that the slot will fill up as much space as it needs, but no more. Alignment within slots does not matter here.
-	* Fill Size - Specifying a fill coefficient will cause it to fill up the space based on the fill coefficients of other slots. Alignment of the slot in the same orientation does not matter here.
-* Max Size - This will specify the maximum size that the slot can be in slate units.
-* Padding - This will specify the padding of the slot within the panel.
-* Alignment - This will determine how the child widget will be aligned within the slot. This option is meaningless with Fill Size specified in the same orientation.
+	* **Auto Size** (Default) - This will specify that the slot will fill up as much space as it needs, but no more. Alignment within slots does not matter here.
+	* **Fill Size** - Specifying a fill coefficient will cause it to fill up the space based on the fill coefficients of other slots. Alignment of the slot in the same orientation does not matter here.
+* **Max Size** - This will specify the maximum size that the slot can be in slate units.
+* **Padding** - This will specify the padding of the slot within the panel.
+* **Alignment** - This will determine how the child widget will be aligned within the slot. This option is meaningless with Fill Size specified in the same orientation.
 
 
 Below is an example of both autosized horizontal boxes and fill sized horizontal boxes nested within a scrollbox. Also shown is how the alignment can be used with these slots.
@@ -108,7 +108,7 @@ Below is an example of both autosized horizontal boxes and fill sized horizontal
 
 ### Uniform Grid Panels
 
-SUniformGridPanel is a panel that distributes its child widgets evenly both vertically and horizontally. Its child slots are specified using a pair of integers which specify the index of the child. Below is such a uniform grid panel.
+**SUniformGridPanel** is a panel that distributes its child widgets evenly both vertically and horizontally. Its child slots are specified using a pair of integers which specify the index of the child. Below is such a uniform grid panel.
 
 	SNew(SUniformGridPanel)
 	.SlotPadding( FMargin( 5.0f ) )
@@ -138,7 +138,7 @@ SUniformGridPanel is a panel that distributes its child widgets evenly both vert
 
 ### Wrap Boxes
 
-SWrapBox is a box that lays out widgets horizontally until these widgets exceed a width, at which point they will be placed on the next line, and so on. Below is an example of such.
+**SWrapBox** is a box that lays out widgets horizontally until these widgets exceed a width, at which point they will be placed on the next line, and so on. Below is an example of such.
 
 	SNew(SWrapBox)
 	.PreferredWidth( 300.f )
@@ -169,7 +169,7 @@ SWrapBox is a box that lays out widgets horizontally until these widgets exceed
 
 ## Radio Buttons
 
-Radio buttons are check boxes in slate, because check boxes require a delegate for how to determine whether they are checked. To make a series of radio buttons, the easiest way to do this is to create an enum which determines which check box is checked. Save an instance of the enum which determines the current choice. Then, for the checking delegate, pass in a function which compares a passed in payload of the correct enum with the current choice. On changing the selection, the current choice needs to be updated.
+**Radio Buttons** are check boxes in slate, because check boxes require a delegate for how to determine whether they are checked. To make a series of radio buttons, the easiest way to do this is to create an enum which determines which check box is checked. Save an instance of the enum which determines the current choice. Then, for the checking delegate, pass in a function which compares a passed in payload of the correct enum with the current choice. On changing the selection, the current choice needs to be updated.
 
 	ERadioChoice CurrentChoice;
 
@@ -198,7 +198,7 @@ Radio buttons are check boxes in slate, because check boxes require a delegate f
 
 To create menus or toolbars, use multiboxes. Commands should use Slate's UI_COMMAND system, though this cannot be done with dynamically generated buttons/controls.
 
-To create a menu, create an FMenuBarBuilder with an FUICommandList passed in. Finally, you can call MakeWidget() on the menu bar builder to get a widget to place.
+To create a menu, create an **FMenuBarBuilder** with an **FUICommandList** passed in. Finally, you can call **MakeWidget()** on the menu bar builder to get a widget to place.
 
 	FMenuBarBuilder MenuBarBuilder( CommandList );
 	{
@@ -226,7 +226,7 @@ In the menus that are created, you can add menu headers, entries, separators, su
 
 ### Context Menus
 	
-To summon one of these menus as a context menu, get the widget generated from an FMenuBuilder, and pass it into the PushMenu function, as shown here.
+To summon one of these menus as a context menu, get the widget generated from an **FMenuBuilder**, and pass it into the **PushMenu()** function, as shown here.
 
 	FSlateApplication::Get().PushMenu(
 		MenuBuilder.MakeWidget(),
@@ -236,7 +236,7 @@ To summon one of these menus as a context menu, get the widget generated from an
 
 ### Toolbars
 
-To create a toolbar, use an FToolBarBuilder instead. For children, you can add tool bar buttons, combo buttons, plain buttons, and pull down menus.
+To create a toolbar, use an **FToolBarBuilder** instead. For children, you can add tool bar buttons, combo buttons, plain buttons, and pull down menus.
 
 	FToolBarBuilder GameToolBarBuilder( InCommandList );
 	{
@@ -260,11 +260,11 @@ To create a toolbar, use an FToolBarBuilder instead. For children, you can add t
 
 ## Item Views
 
-These views take a template argument of a shared pointer to some sort of data. They are populated by a TArray or TSharedPtr's to the data type. Their internal widgets are populated dynamically with the OnGenerateRow or OnGenerateTile delegate passed in, and different widgets are generated for each column.
+These views take a template argument of a shared pointer to some sort of data. They are populated by a **TArray** or **TSharedPtr**'s to the data type. Their internal widgets are populated dynamically with the **OnGenerateRow** or **OnGenerateTile** delegate passed in, and different widgets are generated for each column.
 
 ### List Views
 
-List views are widgets which store a list of child widgets. They are made as SListView<...>.
+List views are widgets which store a list of child widgets. They are made as **SListView<...>**.
 
 	SNew( SListView< TSharedPtr<FTestData> > )
 	.ItemHeight(24)
@@ -293,7 +293,7 @@ List views are widgets which store a list of child widgets. They are made as SLi
 
 ### Tree Views
 
-Tree views, made as STreeView<...>, are similar to list views, except that they support parenting widgets to other widgets in the list. To determine which widgets are parented to which, the OnGetChildren delegate is used to pass back the children of the passed in item.
+Tree views, made as **STreeView<...>**, are similar to list views, except that they support parenting widgets to other widgets in the list. To determine which widgets are parented to which, the **OnGetChildren** delegate is used to pass back the children of the passed in item.
 
 	SNew( STreeView< TSharedPtr<FTestData> > )
 	.ItemHeight(24)
@@ -314,7 +314,7 @@ Tree views, made as STreeView<...>, are similar to list views, except that they
 
 ### Tile Views
 
-Tile Views, made as STileView<...>, are similar to list views, except the child widgets are laid out in a grid rather than a list. Because of this, they do not support columns or column headers.
+Tile views, made as **STileView<...>**, are similar to list views, except the child widgets are laid out in a grid rather than a list. Because of this, they do not support columns or column headers.
 
 	SNew( STileView< TSharedPtr<FTestData> > )
 	.ItemWidth(128)
