#UE4doc #docs:
Added engine version field

diff --git a/Engine/Documentation/Source/Programming/Slate/DetailsCustomization/DetailsCustomization.INT.udn b/Engine/Documentation/Source/Programming/Slate/DetailsCustomization/DetailsCustomization.INT.udn
--- a/Engine/Documentation/Source/Programming/Slate/DetailsCustomization/DetailsCustomization.INT.udn
+++ b/Engine/Documentation/Source/Programming/Slate/DetailsCustomization/DetailsCustomization.INT.udn
@@ -2,6 +2,7 @@ Availability:Public
 Title:Details Panel Customization
 Crumbs:%ROOT%, Programming, Programming/Slate
 Description:Guide to customizing the display of properties in the Details panels within Unreal Editor.
+version: 4.9
 
 The **Details** panel is now fully customizable.  You can rearrange properties via a simple system, or you can fully customize them using [](Programming/Slate).   You can also add other UI to the details using Slate syntax.
 
