#UE4doc #docs:
Added engine version field
Changed the 'Home' link field to link to the overview for the tutorial
Bolded multiple terms
Added the UE4 acronym for multiple uses of Unreal Engine 4

diff --git a/Engine/Documentation/Source/Programming/Tutorials/AutoCamera/1/AutoCamera_1.INT.udn b/Engine/Documentation/Source/Programming/Tutorials/AutoCamera/1/AutoCamera_1.INT.udn
--- a/Engine/Documentation/Source/Programming/Tutorials/AutoCamera/1/AutoCamera_1.INT.udn
+++ b/Engine/Documentation/Source/Programming/Tutorials/AutoCamera/1/AutoCamera_1.INT.udn
@@ -2,6 +2,7 @@ Availability: Public
 Title:1. Place Cameras In The World
 Crumbs:%ROOT%, Programming, Programming/Tutorials, Programming/Tutorials/AutoCamera
 Description:Setting up a world with multiple cameras
+version: 4.9
 
 [VAR:Steps]
 [OBJECT:Navigation]
@@ -10,7 +11,7 @@ Description:Setting up a world with multiple cameras
 	[PARAM:current]
 	[/PARAM]
 	[PARAM:home]
-		[Programming Tutorials Home](Programming\Tutorials\)
+		[](Programming\Tutorials\AutoCamera)
 	[/PARAM]
 	[PARAM:next]
 		[Next Step](Programming\Tutorials\AutoCamera\2 "%Programming\Tutorials\AutoCamera\2:title%")
@@ -20,9 +21,9 @@ Description:Setting up a world with multiple cameras
 
 %Steps%
 
-[REGION:tip]If you are new to **Unreal Engine 4**, you might want to read our [%Programming\QuickStart:title% tutorial](Programming\QuickStart) first. For this tutorial, we will assume you are familiar with creating a project, adding C++ code to it, compiling your code, and adding **Components** to **Actors** in the **Unreal Engine** editor.[/REGION]
+[REGION:tip]If you are new to **Unreal Engine 4** (UE4), you might want to read our [%Programming\QuickStart:title% tutorial](Programming\QuickStart) first. For this tutorial, we will assume you are familiar with creating a project, adding C++ code to it, compiling your code, and adding **Components** to **Actors** in UE4[/REGION]
 
-1. We will begin by creating a new Basic Code project, with starter content, named "HowTo_AutoCamera". The first thing we'll need to do is create two cameras in our world. Since there are multiple ways to set up cameras, we'll go with the two most common ways here. For our first camera, go to the **Modes** tab in the **Placement Browser** and click **Place**, or press Shift-1. In the **All Classes** section, you'll find a **Camera** actor. Drag this into the **Level Editor** and position it so that it has a good view of our scene.
+1. We will begin by creating a new Basic Code project, with starter content, named "HowTo_AutoCamera". The first thing we'll need to do is create two cameras in our world. Since there are multiple ways to set up cameras, we'll go with the two most common ways here. For our first camera, go to the **Modes** tab in the **Placement Browser** and click **Place**, or press **Shift+1**. In the **All Classes** section, you'll find a **Camera** actor. Drag this into the **Level Editor** and position it so that it has a good view of our scene.
 
 	[![](PlaceCameraActor.png)](PlaceCameraActor.png)
 
