UEDOC-2006

Replaced gif with a YouTube video

diff --git a/Engine/Documentation/Source/Engine/Blueprints/HowTo/BPHT_6/BPHT_6.INT.udn b/Engine/Documentation/Source/Engine/Blueprints/HowTo/BPHT_6/BPHT_6.INT.udn
--- a/Engine/Documentation/Source/Engine/Blueprints/HowTo/BPHT_6/BPHT_6.INT.udn
+++ b/Engine/Documentation/Source/Engine/Blueprints/HowTo/BPHT_6/BPHT_6.INT.udn
@@ -76,7 +76,17 @@ Our game is almost complete, let's finish it off by updating the HUD.
 
 1. From the Main Editor view, click the **Play** button to play in the editor. 
 
-	![](GameplayGIF.gif)(convert:false)
+	[OBJECT:EmbeddedVideo]
+	[PARAMLITERAL:width]
+	640
+	[/PARAMLITERAL]
+	[PARAMLITERAL:height]
+	360
+	[/PARAMLITERAL]
+	[PARAMLITERAL:videoid]
+	tbwor-dRyEs
+	[/PARAMLITERAL]
+	[/OBJECT]
 
 
 That's it; your simple shooting gallery game is complete. Over the course of this guide, we have learned how to communicate between three different types of Blueprints using an **Event Dispatcher**, **Cast** nodes, and **Binding**. We have learned how to pass information from the level to a Character Blueprint as well as how to communicate variables from the character to a HUD. 
