UEDOC-2004

gifs have been replaced with YouTube videos

diff --git a/Engine/Documentation/Source/Resources/Showcases/Reflections/Reflections.INT.udn b/Engine/Documentation/Source/Resources/Showcases/Reflections/Reflections.INT.udn
--- a/Engine/Documentation/Source/Resources/Showcases/Reflections/Reflections.INT.udn
+++ b/Engine/Documentation/Source/Resources/Showcases/Reflections/Reflections.INT.udn
@@ -65,7 +65,17 @@ This example shows off the real-time reflection capabilities of Unreal Engine 4.
 
 The following image shows the reflections in this scene blending on and off to demonstrate their effect.
 
-![](ReflectionsOnOff.gif)(convert:false)
+[OBJECT:EmbeddedVideo]
+[PARAMLITERAL:width]
+640
+[/PARAMLITERAL]
+[PARAMLITERAL:height]
+360
+[/PARAMLITERAL]
+[PARAMLITERAL:videoid]
+JZaQVznNBcA
+[/PARAMLITERAL]
+[/OBJECT]
 
 
 ## Reflection Environment
@@ -96,7 +106,17 @@ Screen space reflections are an engine feature that aid in the grounding of obje
 Screen space reflections are active by default in Unreal Engine 4, but can be toggled using the console command _r.SSR.Quality 3_ or _r.SSR.Quality 0_ depending on the set.
 Below is the results of the reflection example scene both with (r.SSR set to 1) and without screen space reflections (r.SSR set to 0).
 
-![](SSROnOff.gif)(convert:false)
+[OBJECT:EmbeddedVideo]
+[PARAMLITERAL:width]
+640
+[/PARAMLITERAL]
+[PARAMLITERAL:height]
+360
+[/PARAMLITERAL]
+[PARAMLITERAL:videoid]
+DUe4z3-5WvU
+[/PARAMLITERAL]
+[/OBJECT]
 
 
 ## Reflection Considerations
