UEDOC-2002

Gifs have been replaced with YouTube Videos

diff --git a/Engine/Documentation/Source/Resources/Showcases/MatineeFightScene/MatineeFightScene.INT.udn b/Engine/Documentation/Source/Resources/Showcases/MatineeFightScene/MatineeFightScene.INT.udn
--- a/Engine/Documentation/Source/Resources/Showcases/MatineeFightScene/MatineeFightScene.INT.udn
+++ b/Engine/Documentation/Source/Resources/Showcases/MatineeFightScene/MatineeFightScene.INT.udn
@@ -289,7 +289,17 @@ You can **Right-Click** on a keyframe and "flip" the toggle to the opposite of i
 
 Two **SpotLights** were added to the scene to provide additional lighting of the characters for each shot. Each of these **Spotlights** were set to **Movable** and are positioned using the **Movement** track. Additionally, the **Intensity**, **LightColor**, and **Radius** for each light is adjusted during the scene depending upon the shot. To better describe what is occurring, take a look at the image below. 
 
-![](ShowingLightsGif.gif)(w:480 convert:false)
+[OBJECT:EmbeddedVideo]
+[PARAMLITERAL:width]
+640
+[/PARAMLITERAL]
+[PARAMLITERAL:height]
+360
+[/PARAMLITERAL]
+[PARAMLITERAL:videoid]
+kfJeFFGuEPQ
+[/PARAMLITERAL]
+[/OBJECT]
 
 When looking at the scene from the editor view and not the **Matinee Preview**, you can see that lights are being toggled on/off and positioned based on the needs of the shot. 
 
@@ -609,11 +619,27 @@ Be aware of Screen Direction. If you establish a character that appears on the l
 Try to make use of the Rule of Thirds in some of your shots. This is a general photography guideline that proposes that an image should be imagined as divided into nine equal parts by two equally spaced horizontal lines and two equally spaced vertical lines, and that important compositional elements should be placed along these lines or their intersections. 
 
 [REGION:imagetable]
-| ![](RuleOfThirds-SideBySide.gif)(w:480 convert:false) |![](RuleofThirdsMatinee.png)(w:320)|
+| [INCLUDE:#v11] |![](RuleofThirdsMatinee.png)(h:240)|
 | :---: | :---: |
 | Center Shot vs. Rule of Thirds Shot | Matinee Example |
 [/REGION]
 
+<!--
+[EXCERPT:v11]
+[OBJECT:EmbeddedVideo]
+[PARAMLITERAL:width]
+460
+[/PARAMLITERAL]
+[PARAMLITERAL:height]
+240
+[/PARAMLITERAL]
+[PARAMLITERAL:videoid]
+xnuP2mepDjs
+[/PARAMLITERAL]
+[/OBJECT]
+[/EXCERPT]
+-->
+
 #### Framing Your Shots
 
 Framing your shots, or drawing attention to the subject of the shot by blocking other parts of the shot with something else in the scene, can give your shots more context and add a sense of depth or layers to the shot. 
