UEDOC-1979

Replacing gif files with Youtube videos

diff --git a/Engine/Documentation/Source/GettingStarted/FromUnity/UnrealEngineforUnityDevelopers.INT.udn b/Engine/Documentation/Source/GettingStarted/FromUnity/UnrealEngineforUnityDevelopers.INT.udn
--- a/Engine/Documentation/Source/GettingStarted/FromUnity/UnrealEngineforUnityDevelopers.INT.udn
+++ b/Engine/Documentation/Source/GettingStarted/FromUnity/UnrealEngineforUnityDevelopers.INT.udn
@@ -99,7 +99,17 @@ The easiest way to get started with C++ in UE4 is to use the editor to **Add Cod
 
 In Unity, a GameObject is a "thing" that can be placed in the world.  The UE4 equivalent is an Actor.  In the Unreal Editor, you can drag a new Empty Actor into the viewport from the Placement panel:
 
-![image alt text](image_4.gif)(convert:false)
+[OBJECT:EmbeddedVideo]
+[PARAMLITERAL:width]
+640
+[/PARAMLITERAL]
+[PARAMLITERAL:height]
+360
+[/PARAMLITERAL]
+[PARAMLITERAL:videoid]
+UXGFtZ71mfQ
+[/PARAMLITERAL]
+[/OBJECT]
 
 You could build a game out of Empty Actors, but UE4 also includes special types of Actors with built-in features, such as a Pawn (for players or AI objects), or Character (for animated creatures.)  Just like Empty Actors, you can drop these special types of Actors down, then add or customize their properties and components.  You will learn more about it later, but for now try to remember that UE4 has a [Gameplay Framework](Gameplay/Framework) that works with these special Actors.
 
@@ -111,7 +121,17 @@ In Unity, you add components to a GameObject to give it functionality.
 
 In UE4, you add components to Actors.  After you have dropped an Empty Actor in your level, click the Add Component button (in the **Details** panel) and choose a component to add.  Here, we are creating a torch by dropping an empty Actor, then adding a mesh component for the base, followed by a light source and then a particle system to create its flame.
 
-![image alt text](image_5.gif)(convert:false)
+[OBJECT:EmbeddedVideo]
+[PARAMLITERAL:width]
+640
+[/PARAMLITERAL]
+[PARAMLITERAL:height]
+360
+[/PARAMLITERAL]
+[PARAMLITERAL:videoid]
+SknellY4I4g
+[/PARAMLITERAL]
+[/OBJECT]
 
 In Unity, a GameObject hold a flat list of components, but in UE4 an Actor actually contains a **hierarchy **of components attached to one another.  You can see this in the example above, where the Light and Particle are attached to the Mesh.  This has some important implications discussed later in [Complex Actors and GameObjects](#compoundobjects).
 
@@ -122,7 +142,17 @@ Unity's workflow is based on prefabs.  In Unity you build a set of GameObjects w
 
 UE4's corresponding workflow is based on Blueprint Classes.  In UE4, you build an Actor with components, select it, and click the "**Blueprint / Add Script**" button (in the **Details** panel).  Then, choose a place to save your Blueprint Class, and click "**Create Blueprint**" to save your new Blueprint Class!
 
-![image alt text](image_6.gif)(convert:false)
+[OBJECT:EmbeddedVideo]
+[PARAMLITERAL:width]
+640
+[/PARAMLITERAL]
+[PARAMLITERAL:height]
+360
+[/PARAMLITERAL]
+[PARAMLITERAL:videoid]
+n35cDTFBl3s
+[/PARAMLITERAL]
+[/OBJECT]
 
 Your new Blueprint Classes can be found in the **Content Browser**.  You can **double-click** to edit them directly, or you can drag and drop them into any level. 
 
