UEDOC-1957

Gif has been replaced with YouTube videos

diff --git a/Engine/Documentation/Source/Engine/UMG/UserGuide/WidgetTypeReference/WrapBox/WrapBox.INT.udn b/Engine/Documentation/Source/Engine/UMG/UserGuide/WidgetTypeReference/WrapBox/WrapBox.INT.udn
--- a/Engine/Documentation/Source/Engine/UMG/UserGuide/WidgetTypeReference/WrapBox/WrapBox.INT.udn
+++ b/Engine/Documentation/Source/Engine/UMG/UserGuide/WidgetTypeReference/WrapBox/WrapBox.INT.udn
@@ -41,15 +41,45 @@ The white images above could represent different items that the player could buy
 
 For example, "buying" and removing items would automatically update the window: 
 
-![](Gif1.gif)(w:520 convert:false)
+[OBJECT:EmbeddedVideo]
+[PARAMLITERAL:width]
+640
+[/PARAMLITERAL]
+[PARAMLITERAL:height]
+360
+[/PARAMLITERAL]
+[PARAMLITERAL:videoid]
+50iC91BMqr0
+[/PARAMLITERAL]
+[/OBJECT]
 
 As would "selling" and adding items (automatically wrapping when reaching the max width of the window): 
 
-![](GIF2.gif)(w:520 convert:false)
+[OBJECT:EmbeddedVideo]
+[PARAMLITERAL:width]
+640
+[/PARAMLITERAL]
+[PARAMLITERAL:height]
+360
+[/PARAMLITERAL]
+[PARAMLITERAL:videoid]
+IN6MhFKqyiI
+[/PARAMLITERAL]
+[/OBJECT]
 
 It could also be used to keep things inside it arranged when resizing the window: 
 
-![](GIF3.gif)(w:520 convert:false)
+[OBJECT:EmbeddedVideo]
+[PARAMLITERAL:width]
+640
+[/PARAMLITERAL]
+[PARAMLITERAL:height]
+360
+[/PARAMLITERAL]
+[PARAMLITERAL:videoid]
+8zKkdb4hees
+[/PARAMLITERAL]
+[/OBJECT]
 
 While the **Wrap Box** only checks width to determine when to arrange the Widgets within it, you could also use this in combination with a **Scroll Box** which will add a scroll bar when the amount of Widgets within the window becomes cut off allowing the user to scroll down to see the remaining Widgets. 
 
