UEDOC-1956

Replacing Gif files with Youtube videos

diff --git a/Engine/Documentation/Source/Engine/UMG/UserGuide/WidgetTypeReference/MenuAnchor/MenuAnchor.INT.udn b/Engine/Documentation/Source/Engine/UMG/UserGuide/WidgetTypeReference/MenuAnchor/MenuAnchor.INT.udn
--- a/Engine/Documentation/Source/Engine/UMG/UserGuide/WidgetTypeReference/MenuAnchor/MenuAnchor.INT.udn
+++ b/Engine/Documentation/Source/Engine/UMG/UserGuide/WidgetTypeReference/MenuAnchor/MenuAnchor.INT.udn
@@ -47,7 +47,17 @@ The **Menu Anchor** widget has some widget specific functions that can be called
 
 Say you wanted to create an Inventory Menu and have a pop-up appear when the player clicks on a item, as seen in the example below. 
 
-![](PopUpGIF.gif)(convert:false)
+[OBJECT:EmbeddedVideo]
+[PARAMLITERAL:width]
+640
+[/PARAMLITERAL]
+[PARAMLITERAL:height]
+360
+[/PARAMLITERAL]
+[PARAMLITERAL:videoid]
+t0yMkviWKUk
+[/PARAMLITERAL]
+[/OBJECT]
 
 The first thing you would want to do is create what your pop-up menu looks like. 
 
