UEDOC-1950

Replaced gif with YouTube video

diff --git a/Engine/Documentation/Source/Engine/UMG/UserGuide/Animation/Animation.INT.udn b/Engine/Documentation/Source/Engine/UMG/UserGuide/Animation/Animation.INT.udn
--- a/Engine/Documentation/Source/Engine/UMG/UserGuide/Animation/Animation.INT.udn
+++ b/Engine/Documentation/Source/Engine/UMG/UserGuide/Animation/Animation.INT.udn
@@ -71,7 +71,17 @@ Below are two example animations used to mimic a simple **Main Menu** screen whe
 This example is for demonstrating how to set up animation in UMG only as the Buttons do not do anything when clicked on. 
 [/REGION]
 
-![](UMGMenuGIF.gif)(convert:false)
+[OBJECT:EmbeddedVideo]
+[PARAMLITERAL:width]
+640
+[/PARAMLITERAL]
+[PARAMLITERAL:height]
+360
+[/PARAMLITERAL]
+[PARAMLITERAL:videoid]
+R_8nsjHZ9ak
+[/PARAMLITERAL]
+[/OBJECT]
 
 1.	Set up your Main Menu however you see fit and include a **Vertical Box** containing **Buttons** on your screen. 
 
