UEDOC-1949

Replacing gif with Youtube video

diff --git a/Engine/Documentation/Source/Engine/UMG/UserGuide/Anchors/Anchors.INT.udn b/Engine/Documentation/Source/Engine/UMG/UserGuide/Anchors/Anchors.INT.udn
--- a/Engine/Documentation/Source/Engine/UMG/UserGuide/Anchors/Anchors.INT.udn
+++ b/Engine/Documentation/Source/Engine/UMG/UserGuide/Anchors/Anchors.INT.udn
@@ -147,7 +147,17 @@ One more adjustment to the Anchor Medallion can fix that. Here we have stated ho
 
 And now as we resize the Canvas Panel in any direction, the Image resizes and stays in position in relation to the Progress Bar which also stretches. 
 
-![](AnchorGIF.gif)(convert:false)
+[OBJECT:EmbeddedVideo]
+[PARAMLITERAL:width]
+640
+[/PARAMLITERAL]
+[PARAMLITERAL:height]
+360
+[/PARAMLITERAL]
+[PARAMLITERAL:videoid]
+zdKgobaHDhs
+[/PARAMLITERAL]
+[/OBJECT]
 
 
 
