UEDOC-1946

Replacing gif with YouTube video

diff --git a/Engine/Documentation/Source/Engine/UMG/QuickStart/2/UMGQS_2.INT.udn b/Engine/Documentation/Source/Engine/UMG/QuickStart/2/UMGQS_2.INT.udn
--- a/Engine/Documentation/Source/Engine/UMG/QuickStart/2/UMGQS_2.INT.udn
+++ b/Engine/Documentation/Source/Engine/UMG/QuickStart/2/UMGQS_2.INT.udn
@@ -162,7 +162,17 @@ Now that we have our visual layout, we can add the hooks to provide the function
 
 1.	Click **Compile** and **Save** then click the **Play** button to play in the editor. 
 
-	![](Example_GIF.gif)(convert:false)
+	[OBJECT:EmbeddedVideo]
+	[PARAMLITERAL:width]
+	640
+	[/PARAMLITERAL]
+	[PARAMLITERAL:height]
+	360
+	[/PARAMLITERAL]
+	[PARAMLITERAL:videoid]
+	-v8nzlvknlE
+	[/PARAMLITERAL]
+	[/OBJECT]
 
 	Our Health, Energy and Ammo values are now displayed on our HUD and reflect our current values from our Character Blueprint. Pressing **Space** will have the character jump and deplete Energy, pressing the **Left Mouse Button** will fire the weapon reducing our Ammo, and pressing **F** will cause us to lose Health (which you could hook up to a damage system at a later date). 
 	
