UEDOC-1938

Replacing gif with youtube video

diff --git a/Engine/Documentation/Source/Engine/UMG/HowTo/Create3DWidgets/Create3DWidgets.INT.udn b/Engine/Documentation/Source/Engine/UMG/HowTo/Create3DWidgets/Create3DWidgets.INT.udn
--- a/Engine/Documentation/Source/Engine/UMG/HowTo/Create3DWidgets/Create3DWidgets.INT.udn
+++ b/Engine/Documentation/Source/Engine/UMG/HowTo/Create3DWidgets/Create3DWidgets.INT.udn
@@ -13,7 +13,17 @@ This is an experimental feature that is currently undergoing development. Some a
 ---
 In the following steps, we are going to create a menu that appears in the game world that can be interacted with as seen below.
 
-![](MenuGIF.gif)(convert:false)
+[OBJECT:EmbeddedVideo]
+[PARAMLITERAL:width]
+640
+[/PARAMLITERAL]
+[PARAMLITERAL:height]
+360
+[/PARAMLITERAL]
+[PARAMLITERAL:videoid]
+lCJkWsMsK84
+[/PARAMLITERAL]
+[/OBJECT]
 
 [EXCERPT:ProjectStart]
 [OBJECT:TopicList]
