UEDOC-1934

Replacing gif file with Youtube video

diff --git a/Engine/Documentation/Source/Engine/Matinee/HowTo/MHT_8/MHT_8.INT.udn b/Engine/Documentation/Source/Engine/Matinee/HowTo/MHT_8/MHT_8.INT.udn
--- a/Engine/Documentation/Source/Engine/Matinee/HowTo/MHT_8/MHT_8.INT.udn
+++ b/Engine/Documentation/Source/Engine/Matinee/HowTo/MHT_8/MHT_8.INT.udn
@@ -211,7 +211,17 @@ In the following steps, we will create a rock that explodes after a timer count
 
 1.	**Click** the **Play** button from the main tool bar to play in the editor. 
 
-	![](MHT8_Gif.gif)(convert:false)
+	[OBJECT:EmbeddedVideo]
+	[PARAMLITERAL:width]
+	640
+	[/PARAMLITERAL]
+	[PARAMLITERAL:height]
+	360
+	[/PARAMLITERAL]
+	[PARAMLITERAL:videoid]
+	8dIuZDHHQa8
+	[/PARAMLITERAL]
+	[/OBJECT]
 
 	_You should see the text countdown and the rock explode with a particle and sound effect, after which fire can be seen in the location where the rock was._
 
