UEDOC-1933

Converted gif file into a Youtube video

diff --git a/Engine/Documentation/Source/Engine/AI/EnvironmentQuerySystem/QuickStart/13/13.INT.udn b/Engine/Documentation/Source/Engine/AI/EnvironmentQuerySystem/QuickStart/13/13.INT.udn
--- a/Engine/Documentation/Source/Engine/AI/EnvironmentQuerySystem/QuickStart/13/13.INT.udn
+++ b/Engine/Documentation/Source/Engine/AI/EnvironmentQuerySystem/QuickStart/13/13.INT.udn
@@ -30,7 +30,17 @@ Now since you setup the **Behavior Tree** [earlier](Engine/AI/EnvironmentQuerySy
 
 ![](play.png)
 
-![](running.gif)(convert:false)
+[OBJECT:EmbeddedVideo]
+[PARAMLITERAL:width]
+640
+[/PARAMLITERAL]
+[PARAMLITERAL:height]
+360
+[/PARAMLITERAL]
+[PARAMLITERAL:videoid]
+BWhnjDIhYwU
+[/PARAMLITERAL]
+[/OBJECT]
 
 
 %Steps%
