UEDOC-1931

Gif has been replaced with a youtube video

diff --git a/Engine/Documentation/Source/Engine/UI/LevelEditor/Viewports/ViewModes/ViewModes.INT.udn b/Engine/Documentation/Source/Engine/UI/LevelEditor/Viewports/ViewModes/ViewModes.INT.udn
--- a/Engine/Documentation/Source/Engine/UI/LevelEditor/Viewports/ViewModes/ViewModes.INT.udn
+++ b/Engine/Documentation/Source/Engine/UI/LevelEditor/Viewports/ViewModes/ViewModes.INT.udn
@@ -503,7 +503,17 @@ TODO
 ## Exposure
 
 [REGION:fullwidth]
-![](Exposure.gif)(convert:false)
+[OBJECT:EmbeddedVideo]
+[PARAMLITERAL:width]
+640
+[/PARAMLITERAL]
+[PARAMLITERAL:height]
+360
+[/PARAMLITERAL]
+[PARAMLITERAL:videoid]
+yhjRn9xZUTs
+[/PARAMLITERAL]
+[/OBJECT]
 [/REGION]
 
 Exposure is a post-processing effect that controls the overall brightness of a scene. This can be set to a fixed value or left to automatic.
