UEDOC-1930

Replacing gif with Youtube video

diff --git a/Engine/Documentation/Source/Engine/Paper2D/Flipbooks/Components/Blueprints/Components_BP.INT.udn b/Engine/Documentation/Source/Engine/Paper2D/Flipbooks/Components/Blueprints/Components_BP.INT.udn
--- a/Engine/Documentation/Source/Engine/Paper2D/Flipbooks/Components/Blueprints/Components_BP.INT.udn
+++ b/Engine/Documentation/Source/Engine/Paper2D/Flipbooks/Components/Blueprints/Components_BP.INT.udn
@@ -125,7 +125,17 @@ This Custom Event uses **SetFlipbook** for the Flipbook Component called **Sprit
 
 The result is the switching between Flipbooks animations as seen below. 
 
-![](Example1GIF.gif)(convert:false)
+[OBJECT:EmbeddedVideo]
+[PARAMLITERAL:width]
+640
+[/PARAMLITERAL]
+[PARAMLITERAL:height]
+360
+[/PARAMLITERAL]
+[PARAMLITERAL:videoid]
+o_8NIonUIEk
+[/PARAMLITERAL]
+[/OBJECT]
 
 [REGION:note]
 You can also see this exhibited inside the **2D Side Scroller** project template from the [Project Browser](Engine/Basics/Projects/Browser) under **New Project**. 
@@ -141,7 +151,17 @@ Above, inside our Character Blueprint we use a **Left Shift** Key Event to updat
 
 The result is a character that moves faster when shift is held (along with the speed of the Flipbook). 
 
-![](Example2GIF.gif)(convert:false)
+[OBJECT:EmbeddedVideo]
+[PARAMLITERAL:width]
+640
+[/PARAMLITERAL]
+[PARAMLITERAL:height]
+360
+[/PARAMLITERAL]
+[PARAMLITERAL:videoid]
+I9Cg9agS12o
+[/PARAMLITERAL]
+[/OBJECT]
 
 
 
