UEDOC-1919

Replaced gif file with Youtube video

diff --git a/Engine/Documentation/Source/Engine/Matinee/HowTo/MHT_7/MHT_7.INT.udn b/Engine/Documentation/Source/Engine/Matinee/HowTo/MHT_7/MHT_7.INT.udn
--- a/Engine/Documentation/Source/Engine/Matinee/HowTo/MHT_7/MHT_7.INT.udn
+++ b/Engine/Documentation/Source/Engine/Matinee/HowTo/MHT_7/MHT_7.INT.udn
@@ -216,7 +216,17 @@ In the following steps, we will have a character jump up onto a platform by appl
 
 1. **Click** the Play button from the main toolbar to play in the editor. 
 
-	![](MHT7_Gif.gif)(convert:false)
+	[OBJECT:EmbeddedVideo]
+	[PARAMLITERAL:width]
+	640
+	[/PARAMLITERAL]
+	[PARAMLITERAL:height]
+	360
+	[/PARAMLITERAL]
+	[PARAMLITERAL:videoid]
+	7KG-7UoVnwc
+	[/PARAMLITERAL]
+	[/OBJECT]
 
 	_You should now see that the **HeroTPP** Actor will jump to the top of the box, along with an explosion sound effect and music playing._
 
