UEDOC-1918

gif has been replaced with a youtube video

diff --git a/Engine/Documentation/Source/Engine/Matinee/HowTo/MHT_6/MHT_6.INT.udn b/Engine/Documentation/Source/Engine/Matinee/HowTo/MHT_6/MHT_6.INT.udn
--- a/Engine/Documentation/Source/Engine/Matinee/HowTo/MHT_6/MHT_6.INT.udn
+++ b/Engine/Documentation/Source/Engine/Matinee/HowTo/MHT_6/MHT_6.INT.udn
@@ -188,7 +188,17 @@ In the following steps, we will create a blinking **Material** on an object and
 
 1.	**Close** the Blueprint, then click the Play button from the main toolbar to play in the editor. 
 
-	![](MHT6_M1_GIF.gif)(w:480 convert:false)
+	[OBJECT:EmbeddedVideo]
+	[PARAMLITERAL:width]
+	640
+	[/PARAMLITERAL]
+	[PARAMLITERAL:height]
+	360
+	[/PARAMLITERAL]
+	[PARAMLITERAL:videoid]
+	ptBxVONkdrU
+	[/PARAMLITERAL]
+	[/OBJECT]
 
 	_By default the chair should remain orange, however when you get close to it, it should blend to and change to the color you selected._
 
@@ -262,7 +272,17 @@ In the following steps, we will create a blinking **Material** on an object and
 
 1.	From the main toolbar click the Play button to play in the editor. 
 
-	![](MHT6_Matinee2Gif.gif)(w:480 convert:false)
+	[OBJECT:EmbeddedVideo]
+	[PARAMLITERAL:width]
+	640
+	[/PARAMLITERAL]
+	[PARAMLITERAL:height]
+	360
+	[/PARAMLITERAL]
+	[PARAMLITERAL:videoid]
+	wp2RU64ssSQ
+	[/PARAMLITERAL]
+	[/OBJECT]
 
 	_You should now see that the steps of the stairs blink back-and-forth between black and white colors._
 
