UEDOC-1917

Replaced gif file with youtube video

diff --git a/Engine/Documentation/Source/Engine/Matinee/HowTo/MHT_5/MHT_5.INT.udn b/Engine/Documentation/Source/Engine/Matinee/HowTo/MHT_5/MHT_5.INT.udn
--- a/Engine/Documentation/Source/Engine/Matinee/HowTo/MHT_5/MHT_5.INT.udn
+++ b/Engine/Documentation/Source/Engine/Matinee/HowTo/MHT_5/MHT_5.INT.udn
@@ -191,7 +191,17 @@ In the following steps, we will create a light that moves, changes colors and br
 
 1. From the main toolbar, **click** the Play icon to play in the editor. 
 
-	![](MHT5_LightingGif.gif)(convert:false)
+	[OBJECT:EmbeddedVideo]
+	[PARAMLITERAL:width]
+	640
+	[/PARAMLITERAL]
+	[PARAMLITERAL:height]
+	360
+	[/PARAMLITERAL]
+	[PARAMLITERAL:videoid]
+	hJu8SiXEoXs
+	[/PARAMLITERAL]
+	[/OBJECT]
 
 	_When you enter the hallway, you should see the light moving up and down the hallway._
 
