UEDOC-1915

Converting gif into Youtube video

diff --git a/Engine/Documentation/Source/Engine/UI/EditorSettings/EditorPreferences.INT.udn b/Engine/Documentation/Source/Engine/UI/EditorSettings/EditorPreferences.INT.udn
--- a/Engine/Documentation/Source/Engine/UI/EditorSettings/EditorPreferences.INT.udn
+++ b/Engine/Documentation/Source/Engine/UI/EditorSettings/EditorPreferences.INT.udn
@@ -73,7 +73,22 @@ The **Editor Preferences** window can be opened from the **Edit** menu.
 
 | Option | Description |
 | ------ | ----------- |
-| **Highlight Objects Under Mouse Cursor** | If enabled, objects will be highlighted using a transparent overlay in viewports while hovering over them with the mouse cursor. [REGION:none]![](selection_hover.gif)(convert:false)[/REGION] |
+| **Highlight Objects Under Mouse Cursor** | If enabled, objects will be highlighted using a transparent overlay in viewports while hovering over them with the mouse cursor. 
+[REGION:none]
+[OBJECT:EmbeddedVideo]
+[PARAMLITERAL:width]
+420
+[/PARAMLITERAL]
+[PARAMLITERAL:height]
+200
+[/PARAMLITERAL]
+[PARAMLITERAL:videoid]
+yvwZUtmW12I
+[/PARAMLITERAL]
+[/OBJECT]
+[/REGION] 
+
+|
 | **Highlight Selected Objects with Brackets** | [INCLUDE:Shared/Types/ULevelEditorViewportSettings#bhighlightwithbrackets] |
 | **Use Selection Outline** | [INCLUDE:Shared/Types/ULevelEditorViewportSettings#buseselectionoutline] |
 | **Selection Highlight Intensity** | [INCLUDE:Shared/Types/ULevelEditorViewportSettings#selectionhighlightintensity] |
