UEDOC-1912

Replaced gif with YouTube video

diff --git a/Engine/Documentation/Source/Engine/Rendering/ParticleSystems/Reference/Modules/Lifetime/LifetimeModules.INT.udn b/Engine/Documentation/Source/Engine/Rendering/ParticleSystems/Reference/Modules/Lifetime/LifetimeModules.INT.udn
--- a/Engine/Documentation/Source/Engine/Rendering/ParticleSystems/Reference/Modules/Lifetime/LifetimeModules.INT.udn
+++ b/Engine/Documentation/Source/Engine/Rendering/ParticleSystems/Reference/Modules/Lifetime/LifetimeModules.INT.udn
@@ -8,7 +8,17 @@ Description:The Lifetime module is used to set the initial lifetime of a particl
 (#Lifetime)
 ## Lifetime
 
-![psys_lifetime.gif](psys_lifetime.gif)(convert:false)
+[OBJECT:EmbeddedVideo]
+[PARAMLITERAL:width]
+640
+[/PARAMLITERAL]
+[PARAMLITERAL:height]
+360
+[/PARAMLITERAL]
+[PARAMLITERAL:videoid]
+nnFIgdgObvU
+[/PARAMLITERAL]
+[/OBJECT]
 
 The Lifetime module is used to set the initial lifetime of a particle at spawn time. It contains the following member(s):
 
