UEDOC-1913

replaced gifs with Youtube videos

diff --git a/Engine/Documentation/Source/Engine/Rendering/ParticleSystems/Reference/Modules/Location/LocationModules.INT.udn b/Engine/Documentation/Source/Engine/Rendering/ParticleSystems/Reference/Modules/Location/LocationModules.INT.udn
--- a/Engine/Documentation/Source/Engine/Rendering/ParticleSystems/Reference/Modules/Location/LocationModules.INT.udn
+++ b/Engine/Documentation/Source/Engine/Rendering/ParticleSystems/Reference/Modules/Location/LocationModules.INT.udn
@@ -9,7 +9,17 @@ Location modules affect the position of particles at various times throughout th
 
 ## Initial Location
 
-![psys_initial_location.gif](psys_initial_location.gif)(convert:false)
+[OBJECT:EmbeddedVideo]
+[PARAMLITERAL:width]
+640
+[/PARAMLITERAL]
+[PARAMLITERAL:height]
+360
+[/PARAMLITERAL]
+[PARAMLITERAL:videoid]
+LAd89U2b-YY
+[/PARAMLITERAL]
+[/OBJECT]
 
 The Initial Location module is used to set the initial location of a particle at spawn time. It contains the following member(s):
 
@@ -172,7 +182,17 @@ The particle used will be the one at the same index as the particle having its l
 
 ## Cylinder
 
-![psys_cylinder.gif](psys_cylinder.gif)(convert:false)
+[OBJECT:EmbeddedVideo]
+[PARAMLITERAL:width]
+640
+[/PARAMLITERAL]
+[PARAMLITERAL:height]
+360
+[/PARAMLITERAL]
+[PARAMLITERAL:videoid]
+sefKMZtdGR0
+[/PARAMLITERAL]
+[/OBJECT]
 
 The Cylinder module is used to set the initial location of particles within a cylinder shape. It contains the following member(s):
 
@@ -231,7 +251,17 @@ The Cylinder (Seeded) module is identical to the [Cylinder](#Cylinder) module in
 
 ## Sphere
 
-![psys_sphere.gif](psys_sphere.gif)(convert:false)
+[OBJECT:EmbeddedVideo]
+[PARAMLITERAL:width]
+640
+[/PARAMLITERAL]
+[PARAMLITERAL:height]
+360
+[/PARAMLITERAL]
+[PARAMLITERAL:videoid]
+ZM6GkFxO-b0
+[/PARAMLITERAL]
+[/OBJECT]
 
 The Sphere module is used to set the initial location of particles within a sphere shape. The class contains the following member(s):
 
