TypoCorrection

diff --git a/Engine/Documentation/Source/Gameplay/HowTo/UsingCameras/Blueprints/UsingCameras_BP.INT.udn b/Engine/Documentation/Source/Gameplay/HowTo/UsingCameras/Blueprints/UsingCameras_BP.INT.udn
--- a/Engine/Documentation/Source/Gameplay/HowTo/UsingCameras/Blueprints/UsingCameras_BP.INT.udn
+++ b/Engine/Documentation/Source/Gameplay/HowTo/UsingCameras/Blueprints/UsingCameras_BP.INT.udn
@@ -1,7 +1,7 @@
-Availability: Public
+﻿Availability: Public
 Title:Using a Static Camera
 Crumbs: %ROOT%, Gameplay, Gameplay\HowTo\UsingCameras
-Description:A How To Guide for Finding Actors in Blueprints. 
+Description:A How To Guide for using a static camera in Blueprints.
 Related: Gameplay\HowTo\UsingCameras
 Related: Gameplay\Framework\Camera
 Related: Engine/Matinee
