TypoCorrection

diff --git a/Engine/Documentation/Source/Gameplay/HowTo/UsingCameras/SpringArmComponents/SpringArm.INT.udn b/Engine/Documentation/Source/Gameplay/HowTo/UsingCameras/SpringArmComponents/SpringArm.INT.udn
--- a/Engine/Documentation/Source/Gameplay/HowTo/UsingCameras/SpringArmComponents/SpringArm.INT.udn
+++ b/Engine/Documentation/Source/Gameplay/HowTo/UsingCameras/SpringArmComponents/SpringArm.INT.udn
@@ -1,7 +1,7 @@
-Availability: Public
+﻿Availability: Public
 Title:Using Spring Arm Components
 Crumbs: %ROOT%, Gameplay, Gameplay\HowTo\UsingCameras
-Description:The Spring Arm component is used to automatically control how the camera handles situations where the it becomes obstructed. 
+Description:The Spring Arm component is used to automatically control how the camera handles situations where it becomes obstructed. 
 Related: Gameplay\HowTo\UsingCameras
 Related: Gameplay\Framework\Camera
 Related: Engine/Matinee
