TypoCorrection

diff --git a/Engine/Documentation/Source/Engine/Physics/PhAT/UserGuide/UserGuide.INT.udn b/Engine/Documentation/Source/Engine/Physics/PhAT/UserGuide/UserGuide.INT.udn
--- a/Engine/Documentation/Source/Engine/Physics/PhAT/UserGuide/UserGuide.INT.udn
+++ b/Engine/Documentation/Source/Engine/Physics/PhAT/UserGuide/UserGuide.INT.udn
@@ -1,4 +1,4 @@
-Availability: Public
+﻿Availability: Public
 Title:Physics Asset Tool User Guide
 Crumbs: %ROOT%, Engine, Engine/Physics, Engine/Physics/PhAT
 Description:How to use the Physics Asset Tool to create Physics Asset for Skeletal Meshes.
@@ -200,7 +200,7 @@ Each Physics Body in a Physics Asset can have a Physical Material assigned to it
 ## Constraints
 
 [REGION:note]
-The use of Physics Constraints is covered in the **[Physics Constraints User Guide](Engine/Physics\Constraints\ConstraintsUserGuide)*8 and their properties are covered in the **[Physics Constraint Reference](Engine/Physics\Constraints\ConstraintsReference)**. This section will only cover workflows that are specific to the Physics Asset Tool or ones that have deviated substantially from the norm.
+The use of Physics Constraints is covered in the **[Physics Constraints User Guide](Engine/Physics\Constraints\ConstraintsUserGuide)** and their properties are covered in the **[Physics Constraint Reference](Engine/Physics\Constraints\ConstraintsReference)**. This section will only cover workflows that are specific to the Physics Asset Tool or ones that have deviated substantially from the norm.
 [/REGION]
 
 ### Editing Physics Constraints
