TypoCorrection

diff --git a/Engine/Documentation/Source/Engine/Physics/Destructibles/DestructibleProperties/DestructibleProperties.INT.udn b/Engine/Documentation/Source/Engine/Physics/Destructibles/DestructibleProperties/DestructibleProperties.INT.udn
--- a/Engine/Documentation/Source/Engine/Physics/Destructibles/DestructibleProperties/DestructibleProperties.INT.udn
+++ b/Engine/Documentation/Source/Engine/Physics/Destructibles/DestructibleProperties/DestructibleProperties.INT.udn
@@ -1,4 +1,4 @@
-Availability:Public
+﻿Availability:Public
 Title:Destructible Properties Matrix
 Crumbs: %ROOT%, Engine
 Description:Descriptions of all the properties on Destructible Actors in the DestructibleMesh Editor.
@@ -99,7 +99,7 @@ For the Physics and Collision section of properties, see the [Collision](Engine/
 | **Is Support Chunk** | Defines the chunk to be environmentally supported, if the appropriate NxDestructibleParametersFlag flags are set in NxDestructibleParameters.  |
 | **Do Not Fracture** | Defines the chunk to be unfracturable. If this is true, none of its children will be fracturable. |
 | **Do Not Damage** | Defines the chunk to be undamagable. This means this chunk will not fracture, but its children might |
-| **Do not Crumble** | Defines the chunk to be uncrumbleable. This means this chunk will not be broken down into fluid mesh particles no matter how much damage it takes. Note: this only applies to chunks with no children. For a chunk with children then: 1) The chunk may be broken down into its children, and then its children may be crumbled, if the doNotCumble Flag is not set on them, 2) If the destructible module's chunk depth offset LOD may be set such that this chunk effectively has no children. In this case, the doNotCrumble flag will apply to it. |
+| **Do not Crumble** | Defines the chunk to be uncrumbleable. This means this chunk will not be broken down into fluid mesh particles no matter how much damage it takes. Note: this only applies to chunks with no children. For a chunk with children then: 1) The chunk may be broken down into its children, and then its children may be crumbled, if the doNotCrumble Flag is not set on them, 2) If the destructible module's chunk depth offset LOD may be set such that this chunk effectively has no children. In this case, the doNotCrumble flag will apply to it. |
 
 	
 ## Fracture Settings Category
