#UE4doc #docs:
Added Engine Version field
Changed Home link to point to the overview for the tutorial
Changed font settings for multiple terms

diff --git a/Engine/Documentation/Source/Programming/Tutorials/AutoCamera/2/AutoCamera_2.INT.udn b/Engine/Documentation/Source/Programming/Tutorials/AutoCamera/2/AutoCamera_2.INT.udn
--- a/Engine/Documentation/Source/Programming/Tutorials/AutoCamera/2/AutoCamera_2.INT.udn
+++ b/Engine/Documentation/Source/Programming/Tutorials/AutoCamera/2/AutoCamera_2.INT.udn
@@ -2,6 +2,7 @@ Availability: Public
 Title:2. Control Camera View In C++
 Crumbs:%ROOT%, Programming, Programming/Tutorials, Programming/Tutorials/AutoCamera
 Description:Coding a class to change camera views
+version: 4.9
 
 [VAR:Steps]
 [OBJECT:Navigation]
@@ -11,7 +12,7 @@ Description:Coding a class to change camera views
 	[PARAM:current]
 	[/PARAM]
 	[PARAM:home]
-		[Programming Tutorials Home](Programming\Tutorials\)
+		[](Programming\Tutorials\AutoCamera)
 	[/PARAM]
 	[PARAM:next]
 		[Next Step](Programming\Tutorials\AutoCamera\3 "%Programming\Tutorials\AutoCamera\3:title%")
@@ -25,7 +26,7 @@ Description:Coding a class to change camera views
 
 	[![](NameActorClass.png)](NameActorClass.png)
 
-1. In `CameraDirector.h`, add the following code to the bottom of the **ACameraDirector** class definition:
+1. In CameraDirector.h, add the following code to the bottom of the ACameraDirector class definition:
 
 		UPROPERTY(EditAnywhere)
 		AActor* CameraOne;
@@ -37,7 +38,7 @@ Description:Coding a class to change camera views
 
 	[REGION:note]The **UPROPERTY** macro makes our variables visible to **Unreal Engine**. This way, values set in these variables will not be reset when we launch the game or reload our level or project in a future work session. We have also added the **EditAnywhere** keyword, which will allow us to set CameraOne and CameraTwo in the **Unreal Editor**.[/REGION]
 
-1. In `CameraDirector.cpp`, add the following line of code to the top of the file, right underneath the other "#include" lines:
+1. In CameraDirector.cpp, add the following line of code to the top of the file, right underneath the other "#include" lines:
 
 		#include "Kismet/GameplayStatics.h"
 
