#UE4doc #docs:
Added Engine Version field
Changed Home link to point to the overview for the tutorial
Changed font settings for multiple terms
Changed color of border indicator in picture from green to yellow

diff --git a/Engine/Documentation/Source/Programming/Tutorials/AutoCamera/3/AutoCamera_3.INT.udn b/Engine/Documentation/Source/Programming/Tutorials/AutoCamera/3/AutoCamera_3.INT.udn
--- a/Engine/Documentation/Source/Programming/Tutorials/AutoCamera/3/AutoCamera_3.INT.udn
+++ b/Engine/Documentation/Source/Programming/Tutorials/AutoCamera/3/AutoCamera_3.INT.udn
@@ -2,6 +2,7 @@ Availability: Public
 Title:3. Place A Camera Director In The World
 Crumbs:%ROOT%, Programming, Programming/Tutorials, Programming/Tutorials/AutoCamera
 Description:Placing a Camera Director in the world and assigning Cameras to it
+version:4.9
 
 [VAR:Steps]
 [OBJECT:Navigation]
@@ -11,7 +12,7 @@ Description:Placing a Camera Director in the world and assigning Cameras to it
 	[PARAM:current]
 	[/PARAM]
 	[PARAM:home]
-		[Programming Tutorials Home](Programming\Tutorials\)
+		[](Programming\Tutorials\AutoCamera)
 	[/PARAM]
 	[PARAM:next]
 		[Next Step](Programming\Tutorials\AutoCamera\4 "%Programming\Tutorials\AutoCamera\4:title%")
