#UE4doc #docs: fixed minor grammar/spelling errors and removed broken characters: ?

diff --git a/Engine/Documentation/Source/Support/Builds/ReleaseNotes/2013/Sept/ReleaseNotesSeptember2013.INT.udn b/Engine/Documentation/Source/Support/Builds/ReleaseNotes/2013/Sept/ReleaseNotesSeptember2013.INT.udn
--- a/Engine/Documentation/Source/Support/Builds/ReleaseNotes/2013/Sept/ReleaseNotesSeptember2013.INT.udn
+++ b/Engine/Documentation/Source/Support/Builds/ReleaseNotes/2013/Sept/ReleaseNotesSeptember2013.INT.udn
@@ -37,7 +37,7 @@ Availability:Licensee
 	* Linear color curve support has been added to timelines and users now have a gradient editor for them (editing the curve directly is still supported though). These can be created in or imported to the content browser as well as created directly in the Blueprint timeline editor:
 	
 		![](LinearColor01.png)
-	* When users open a color curve by default they�ll see a gradient which can be used to add keys while visualizing how the color changes over time:
+	* When users open a color curve by default they will see a gradient which can be used to add keys while visualizing how the color changes over time:
 	
 		![](LinearColor02.png)
 	
@@ -85,7 +85,7 @@ Availability:Licensee
 
 		![](AutoInsertNodes.png)
 
-* **Improved �Add Component� Dropdown**
+* **Improved 'Add Component' Dropdown**
 	* This now has a search box, _one-click_ add, icons, and categories.
 
 		![](AddComponentDropdown.png)             
@@ -101,7 +101,7 @@ Availability:Licensee
 		![](BetterHandlingforMultipleSelectedActors.png)
 
 * **New Library Functions** 
-	* New Blueprint Library functions have ben added!
+	* New Blueprint Library functions have been added!
 
 		![](NewLibraryFunctions.png)
 	
@@ -181,7 +181,7 @@ Availability:Licensee
 
 
 #### Blueprints
-* Added `ExpandEnumAsExecs` metadata support. This allows an function enum param to be exposed as multiple exec inputs.
+* Added `ExpandEnumAsExecs` metadata support. This allows a function enum param to be exposed as multiple exec inputs.
 * Added **GetAllActorsOfClass**, **GetActorArrayAverageLocation**, **GetActorArrayBounds**, **MoveActorToActor** to `GameplayStatics`
 * Added **MoveComponentTo** latent function to interpolate a components location/rotation over time
 * Additional support for `FText` in _Blueprints_
@@ -246,11 +246,11 @@ Availability:Licensee
 * Screen space Light Shafts
 	* Two separate light shaft types: occlusion which only affects fog / atmosphere, and bloom which is not physically based
 * Added **VisualizeOutOfBoundsPixels** show flag, which colors pixels that are outside of their bounds in a striped pattern
-* Reflection environment cubemaps are cached in the DerivedDataCache, so that they don�t need to be recaptured every load
+* Reflection environment cubemaps are cached in the DerivedDataCache, so that they don't need to be recaptured every load
 * Cubemap data is cooked into packages
-	* Unfortunately there are some failure cases, as the derived data can only be generated by a GPU, so it can�t be generated during cooking
-* Diffuse GI on movable translucency � indirect lighting cache interpolates a single value per object
-* New translucency lighting mode � _TLM_Surface_, which gets the nearest reflection capture cubemap applied for indirect specular
+	* Unfortunately there are some failure cases, as the derived data can only be generated by a GPU, so it can't be generated during cooking
+* Diffuse GI on movable translucency indirect lighting cache interpolates a single value per object
+* New translucency lighting mode _TLM_Surface_, which gets the nearest reflection capture cubemap applied for indirect specular
 * UnrealBuildTool allows multiple instances to build separate branches concurrently
 	* Previously new instances of UnrealBuildTool would block until the running instance finished. You can now run multiple instances if they are in separate branches (different UE4 root folder paths).
 * **Film Post**
@@ -288,7 +288,7 @@ Availability:Licensee
 	* These files are loadable into Gephi, which can create interesting visualizations of the graph network 
 * Added `UE4PrereqSetup.msi` installer
 	* Replaces UE4Redist.exe and installs the prerequisites needed for running UE4
-* Behavior Trees was moved from game to engine level and it�s available for licensees
+* Behavior Trees was moved from game to engine level and it is available for licensees
 * Basic editor for Behavior Trees was added to editor
 * Navigation data exporter (`ExportNavigation` console command) was added to engine, to export collision geometry (from navigation point of view) to OBJ file for navigation debugging purposes,
 
@@ -423,15 +423,15 @@ Availability:Licensee
 	* We no longer include all platforms by default. Instead, you only get either Win64 or Mac. 
 	* If you need to build other platforms, there is a new `GenerateProjectFiles_AllPlatforms.bat` file for that. 
 	* Solution has flattened down to just four top-levels folders. 
-	* All code is now under your game�s project file, organized the same way it is on the disk. 
+	* All code is now under your game's project file, organized the same way it is on the disk. 
 	* Other notes 
 	* Shipping and Test configurations are omitted by default. Pass `-ShippingConfigs` to get those. 
-	* Rocket targets aren�t included in the project files by default. Pass `-RocketTargets` to get those. 
-	* Non-English UDN files are now omitted from project files. Pass �-AllLanguages� to get those. 
+	* Rocket targets aren't included in the project files by default. Pass `-RocketTargets` to get those. 
+	* Non-English UDN files are now omitted from project files. Pass '-AllLanguages' to get those. 
 	* Visual Studio no longer prompts you to overwrite `UE4.sln` every time you exit 
 	* You can now skip IntelliSense generation with `-NoIntelliSense` if you want. 
 	* WinRT is no longer considered a _desktop platform_ to UBT 
-	* We removed Shipping and Test configs for Editor from the UI, because they aren�t supported in UE4 
+	* We removed Shipping and Test configs for Editor from the UI, because they aren't supported in UE4 
 * Changed DestroyActor to only detach SceneComponents attached to SceneComponents on the deleted actor. 
 * Changed markdown process to leave variable and metadata refs in code blocks, spans and quotes ignored. 
 * Improved regex used to handle PUBLISH tags to allow whitespace on the same line after tags 
@@ -442,10 +442,10 @@ Availability:Licensee
 		* PlayerReplicationInfo is now PlayerState
 	* ReplicationInfo base class was removed.
 	* Redirectors are in place and all uses of these classes should update seamlessly.
-* Added exception check to error that shouldnt occur (but does) in CRU 
+* Added exception check to error that should not occur (but does) in CRU 
 * Added domain exception handler to help track down obscure crash
 * **UE4Redist has been replaced with UE4PrereqSetup**
-	* To install prerequisites for running UE4 on Windows platform, we�ll now use the new `UE4PrereqSetup.msi` program. `UE4Redist.exe` should no longer be used and has been deleted from P4.  
+	* To install prerequisites for running UE4 on Windows platform, we'll now use the new `UE4PrereqSetup.msi` program. `UE4Redist.exe` should no longer be used and has been deleted from P4.  
 	* In Perforce, UE4PrereqSetup is currently located in `//depot/UE4/Engine/Extras/Redist/en-us/UE4PrereqSetup.msi`
 	* You should never really need to uninstall it, but if you need to just use the Control Panel to remove it like any other program.
  
@@ -563,7 +563,7 @@ Availability:Licensee
 * Collapsed graphs can be renamed inside the graph node.
 * The **Component** tab viewport in the **Blueprint Editor** now ticks particles.
 * Added a right-click context menu to the **Timeline Editor** for improved usage.
-* The transform gizmo is no longer visible in the preview viewport when a root component node is selected in **Components** mode within the **Blueprint eEitor**.
+* The transform gizmo is no longer visible in the preview viewport when a root component node is selected in **Components** mode within the **Blueprint editor**.
 * Made access to the _Blueprint_ preview actor more efficient in **Components** mode of the **Blueprint Editor**.
 * Multifonts have been removed.
 * Improved debug instance selection when opening _Blueprints_. Editing a _Blueprint_ from the asset menu will find an instance to debug if there is one. Additionally the debug instance should now update correctly when editing subsequent _Blueprint_ instances.
