#UE4doc #docs: removed broken characters: ?

diff --git a/Engine/Documentation/Source/Support/Builds/ReleaseNotes/2014/4_1/ReleaseNotes_4_1.INT.udn b/Engine/Documentation/Source/Support/Builds/ReleaseNotes/2014/4_1/ReleaseNotes_4_1.INT.udn
--- a/Engine/Documentation/Source/Support/Builds/ReleaseNotes/2014/4_1/ReleaseNotes_4_1.INT.udn
+++ b/Engine/Documentation/Source/Support/Builds/ReleaseNotes/2014/4_1/ReleaseNotes_4_1.INT.udn
@@ -255,7 +255,7 @@ We'd like to thank all of the amazing developers who contributed code in this re
 * **DiCoJamesTan (James Tan)**
     * [Fix attachment bug with ChildActorComponent ](https://github.com/EpicGames/UnrealEngine/pull/20)
 
-* **fxtentacle (Hajo Nils Krabbenh�ft)**
+* **fxtentacle (Hajo Nils Krabbenh.ft)**
     * [Fix Azimuth & Elevation calculation on Mac](https://github.com/EpicGames/UnrealEngine/pull/9)
 
 * **Kamayuq**
