#UE4doc #docs: fixed broken characters: ?

diff --git a/Engine/Documentation/Source/Support/Builds/ReleaseNotes/2014/4_4/ReleaseNotes_4_4.INT.udn b/Engine/Documentation/Source/Support/Builds/ReleaseNotes/2014/4_4/ReleaseNotes_4_4.INT.udn
--- a/Engine/Documentation/Source/Support/Builds/ReleaseNotes/2014/4_4/ReleaseNotes_4_4.INT.udn
+++ b/Engine/Documentation/Source/Support/Builds/ReleaseNotes/2014/4_4/ReleaseNotes_4_4.INT.udn
@@ -370,7 +370,7 @@ The friction mode can be set project wide or overwritten per physical material.
 * Cascade Mesh emitter improvements
 	* Mesh emitters that don't use PSA_Velocity or PSA_AwayFromCenter had their rotation rate increase over time, ending up with them spinning uncontrollably after a few seconds, and order of operations between 
 	* Orientation and Scaling of particles wasn't as expected. These problems are fixed with this update.
-	* Additionally, MediumDetailSpawnRateScale has been deprecated and replaced with QualityLevelSpawnRateScale. This new value is tied to the scalability settings' effects quality level, and results in the spawn rate being multiplied by QualityLevelSpawnRateScale for every reduction of effects quality level from Epic. For example, if an emitter�s spawn rate is authored to 100, and its QualityLevelSpawnRateScale is 0.5, reducing effects quality to High will result in a spawn rate of 50, going to Medium will set spawn rate to 25. Setting QualityLevelSpawnRateScale to 1.0 results in no change across effects quality changes.
+	* Additionally, MediumDetailSpawnRateScale has been deprecated and replaced with QualityLevelSpawnRateScale. This new value is tied to the scalability settings' effects quality level, and results in the spawn rate being multiplied by QualityLevelSpawnRateScale for every reduction of effects quality level from Epic. For example, if an emitter's spawn rate is authored to 100, and its QualityLevelSpawnRateScale is 0.5, reducing effects quality to High will result in a spawn rate of 50, going to Medium will set spawn rate to 25. Setting QualityLevelSpawnRateScale to 1.0 results in no change across effects quality changes.
 * Particle effect meshes now adhere to the view mode correctly
 * Beams now pass in a second UV set which covers 0 to 1 in U across the length of the beam. UV set 0 now tiles according to the tiling parameters in the typedata module.
 * Fixed issue where turning on GBuffer Hints in the main viewport crashes the Editor.
@@ -416,7 +416,7 @@ The friction mode can be set project wide or overwritten per physical material.
 * New: Get Display Name now available in Blueprints for class variables.
 * New: InstancedStaticMeshComponents now set the FHitResult.Item property with the index an instance hit by a collision event. This can be used with the above functions to remove a specific instance and replace it with a real StaticMeshComponent to provide interactivity with Foliage or other behavior.
 * New: Introduced UNavigationPath for BP usage - not much one can do with it just yet.
-* New: Instigator can now be set as part of Spawn Actor Blueprint node.  If the Instigator is not explicitly set via a wired in pin, the Instigator will automatically be based on the calling context. If Spawn Actor is called from a Pawn context, the Pawn will be the Instigator. If Spawn Actor is called from an actor context, the actor�s Instigator will be the Instigator.
+* New: Instigator can now be set as part of Spawn Actor Blueprint node.  If the Instigator is not explicitly set via a wired in pin, the Instigator will automatically be based on the calling context. If Spawn Actor is called from a Pawn context, the Pawn will be the Instigator. If Spawn Actor is called from an actor context, the actor's Instigator will be the Instigator.
 * New: PawnActions have been introduced in the AIModule. Early access level currently.
 * New: The Curve Table editor now displays data in a grid widget.
 * New: Touch events (but not touch over events) are now enabled by default on Player Controllers.
@@ -487,7 +487,7 @@ The friction mode can be set project wide or overwritten per physical material.
 **New: Game Instance objects (persistent gameplay data)**
 * New: The Game Instance is a transient object that has a lifetime that exists across level transitions. 
 * It is game extendable (set in the Maps & Modes section of the Project Settings) as both a C++ and Blueprint class.  
-* It can be used to store cross-level data, manage the game�s application state, and act as an owner for objects that you would like to exist across the entire game session�s lifetime.  
+* It can be used to store cross-level data, manage the game's application state, and act as an owner for objects that you would like to exist across the entire game session's lifetime.  
 * The GameInstance exists on both the client and the server, however, because it is a UObject there is currently no replication of data between the instances.
 
 #### Paper 2D 
@@ -558,7 +558,7 @@ The friction mode can be set project wide or overwritten per physical material.
 * New: Updated Mac to p4api-2014.1 so that we may support SSL connections.
 * New: Updated Mac OS X version check to warn users if their version is older than 10.9.4 plus fixed the check so it no longer detects 10.10 as older than 10.9.1
 * Changed trackpad detection on Mac to report trackpad only if at least two fingers are touching it. This should limit issues with people accidentally touching trackpad while using mouse + WASD
-* Don�t disable use of glMapBuffer when using MTGL as there�s an OpenGL bug that results in rendering errors in that case.
+* Do not disable use of glMapBuffer when using MTGL as there is an OpenGL bug that results in rendering errors in that case.
 * Exposed the amount of VRAM in the GPU to the OpenGLRHI on OS X.
 * Improved the code for enumerating installed engine versions in UE4EditorServices so that it doesn't include incomplete or deleted installs.
 * Improved RMB handling in SNodePanel for Magic Mouse and trackpad.
@@ -571,7 +571,7 @@ The friction mode can be set project wide or overwritten per physical material.
 * Only flush after glTexStorage calls when necessary to avoid inhibiting performance on Macs with unaffected GPUs.
 * Removed the GL_TEXTURE_BUFFER workaround for OS X the OpenGL driver fix is in 10.9.4.
 	* All Macs will use the GL_TEXTURE_BUFFER path now as everything but a 2013 Mac Pro running 10.9.3 will work correctly.
-* Stopped using Cocoa�s child window mechanism since it isn't what we want.
+* Stopped using Cocoa's child window mechanism since it isn't what we want.
 	* Tear-out or panel-style windows are raised to float above normal windows.
 	* All windows appear in Expose/Spaces/Mission Control & the are part of system window cycling.
 	* Modal windows are treated slightly specially to ensure that the last shown modal window currently still open remains on top and attempts to change focus fail with visual indication.
@@ -677,7 +677,7 @@ See our [Android Device compatibility wiki page](https://wiki.unrealengine.com/A
 	* In blueprint editor case it was not possible to modify constraint frame of reference. This new value allows to set it manually.
 * New: Exposed max depenetration velocity for more slowly resolving overlapping simulated objects.
 * Removed enabling collision on no collision object for drag/drop operation. Now no collision should not respond to Visibility trace.
-* Simplified navigation cluster graph to work only as fast �path exists� check.
+* Simplified navigation cluster graph to work only as fast 'path exists' check.
 * Stopping movement from MovementComponent will abort path following.
 * Reenabled RT destruction in APEX. We are still not supporting this directly, but the APEX libs will have this flag turned back on.
 * Overlap check performance: Blueprint overlap queries now check the async scene.
@@ -714,7 +714,7 @@ See our [Android Device compatibility wiki page](https://wiki.unrealengine.com/A
 	* Supports displaying your web cam in the editor and in-game
 * New: Experimental support for streaming live broadcasts from Unreal Editor (not ready to use yet)
 * New: Experimental Twitch broadcasting support (not ready for use yet)
-	* You�ll need to acquire the Twitch SDK and compile the engine yourself to try this out
+	* You will need to acquire the Twitch SDK and compile the engine yourself to try this out
 	* Implements live video/audio broadcasting to Twitch
 	* Supports displaying your web cam in the editor and game
 	* Uses browser-based authentication to login to Twitch
@@ -788,7 +788,7 @@ See our [Android Device compatibility wiki page](https://wiki.unrealengine.com/A
 * New: Added a PostEngineInit load phase for modules, which allows modules to be registered after GEngine/GEditor have been created.
 * New: Ability to override config options on the commandline. Use the following format:
 	* -ini:IniName:[Section1]:Key1=Value1,[Section2]:Key2=Value2
-* New: Added file system block alignment option in Unreal pak. Prevents small files crossing a file system block boundary which can incur a performance hit on some platforms. This option help reduce that cost. Use �blocksize=X[KB|MB] to enable.
+* New: Added file system block alignment option in Unreal pak. Prevents small files crossing a file system block boundary which can incur a performance hit on some platforms. This option help reduce that cost. Use blocksize=X[KB|MB] to enable.
 * Added ENGINE_API to FComponentInstanceDataBase so it could be used in other modules.
 * Added ISourceControlLabel::Sync() overload that takes multiple files.
 * New: Added a proper error message if a file cannot be deleted during project upgrade, and added functionality to checkout the project file if necessary.
@@ -820,14 +820,14 @@ See our [Android Device compatibility wiki page](https://wiki.unrealengine.com/A
 * Removed ScriptAsset, ScriptFunction and ScriptActor classes from ScriptPlugin
 * Removed DLLEXPORT specification from a few templates in CoreUObject module that could cause failure if inline expansion is turned off.
 * If a tooltip does not exist for a UField, use the name of the field instead.
-* Fixed �stat unitgraph� not appearing in console autocomplete.
+* Fixed 'stat unitgraph' not appearing in console autocomplete.
 * Fixed copying files under source control not maintaining branch history.
 * Fixed an issue with not displaying self/otherchildren when using 'stat dumpframe'.
 * Fixed a bug that cause low-frequency GC to happen at inappropriate times.
 * Fixed crash when project module fails to load.
 * Fixed compiler errors when brute forcing STAT to be true in Shipping builds.
 * Fixed for multithreaded crash in ValidateHeap() & dumpAllocatorStats() in FMallocBinned.
-* Fixed not being able to toggle fullscreen mode when launched with �windowed on the command line.
+* Fixed not being able to toggle fullscreen mode when launched with windowed on the command line.
 * Fixed compile-on-startup when plugin modules don't already exist on startup.
 * Fixed compiling from the editor where the target does not match the game name.
 * Fixed scraping of UnrealBuildTool environment variables in non-US locales. 
@@ -878,7 +878,7 @@ See our [Android Device compatibility wiki page](https://wiki.unrealengine.com/A
 
 
 #### UnrealHeaderTool
-* Removed 'ScriptPlugins' ini setting. Plugins will be discovered by their 'Category' specified in uplugin file (for UHT it should be prefixed with �UnrealHeaderTool.Code Generator.�)
+* Removed 'ScriptPlugins' ini setting. Plugins will be discovered by their 'Category' specified in uplugin file (for UHT it should be prefixed with 'UnrealHeaderTool.Code Generator.')
 
 ## Known Issues
 * There is a known issue with IPP where if a user tries to import a mobile provision with a concrete application id (com.yourcompany.gamename as opposed to com.yourcompany.*) then it will fail to properly find the imported mobile provision.  The workarounds for this are as follows: 
