#UE4doc #docs:
Changes weren't saved before previous submit.

diff --git a/Engine/Documentation/Source/Programming/Tutorials/AutoCamera/4/AutoCamera_4.INT.udn b/Engine/Documentation/Source/Programming/Tutorials/AutoCamera/4/AutoCamera_4.INT.udn
--- a/Engine/Documentation/Source/Programming/Tutorials/AutoCamera/4/AutoCamera_4.INT.udn
+++ b/Engine/Documentation/Source/Programming/Tutorials/AutoCamera/4/AutoCamera_4.INT.udn
@@ -2,7 +2,7 @@ Availability: Public
 Title:4. On Your Own!
 Crumbs:%ROOT%, Programming, Programming/Tutorials, Programming/Tutorials/AutoCamera
 Description:Ideas for further exploration!
-
+version:4.9
 [VAR:Steps]
 [OBJECT:Navigation]
 	[PARAM:previous]
@@ -11,7 +11,7 @@ Description:Ideas for further exploration!
 	[PARAM:current]
 	[/PARAM]
 	[PARAM:home]
-		[Programming Tutorials Home](Programming\Tutorials\)
+		[](Programming\Tutorials\AutoCamera)
 	[/PARAM]
 	[PARAM:next]
 	[/PARAM]
@@ -22,13 +22,13 @@ Description:Ideas for further exploration!
 
 Using what you have learned, try to do the following:
 
-* Attach a **Camera** to a moving **Actor** to create a crane or dolly shot.
-* Use a single [array](Programming/UnrealArchitecture/TArrays) variable to store your cameras, instead of _CameraOne_ and _CameraTwo_, so you can go through a sequence of any number of cameras, instead of just two.
+* Attach a Camera to a moving Actor to create a crane or dolly shot.
+* Use a single [array](Programming/UnrealArchitecture/TArrays) variable to store your cameras, instead of CameraOne and CameraTwo, so you can go through a sequence of any number of cameras, instead of just two.
 * Instead of using [Actor](Programming/UnrealArchitecture/Actors) pointers to store your cameras, create a [structure](Programming/UnrealArchitecture/Reference/Structs) that holds the pointer, as well as time before changing the view, and blend time to the new view.
 
 As for the specifics covered in this tutorial:
 
-* For more information on **Cameras** and ways to control them, see the [Camera](Gameplay/Framework/Camera) framework page, or try the [%Programming/Tutorials/PlayerCamera:title%](Programming/Tutorials/PlayerCamera) tutorial.
+* For more information on Cameras and ways to control them, see the [Camera](Gameplay/Framework/Camera) framework page, or try the [%Programming/Tutorials/PlayerCamera:title%](Programming/Tutorials/PlayerCamera) tutorial.
 * For further tutorials, see the [%Programming/Tutorials:title%](Programming/Tutorials) page.
 
 
