#UE4doc #docs:
Added Engine Version field
Changed Home link to point to the overview for the tutorial
Changed font settings on a few terms

diff --git a/Engine/Documentation/Source/Programming/Tutorials/Components/1/Components_1.INT.udn b/Engine/Documentation/Source/Programming/Tutorials/Components/1/Components_1.INT.udn
--- a/Engine/Documentation/Source/Programming/Tutorials/Components/1/Components_1.INT.udn
+++ b/Engine/Documentation/Source/Programming/Tutorials/Components/1/Components_1.INT.udn
@@ -2,7 +2,7 @@ Availability: Public
 Title:1. Creating And Attaching Components
 Crumbs:%ROOT%, Programming, Programming/Tutorials, Programming/Tutorials/Components
 Description:Creating a variety of components and attaching them to a customized pawn.
-
+version:4.9
 [VAR:Steps]
 [OBJECT:Navigation]
 	[PARAM:previous]
@@ -10,7 +10,7 @@ Description:Creating a variety of components and attaching them to a customized
 	[PARAM:current]
 	[/PARAM]
 	[PARAM:home]
-		[Programming Tutorials Home](Programming\Tutorials\)
+		[](Programming\Tutorials\Components\)
 	[/PARAM]
 	[PARAM:next]
 		[Next Step](Programming\Tutorials\Components\2 "%Programming\Tutorials\Components\2:title%")
@@ -20,30 +20,30 @@ Description:Creating a variety of components and attaching them to a customized
 
 %Steps%
 
-[REGION:tip]If you are new to **Unreal Engine 4**, you might want to read our [%Programming\QuickStart:title% tutorial](Programming\QuickStart) first. For this tutorial, we will assume you are familiar with creating a project and adding C++ code to it, as well as configuring input in the **Unreal Editor**. If you are unfamiliar with creating your own **Pawn** class or configuring input, the [%Programming\Tutorials\PlayerInput:title% tutorial](Programming\Tutorials\PlayerInput) might be a good place to start.[/REGION]
+[REGION:tip]If you are new to **Unreal Engine 4** (UE4), you might want to read our [%Programming\QuickStart:title% tutorial](Programming\QuickStart) first. For this tutorial, we will assume you are familiar with creating a project and adding C++ code to it, as well as configuring input in the **Unreal Editor**. If you are unfamiliar with creating your own **Pawn** class or configuring input, the [%Programming\Tutorials\PlayerInput:title% tutorial](Programming\Tutorials\PlayerInput) might be a good place to start.[/REGION]
 
-1. We will begin by creating a new, Basic Code project, with starter content, named "HowTo_Components". The first thing we'll need to add to this project will be a customized **Pawn** that will hold our **Components**, move around our level, and collide with solid objects. For this tutorial, we'll name it "CollidingPawn".
+1. We will begin by creating a new, Basic Code project, with starter content, named HowTo_Components. The first thing we'll need to add to this project will be a customized Pawn that will hold our **Components**, move around our level, and collide with solid objects. For this tutorial, we'll name it CollidingPawn.
 
 	[![](ChooseParentClass.png)](ChooseParentClass.png)
 
 	[![](NamePawnClass.png)](NamePawnClass.png)
 
-1. In **Visual Studio**, we should open `CollidingPawn.h` and add the following code to the bottom of our class definition:
+1. In **Visual Studio**, we should open CollidingPawn.h and add the following code to the bottom of our class definition:
 
 		UParticleSystemComponent *OurParticleSystem;
 
-	We will use this variable to keep track of a **Particle System Component** that we will create later. We can create **Components** without making variables to track them, but if we want to use those **Components** in our code, we should store them in class member variables like this.
+	We will use this variable to keep track of a **Particle System Component** that we will create later. We can create Components without making variables to track them, but if we want to use those Components in our code, we should store them in class member variables like this.
 
-1. We can now open `CollidingPawn.cpp` and edit the constructor function, **ACollidingPawn::ACollidingPawn**, by adding code that will spawn a variety of useful **Components** and arrange them in a hierarchy. We will create a **Sphere Component** to interact with the physical world, a **Static Mesh Component** to represent our collision shape visually, a **Particle System Component** that we can turn on or off at will, and a **Spring Arm Component** that we can use to attach a **Camera Component** to control our in-game perspective.
+1. We can now open CollidingPawn.cpp and edit the constructor function, **ACollidingPawn::ACollidingPawn**, by adding code that will spawn a variety of useful Components and arrange them in a hierarchy. We will create a **Sphere Component** to interact with the physical world, a **Static Mesh Component** to represent our collision shape visually, a Particle System Component that we can turn on or off at will, and a **Spring Arm Component** that we can use to attach a **Camera Component** to control our in-game perspective.
 
-1. First up is deciding which **Component** will be the root of our hierarchy. In this tutorial, we'll go with the **Sphere Component**, since it is the physical presence that can interact and collide with our game world. Note that an **Actor** can have multiple physics-enabled **Components** in its hierarchy, but for this tutorial, we only need one.
+1. First up is deciding which Component will be the root of our hierarchy. In this tutorial, we'll go with the Sphere Component, since it is the physical presence that can interact and collide with our game world. Note that an **Actor** can have multiple physics-enabled Components in its hierarchy, but for this tutorial, we only need one.
 			// Our root component will be a sphere that reacts to physics
 			USphereComponent* SphereComponent = CreateDefaultSubobject<USphereComponent>(TEXT("RootComponent"));
 			RootComponent = SphereComponent;
 			SphereComponent->InitSphereRadius(40.0f);
 			SphereComponent->SetCollisionProfileName(TEXT("Pawn"));
 	
-1. Next, we will create and attach a visible sphere from a **Static Mesh** asset that has a radius of 50. Since this doesn't perfectly line up with the 40-radius **Sphere Component** we just created, we'll scale it down to 80%. We also need to move it down by 40 units in order to have its center line up with the center of the **Sphere Component**.
+1. Next, we will create and attach a visible sphere from a **Static Mesh** asset that has a radius of 50. Since this doesn't perfectly line up with the 40-radius Sphere Component we just created, we'll scale it down to 80%. We also need to move it down by 40 units in order to have its center line up with the center of the Sphere Component.
 			// Create and position a mesh component so we can see where our sphere is
 			UStaticMeshComponent* SphereVisual = CreateDefaultSubobject<UStaticMeshComponent>(TEXT("VisualRepresentation"));
 			SphereVisual->AttachTo(RootComponent);
@@ -55,9 +55,9 @@ Description:Creating a variety of components and attaching them to a customized
 				SphereVisual->SetWorldScale3D(FVector(0.8f));
 			}
 
-	[REGION:note]Although you can see the asset location of our **Static Mesh** asset in the code, hard-coding asset paths is not generally considered the best way to load assets. It is usually preferred to create the **Component** itself in code if the class needs to use it, and then select the assets in the **Unreal Editor**. However, it is possible to do this directly in code, and could be faster for programmers debugging or building new features.[/REGION]
+	[REGION:note]Although you can see the asset location of our Static Mesh asset in the code, hard-coding asset paths is not generally considered the best way to load assets. It is usually preferred to create the Component itself in code if the class needs to use it, and then select the assets in the **Unreal Editor**. However, it is possible to do this directly in code, and could be faster for programmers debugging or building new features.[/REGION]
 
-1. We can now attach an inactive **Particle System Component** to our hierarchy. We can manipulate this component in code, and will later set up an input to turn it on or off. Notice that the **Particle System Component** is attached directly to the **Static Mesh Component**, not to the root. It is also slightly offset from the bottom-center of the mesh so that it will be more visible during play.
+1. We can now attach an inactive Particle System Component to our hierarchy. We can manipulate this component in code, and will later set up an input to turn it on or off. Notice that the Particle System Component is attached directly to the Static Mesh Component, not to the root. It is also slightly offset from the bottom-center of the mesh so that it will be more visible during play.
 			// Create a particle system that we can activate or deactivate
 			OurParticleSystem = CreateDefaultSubobject<UParticleSystemComponent>(TEXT("MovementParticles"));
 			OurParticleSystem->AttachTo(SphereVisual);
@@ -69,7 +69,7 @@ Description:Creating a variety of components and attaching them to a customized
 				OurParticleSystem->SetTemplate(ParticleAsset.Object);
 			}
 
-1. A **Spring Arm Component** can provide us with a smoother attachment point for our camera by allowing our camera to accelerate and decelerate more slowly than the **Pawn** it is following. It also has a built-in feature that prevents the camera from going through solid objects, for situations such as when a player backs into a corner in a third-person game. While it isn't necessary, it can be a very fast and easy way to get smoother-feeling camera work into a game.
+1. A Spring Arm Component can provide us with a smoother attachment point for our camera by allowing our camera to accelerate and decelerate more slowly than the Pawn it is following. It also has a built-in feature that prevents the camera from going through solid objects, for situations such as when a player backs into a corner in a third-person game. While it isn't necessary, it can be a very fast and easy way to get smoother-feeling camera work into a game.
 			// Use a spring arm to give the camera smooth, natural-feeling motion.
 			USpringArmComponent* SpringArm = CreateDefaultSubobject<USpringArmComponent>(TEXT("CameraAttachmentArm"));
 			SpringArm->AttachTo(RootComponent);
@@ -78,18 +78,18 @@ Description:Creating a variety of components and attaching them to a customized
 			SpringArm->bEnableCameraLag = true;
 			SpringArm->CameraLagSpeed = 3.0f;
 
-1. The actual **Camera Component** is easy to create, and we don't require any special settings. The **Spring Arm** has a special built-in **socket** that we can attach to, rather than attaching to its base.
+1. The actual Camera Component is easy to create, and we don't require any special settings. The **Spring Arm** has a special built-in **socket** that we can attach to, rather than attaching to its base.
 			// Create a camera and attach to our spring arm
 			UCameraComponent* Camera = CreateDefaultSubobject<UCameraComponent>(TEXT("ActualCamera"));
 			Camera->AttachTo(SpringArm, USpringArmComponent::SocketName);
 
-1. Now that our components are created and attached, we should set this **Pawn** to be controlled by the default player. The following code is all we need:
+1. Now that our components are created and attached, we should set this Pawn to be controlled by the default player. The following code is all we need:
 			// Take control of the default player
 			AutoPossessPlayer = EAutoReceiveInput::Player0;
 
 ----
 
-Our new **Pawn** now has a collection of useful **Components** attached to it, and is ready to be configured for user control. We can now return to the **Unreal Editor**.
+Our new Pawn now has a collection of useful Components attached to it, and is ready to be configured for user control. We can now return to the **Unreal Editor**.
 
 [OBJECT:Section]
     [PARAMLITERAL:id]
