#UE4doc #docs:
Added Engine Version field
Changed Home link to point to the overview for the tutorial
Changed font settings on a few terms
Replaced a picture to show updated menu
Replaced a picture to change the border indicator from green to yellow

diff --git a/Engine/Documentation/Source/Programming/Tutorials/Components/2/Components_2.INT.udn b/Engine/Documentation/Source/Programming/Tutorials/Components/2/Components_2.INT.udn
--- a/Engine/Documentation/Source/Programming/Tutorials/Components/2/Components_2.INT.udn
+++ b/Engine/Documentation/Source/Programming/Tutorials/Components/2/Components_2.INT.udn
@@ -2,6 +2,7 @@ Availability: Public
 Title:2. Configuring input and creating a Pawn Movement Component
 Crumbs:%ROOT%, Programming, Programming/Tutorials, Programming/Tutorials/Components
 Description:Setting up input and building a custom Pawn Movement Component to define how our Pawn will move.
+version:4.9
 
 [VAR:Steps]
 [OBJECT:Navigation]
@@ -11,7 +12,7 @@ Description:Setting up input and building a custom Pawn Movement Component to de
 	[PARAM:current]
 	[/PARAM]
 	[PARAM:home]
-		[Programming Tutorials Home](Programming\Tutorials\)
+		[](Programming\Tutorials\Components\)
 	[/PARAM]
 	[PARAM:next]
 		[Next Step](Programming\Tutorials\Components\3 "%Programming\Tutorials\Components\3:title%")
@@ -25,7 +26,7 @@ Description:Setting up input and building a custom Pawn Movement Component to de
 
 	[![](EditProjectSettings.png)](EditProjectSettings.png)
 
-	Once there, we can select **Input** from the **Engine** section in the left panel. We will need an **Action Mapping** to set up our particle-system toggle, two **Axis Mappings** to move our **Pawn**, and one more **Axis Mapping** to turn our **Pawn**.
+	Once there, we can select **Input** from the **Engine** section in the left panel. We will need an **Action Mapping** to set up our particle-system toggle, two **Axis Mappings** to move our **Pawn**, and one more **Axis Mapping** to turn our Pawn.
 
 	| ----- | ----- | ----- |
 	| Action Mapping | | |
@@ -41,11 +42,11 @@ Description:Setting up input and building a custom Pawn Movement Component to de
 
 	[![](InputSettings.png)](InputSettings.png)
 
-1. Instead of handling all our movement in our **Pawn** class directly, we will create a **Movement Component** to manage it for us. For this tutorial, we'll extend the **Pawn Movement Component** class. We'll start by selecting the **Add Code to Project** command from the **File** dropdown menu.
+1. Instead of handling all our movement in our Pawn class directly, we will create a **Movement Component** to manage it for us. For this tutorial, we'll extend the **Pawn Movement Component** class. We'll start by selecting the **Add Code to Project** command from the **File** dropdown menu.
 
 	[![](AddCPPCode.png)](AddCPPCode.png)
 
-	Unlike our **Pawn** class, **Pawn Movement Component** is not visible by default. To find it, we need to check the **Show All Classes** option.
+	Unlike our Pawn class, Pawn Movement Component is not visible by default. To find it, we need to check the **Show All Classes** option.
 
 	[![](ShowAllClasses.png)](ShowAllClasses.png)
 
@@ -53,15 +54,15 @@ Description:Setting up input and building a custom Pawn Movement Component to de
 
 	[![](ChooseParentClass.png)](ChooseParentClass.png)
 
-	[REGION:tip]**Pawn Movement Components** have some powerful, built-in features to help with common physics functionality, and are a good way to share movement code between many **Pawn** types. Using **Components** to separate functionality is a good practice to reduce clutter as your project grows and your **Pawns** become more complex.[/REGION]
+	[REGION:tip]Pawn Movement Components have some powerful, built-in features to help with common physics functionality, and are a good way to share movement code between many Pawn types. Using Components to separate functionality is a good practice to reduce clutter as your project grows and your Pawns become more complex.[/REGION]
 
-	Since we called our **Pawn** class "CollidingPawn", let's call this "CollidingPawnMovementComponent".
+	Since we called our Pawn class CollidingPawn, let's call this CollidingPawnMovementComponent.
 
 	[![](NameComponent.png)](NameComponent.png)
 
 ----
 
-We have just defined our input configuration and created a custom **Pawn Movement Component**. We're ready to go back to **Visual Studio** and write code to define how **Pawns** will move when our **Pawn Movement Component** is in use.
+We have just defined our input configuration and created a custom Pawn Movement Component. We're ready to go back to **Visual Studio** and write code to define how Pawns will move when our Pawn Movement Component is in use.
 
 [OBJECT:Section]
     [PARAMLITERAL:id]
