#UE4doc #docs:
Added Engine Version field
Changed Home link to point to the overview for the tutorial
Changed font settings on a few terms

diff --git a/Engine/Documentation/Source/Programming/Tutorials/Components/3/Components_3.INT.udn b/Engine/Documentation/Source/Programming/Tutorials/Components/3/Components_3.INT.udn
--- a/Engine/Documentation/Source/Programming/Tutorials/Components/3/Components_3.INT.udn
+++ b/Engine/Documentation/Source/Programming/Tutorials/Components/3/Components_3.INT.udn
@@ -2,7 +2,7 @@ Availability: Public
 Title:3. Coding our Pawn Movement Component's behavior
 Crumbs:%ROOT%, Programming, Programming/Tutorials, Programming/Tutorials/Components
 Description:Writing the code that will define how our Pawn Movement Component will determine where the Pawn it controls will go each frame.
-
+version:4.9
 [VAR:Steps]
 [OBJECT:Navigation]
 	[PARAM:previous]
@@ -11,7 +11,7 @@ Description:Writing the code that will define how our Pawn Movement Component wi
 	[PARAM:current]
 	[/PARAM]
 	[PARAM:home]
-		[Programming Tutorials Home](Programming\Tutorials\)
+		[](Programming\Tutorials\Components\)
 	[/PARAM]
 	[PARAM:next]
 		[Next Step](Programming\Tutorials\Components\4 "%Programming\Tutorials\Components\4:title%")
@@ -21,12 +21,12 @@ Description:Writing the code that will define how our Pawn Movement Component wi
 
 %Steps%
 
-1. We can go back to **Visual Studio** now and program our customized **Pawn Movement Component**. All we really need to write is a TickComponent function (similar to an **Actor's** Tick function) to tell us how to move each frame. In `CollidingPawnMovementComponent.h`, we'll need to override TickComponent in the class definition:
+1. We can go back to **Visual Studio** now and program our customized **Pawn Movement Component**. All we really need to write is a **TickComponent** function (similar to an **Actor's** **Tick** function) to tell us how to move each frame. In CollidingPawnMovementComponent.h, we'll need to override TickComponent in the class definition:
 
 		public:
 			virtual void TickComponent(float DeltaTime, enum ELevelTick TickType, FActorComponentTickFunction *ThisTickFunction) override;	
 
-	We'll define this function in `CollidingPawnMovementComponent.cpp`:
+	We'll define this function in CollidingPawnMovementComponent.cpp:
 
 		void UCollidingPawnMovementComponent::TickComponent(float DeltaTime, enum ELevelTick TickType, FActorComponentTickFunction *ThisTickFunction)
 		{
@@ -53,17 +53,17 @@ Description:Writing the code that will define how our Pawn Movement Component wi
 			}
 		};
 
-	This code will move our **Pawn** smoothly around the world, sliding off of surfaces where appropriate. There is no gravity applied to our **Pawn**, and its maximum speed is hard-coded to 150 **Unreal Units** per second.
+	This code will move our **Pawn** smoothly around the world, sliding off of surfaces where appropriate. There is no gravity applied to our Pawn, and its maximum speed is hard-coded to 150 **Unreal Units** per second.
 
 	[REGION:note]This TickComponent function makes use of a few of the powerful features offered by the **UPawnMovementComponent** class.
-	<br>ConsumeInputVector reports and clears the value of a built-in variable that we will use to store our movement inputs.
-	<br>SafeMoveUpdatedComponent uses **Unreal Engine** physics to move our **Pawn Movement Component** while respecting solid barriers.
-	<br>SlideAlongSurface handles the calculations and physics involved with sliding smoothly along collision surfaces like walls and ramps when a move results in a collision, rather than simply stopping in place and "sticking" to the wall or ramp.
-	<br>There are more features included in **Pawn Movement Components** that are worthy of examination, but they are not needed for the scope of this tutorial. Looking at other classes, such as **Floating Pawn Movement**, **Spectator Pawn Movement**, or **Character Movement Component**, could provide additional usage examples and ideas.[/REGION]
+	<br>**ConsumeInputVector** reports and clears the value of a built-in variable that we will use to store our movement inputs.
+	<br>**SafeMoveUpdatedComponent** uses **Unreal Engine** physics to move our Pawn Movement Component while respecting solid barriers.
+	<br>**SlideAlongSurface** handles the calculations and physics involved with sliding smoothly along collision surfaces like walls and ramps when a move results in a collision, rather than simply stopping in place and sticking to the wall or ramp.
+	<br>There are more features included in Pawn Movement Components that are worthy of examination, but they are not needed for the scope of this tutorial. Looking at other classes, such as **Floating Pawn Movement**, **Spectator Pawn Movement**, or **Character Movement Component**, could provide additional usage examples and ideas.[/REGION]
 
 ----
 
-With our **Pawn Movement Component's** behavior defined, we can now write the code that will tie it all together in our customized **Pawn** class.
+With our Pawn Movement Component's behavior defined, we can now write the code that will tie it all together in our customized Pawn class.
 
 [OBJECT:Section]
     [PARAMLITERAL:id]
