#UE4doc #docs:
Added Engine Version field
Changed Home link to point to the overview for the tutorial
Changed font settings on a few terms
Edited tutorial to correct an error with a step that led to a compilation error
Added a warning field to put a spotlight on generated.h files.

diff --git a/Engine/Documentation/Source/Programming/Tutorials/Components/4/Components_4.INT.udn b/Engine/Documentation/Source/Programming/Tutorials/Components/4/Components_4.INT.udn
--- a/Engine/Documentation/Source/Programming/Tutorials/Components/4/Components_4.INT.udn
+++ b/Engine/Documentation/Source/Programming/Tutorials/Components/4/Components_4.INT.udn
@@ -2,6 +2,7 @@ Availability: Public
 Title:4. Using our Pawn and Components together
 Crumbs:%ROOT%, Programming, Programming/Tutorials, Programming/Tutorials/Components
 Description:Tying the functionality of the Pawn, the input configuration, and the Components together.
+version:4.9
 
 [VAR:Steps]
 [OBJECT:Navigation]
@@ -11,7 +12,7 @@ Description:Tying the functionality of the Pawn, the input configuration, and th
 	[PARAM:current]
 	[/PARAM]
 	[PARAM:home]
-		[Programming Tutorials Home](Programming\Tutorials\)
+		[](Programming\Tutorials\Components\)
 	[/PARAM]
 	[PARAM:next]
 		[Next Step](Programming\Tutorials\Components\5 "%Programming\Tutorials\Components\5:title%")
@@ -21,41 +22,44 @@ Description:Tying the functionality of the Pawn, the input configuration, and th
 
 %Steps%
 
-1. In order to use our custom **Pawn Movement Component**, we will first need to add a variable to our **Pawn** class to keep track of it. At the bottom of our class definition in `CollidingPawn.h`, near where we added the "OurParticleSystem" variable, we should add:
+1. In order to use our custom **Pawn Movement Component**, we will first need to add a variable to our **Pawn** class to keep track of it. At the bottom of our class definition in CollidingPawn.h, near where we added the "OurParticleSystem" variable, we should add:
 
 		class UCollidingPawnMovementComponent* OurMovementComponent;
 
-1. Once we have a place to keep track of it, we will need to create a **Colliding Pawn Movement Component** to store in our new variable, so let's open `CollidingPawn.cpp` and add the following to the top of the file, underneath the "#include" lines, so that our code will be able to reference our new class:
+1. Once we have a place to keep track of it, we will need to create a **Colliding Pawn Movement Component** to store in our new variable, so let's open CollidingPawn.cpp and add the following to the top of the file, underneath the #include "GameFramework/Pawn.h" line, so that our code will be able to reference our new class:
 
 		#include "CollidingPawnMovementComponent.h"
+[REGION:warning]
+Be sure that the last #include in the list is the generated.h (#include "CollidingPawn.generated.h" in this case) as it will cause a compile error otherwise.
+[/REGION]
 
-	Creating a **Pawn Movement Component** and associating it with our **Pawn** is simple. At the bottom of **ACollidingPawn::ACollidingPawn**, we can add this code:
+	Creating a Pawn Movement Component and associating it with our Pawn is simple. At the bottom of **ACollidingPawn::ACollidingPawn**, we can add this code:
 
 		// Create an instance of our movement component, and tell it to update the root.
 		OurMovementComponent = CreateDefaultSubobject<UCollidingPawnMovementComponent>(TEXT("CustomMovementComponent"));
 		OurMovementComponent->UpdatedComponent = RootComponent;
 
-	[REGION:note]Unlike the other **Components** we've seen so far, we don't need to attach this **Component** to our own **Component** hierarchy. This is because our other **Components** are all types of **Scene Components**, which inherently require physical locations. **Movement Controllers**, however, are not **Scene Components** and do not represent physical objects, so the concept of existing at a physical location or being physically attached to another **Component** does not apply to them.[/REGION]
+	[REGION:note]Unlike the other **Components** we've seen so far, we don't need to attach this Component to our own Component hierarchy. This is because our other Components are all types of **Scene Components**, which inherently require physical locations. **Movement Controllers**, however, are not Scene Components and do not represent physical objects, so the concept of existing at a physical location or being physically attached to another Component does not apply to them.[/REGION]
 
-1. **Pawns** have a function called "GetMovementComponent" that is used to enable other classes in the engine to access the **Pawn Movement Component** that the **Pawn** is currently using. We will need to override that function so that it returns our custom **Pawn Movement Component**. In the class definition in `CollidingPawn.h`, we need to add:
+1. Pawns have a function called **GetMovementComponent** that is used to enable other classes in the engine to access the Pawn Movement Component that the Pawn is currently using. We will need to override that function so that it returns our custom Pawn Movement Component. In the class definition in CollidingPawn.h, we need to add:
 
 		virtual UPawnMovementComponent* GetMovementComponent() const override;
 
-	And in `CollidingPawn.cpp`, we need to add the definition of our overridden function, as follows:
+	And in CollidingPawn.cpp, we need to add the definition of our overridden function, as follows:
 
 		UPawnMovementComponent* ACollidingPawn::GetMovementComponent() const
 		{
 			return OurMovementComponent;
 		}
 
-1. With our new **Pawn Movement Component** set up, we can create code to handle the input that our **Pawn** will receive. We will start by declaring a few functions in our class definition in `CollidingPawn.h`:
+1. With our new Pawn Movement Component set up, we can create code to handle the input that our Pawn will receive. We will start by declaring a few functions in our class definition in CollidingPawn.h:
 
 		void MoveForward(float AxisValue);
 		void MoveRight(float AxisValue);
 		void Turn(float AxisValue);
 		void ParticleToggle();
 
-	In `CollidingPawn.cpp`, we will add the definitions of those functions as follows:
+	In CollidingPawn.cpp, we will add the definitions of those functions as follows:
 
 			void ACollidingPawn::MoveForward(float AxisValue)
 			{
@@ -100,7 +104,7 @@ Description:Tying the functionality of the Pawn, the input configuration, and th
 
 ----
 
-Our programming work is done, and we can now place our custom **Pawn** in the world and move it around.
+Our programming work is done, and we can now place our custom Pawn in the world and move it around.
 
 [OBJECT:Section]
     [PARAMLITERAL:id]
