#UE4doc #docs:
Added Engine Version field
Changed Home link to point to the overview for the tutorial
Changed font settings on a few terms
Replaced (Our Project Name) with HowTo_Components as that was the file name given in the first step

diff --git a/Engine/Documentation/Source/Programming/Tutorials/Components/5/Components_5.INT.udn b/Engine/Documentation/Source/Programming/Tutorials/Components/5/Components_5.INT.udn
--- a/Engine/Documentation/Source/Programming/Tutorials/Components/5/Components_5.INT.udn
+++ b/Engine/Documentation/Source/Programming/Tutorials/Components/5/Components_5.INT.udn
@@ -2,6 +2,7 @@ Availability: Public
 Title:5. Playing In Editor
 Crumbs:%ROOT%, Programming, Programming/Tutorials, Programming/Tutorials/Components
 Description:Testing our code by playing in the editor.
+version:4.9
 
 [VAR:Steps]
 [OBJECT:Navigation]
@@ -11,7 +12,7 @@ Description:Testing our code by playing in the editor.
 	[PARAM:current]
 	[/PARAM]
 	[PARAM:home]
-		[Programming Tutorials Home](Programming\Tutorials\)
+		[](Programming\Tutorials\Components\)
 	[/PARAM]
 	[PARAM:next]
 		[Next Step](Programming\Tutorials\Components\6 "%Programming\Tutorials\Components\6:title%")
@@ -25,7 +26,7 @@ Description:Testing our code by playing in the editor.
 
 	[![](CompileFromEditor.png)](CompileFromEditor.png)
 
-1. We need to drop an instance of our **CollidingPawn** into the world. We can find the class in the **Content Browser** under "C++ Classes/(Our Project Name)/CollidingPawn".
+1. We need to drop an instance of our **CollidingPawn** into the world. We can find the class in the **Content Browser** under "C++ Classes/HowTo_Components/CollidingPawn".
 
 	[![](ClassInContentBrowser.png)](ClassInContentBrowser.png)
 
@@ -37,7 +38,7 @@ Description:Testing our code by playing in the editor.
 
 ----
 
-**Unreal Engine 4** offers a wide variety of **Components** - we've only scratched the surface with a few of the more common ones. Try exploring the built-in **Components**, or writing some of your own! They are flexible, powerful, and help to keep your project's code organized and reusable.
+**Unreal Engine 4** offers a wide variety of **Components** - we've only scratched the surface with a few of the more common ones. Try exploring the built-in Components, or writing some of your own! They are flexible, powerful, and help to keep your project's code organized and reusable.
 
 [OBJECT:Section]
     [PARAMLITERAL:id]
