#UE4doc #docs:
Added Engine Version field
Changed Home link to point to the overview for the tutorial
Changed font settings on a few terms
Fixed link to 'Previous Step' to point to Step 5 instead of Step 4

diff --git a/Engine/Documentation/Source/Programming/Tutorials/Components/6/Components_6.INT.udn b/Engine/Documentation/Source/Programming/Tutorials/Components/6/Components_6.INT.udn
--- a/Engine/Documentation/Source/Programming/Tutorials/Components/6/Components_6.INT.udn
+++ b/Engine/Documentation/Source/Programming/Tutorials/Components/6/Components_6.INT.udn
@@ -2,16 +2,17 @@ Availability: Public
 Title:6. On Your Own!
 Crumbs:%ROOT%, Programming, Programming/Tutorials, Programming/Tutorials/Components
 Description:Ideas for further exploration!
+version:4.9
 
 [VAR:Steps]
 [OBJECT:Navigation]
 	[PARAM:previous]
-		[Previous Step](Programming\Tutorials\Components\4 "%Programming\Tutorials\Components\4:title%")
+		[Previous Step](Programming\Tutorials\Components\5 "%Programming\Tutorials\Components\5:title%")
 	[/PARAM]
 	[PARAM:current]
 	[/PARAM]
 	[PARAM:home]
-		[Programming Tutorials Home](Programming\Tutorials\)
+		[](Programming\Tutorials\Components\)
 	[/PARAM]
 	[PARAM:next]
 	[/PARAM]
@@ -23,12 +24,12 @@ Description:Ideas for further exploration!
 Using what you have learned, try to do the following:
 
 * Create a **Component** that automatically orbits its parent.
-* Build a **Component** that spawns up to three children, each of which despawn on their own after a set amount of time.
-* Learn how to attach complete **Actors** to other **Actors** via **Components**.
+* Build a Component that spawns up to three children, each of which despawn on their own after a set amount of time.
+* Learn how to attach complete **Actors** to other Actors via Components.
 
 As for the specifics covered in this tutorial:
 
-* For more information about **Components**, try the [%Engine/Components:title%](Engine/Components) page.
+* For more information about Components, try the [%Engine/Components:title%](Engine/Components) page.
 * For further tutorials, see the [%Programming/Tutorials:title%](Programming/Tutorials) page.
 
 
