#UE4doc #docs:
Added Engine Version field
Changed Home link to point to the overview for the tutorial
Changed font settings on a few terms

diff --git a/Engine/Documentation/Source/Programming/Tutorials/PlayerCamera/1/PlayerCamera_1.INT.udn b/Engine/Documentation/Source/Programming/Tutorials/PlayerCamera/1/PlayerCamera_1.INT.udn
--- a/Engine/Documentation/Source/Programming/Tutorials/PlayerCamera/1/PlayerCamera_1.INT.udn
+++ b/Engine/Documentation/Source/Programming/Tutorials/PlayerCamera/1/PlayerCamera_1.INT.udn
@@ -2,6 +2,7 @@ Availability: Public
 Title:1. Attach A Camera To A Pawn
 Crumbs:%ROOT%, Programming, Programming/Tutorials, Programming/Tutorials/PlayerCamera
 Description:Creating a spring arm and a camera.
+version:4.9
 
 [VAR:Steps]
 [OBJECT:Navigation]
@@ -10,7 +11,7 @@ Description:Creating a spring arm and a camera.
 	[PARAM:current]
 	[/PARAM]
 	[PARAM:home]
-		[Programming Tutorials Home](Programming\Tutorials\)
+		[](Programming\Tutorials\PlayerCamera)
 	[/PARAM]
 	[PARAM:next]
 		[Next Step](Programming\Tutorials\PlayerCamera\2 "%Programming\Tutorials\PlayerCamera\2:title%")
@@ -22,20 +23,20 @@ Description:Creating a spring arm and a camera.
 
 [REGION:tip]If you are new to **Unreal Engine 4**, you might want to read our [%Programming\QuickStart:title% tutorial](Programming\QuickStart) first. For this tutorial, we will assume you are familiar with creating a project, adding C++ code to it, compiling your code, and adding **Components** to **Actors** in the **Unreal Engine** editor.[/REGION]
 
-1. We will begin by creating a new, Basic Code project, with starter content, named "HowTo_PlayerCamera". We'll want to create a custom **Pawn** class, so let's do that first. For this tutorial, we'll use the name "PawnWithCamera" for our new **Pawn** class.
+1. We will begin by creating a new, Basic Code project, with starter content, named HowTo_PlayerCamera. We'll want to create a custom **Pawn** class, so let's do that first. For this tutorial, we'll use the name PawnWithCamera for our new Pawn class.
 
 	[![](NamePawnClass.png)](NamePawnClass.png)
 
-1. Next, in **Visual Studio**, we should open `PawnWithCamera.h` and add the following code to the bottom of our class definition:
+1. Next, in **Visual Studio**, we should open PawnWithCamera.h and add the following code to the bottom of our class definition:
 
 		protected:
 			UPROPERTY(EditAnywhere)
 			USpringArmComponent* OurCameraSpringArm;
 			UCameraComponent* OurCamera;
 
-	We will use these variables to create a **SpringArmComponent** with a **CameraComponent** attached to the end. Spring arms are a good, simple way to attach cameras (or other **Components**) so that they're not overly rigid and feel more fluid as they move.
+	We will use these variables to create a **SpringArmComponent** with a **CameraComponent** attached to the end. Spring arms are a good, simple way to attach cameras (or other Components) so that they're not overly rigid and feel more fluid as they move.
 
-1. After that, we actually need to create our **Components** in our constructor. Add the following code to `PawnWithCamera.cpp` inside of **APawnWithCamera::APawnWithCamera**:
+1. After that, we actually need to create our Components in our constructor. Add the following code to PawnWithCamera.cpp inside of **APawnWithCamera::APawnWithCamera**:
 
 		//Create our components
 		RootComponent = CreateDefaultSubobject<USceneComponent>(TEXT("RootComponent"));
@@ -46,19 +47,19 @@ Description:Creating a spring arm and a camera.
 		OurCameraSpringArm->bEnableCameraLag = true;
 		OurCameraSpringArm->CameraLagSpeed = 3.0f;
 
-	This creates a basic, empty **SceneComponent** as the root of our **Component** hierarchy, then creates and attaches a **SpringArmComponent** to it. The **Spring Arm** is then set at a default pitch of -60 degrees (that is, looking 60 degrees downward) and a position of 50 units above the root. We also establish a few values specific to the **SpringArmComponent** class that will determine its length and the smoothness of its motion. With that finished, we simply need to create and attach a **CameraComponent** to the socket on the end of the **SpringArmComponent**, as follows:
+	This creates a basic, empty **SceneComponent** as the root of our Component hierarchy, then creates and attaches a SpringArmComponent to it. The **Spring Arm** is then set at a default pitch of -60 degrees (that is, looking 60 degrees downward) and a position of 50 units above the root. We also establish a few values specific to the SpringArmComponent class that will determine its length and the smoothness of its motion. With that finished, we simply need to create and attach a CameraComponent to the socket on the end of the SpringArmComponent, as follows:
 
 		OurCamera = CreateDefaultSubobject<UCameraComponent>(TEXT("GameCamera"));
 		OurCamera->AttachTo(OurCameraSpringArm, USpringArmComponent::SocketName);
 
-1. Finally, we can set our **Pawn** to be controlled by the default, local player automatically upon being spawned by adding this piece of code:
+1. Finally, we can set our Pawn to be controlled by the default, local player automatically upon being spawned by adding this piece of code:
 
 		//Take control of the default Player
 		AutoPossessPlayer = EAutoReceiveInput::Player0;
 
 ----
 
-We now have a simple **Pawn** that will allow us to control our camera comfortably. Next, we'll configure our input in the **Unreal Engine** editor and create code that reacts to it.
+We now have a simple Pawn that will allow us to control our camera comfortably. Next, we'll configure our input in the **Unreal Engine** editor and create code that reacts to it.
 
 [OBJECT:Section]
     [PARAMLITERAL:id]
