#UE4doc #docs:
Added Engine Version field
Changed Home link to point to the overview for the tutorial

diff --git a/Engine/Documentation/Source/Programming/Tutorials/PlayerCamera/2/PlayerCamera_2.INT.udn b/Engine/Documentation/Source/Programming/Tutorials/PlayerCamera/2/PlayerCamera_2.INT.udn
--- a/Engine/Documentation/Source/Programming/Tutorials/PlayerCamera/2/PlayerCamera_2.INT.udn
+++ b/Engine/Documentation/Source/Programming/Tutorials/PlayerCamera/2/PlayerCamera_2.INT.udn
@@ -2,6 +2,7 @@ Availability: Public
 Title:2. Configure Input To Control The Camera
 Crumbs:%ROOT%, Programming, Programming/Tutorials, Programming/Tutorials/PlayerCamera
 Description:Configuring input that will be applied to a pawn and camera.
+version:4.9
 
 [VAR:Steps]
 [OBJECT:Navigation]
@@ -11,7 +12,7 @@ Description:Configuring input that will be applied to a pawn and camera.
 	[PARAM:current]
 	[/PARAM]
 	[PARAM:home]
-		[Programming Tutorials Home](Programming\Tutorials\)
+		[](Programming\Tutorials\PlayerCamera)
 	[/PARAM]
 	[PARAM:next]
 		[Next Step](Programming\Tutorials\PlayerCamera\3 "%Programming\Tutorials\PlayerCamera\3:title%")
