#UE4doc #docs:
Added Engine Version field
Changed Home link to point to the overview for the tutorial
Unbolded the second iteration of 'Camera' and removed some unnecessary quotations.

diff --git a/Engine/Documentation/Source/Programming/Tutorials/PlayerCamera/4/PlayerCamera_4.INT.udn b/Engine/Documentation/Source/Programming/Tutorials/PlayerCamera/4/PlayerCamera_4.INT.udn
--- a/Engine/Documentation/Source/Programming/Tutorials/PlayerCamera/4/PlayerCamera_4.INT.udn
+++ b/Engine/Documentation/Source/Programming/Tutorials/PlayerCamera/4/PlayerCamera_4.INT.udn
@@ -2,7 +2,7 @@ Availability: Public
 Title:4. On Your Own!
 Crumbs:%ROOT%, Programming, Programming/Tutorials, Programming/Tutorials/PlayerCamera
 Description:Ideas for further exploration!
-
+version:4.9
 [VAR:Steps]
 [OBJECT:Navigation]
 	[PARAM:previous]
@@ -11,7 +11,7 @@ Description:Ideas for further exploration!
 	[PARAM:current]
 	[/PARAM]
 	[PARAM:home]
-		[Programming Tutorials Home](Programming\Tutorials\)
+		[](Programming\Tutorials\PlayerCamera)
 	[/PARAM]
 	[PARAM:next]
 	[/PARAM]
@@ -22,15 +22,15 @@ Description:Ideas for further exploration!
 
 Using what you have learned, try to do the following:
 
-* Give the player a "run" key to hold down that will increase the **Pawn's** movement speed.
+* Give the player a run key to hold down that will increase the **Pawn's** movement speed.
 * Experiment with different ways to mix automatic and input-driven camera movement. This is a very deep area of game development with a lot of room to explore!
 * Increase, decrease, or remove the lag from your **Spring Component** to get a better understanding of how much lag can affect your camera's overall feel.
-* Implement a small amount of periodic motion, possibly slightly randomized or using a **Curve** asset, to create a "handheld" feel to your camera.
+* Implement a small amount of periodic motion, possibly slightly randomized or using a **Curve** asset, to create a handheld feel to your camera.
 * Give your **Camera** some amount of automatic turning, so that the camera gradually gets behind the moving player object and faces the direction the player is moving.
 
 As for the specifics covered in this tutorial:
 
-* For more information on **Cameras** and ways to control them, see the [Camera](Gameplay/Framework/Camera) framework page.
+* For more information on Cameras and ways to control them, see the [Camera](Gameplay/Framework/Camera) framework page.
 * To learn more about **Components**, try the [%Programming/Tutorials/Components:title%](Programming/Tutorials/Components) tutorial.
 * For further tutorials, see the [%Programming/Tutorials:title%](Programming/Tutorials) page.
 
