#UE4doc #docs:
Added Engine Version field
Changed Home link to point to the overview for the tutorial
Changed font settings on a few terms

diff --git a/Engine/Documentation/Source/Programming/Tutorials/PlayerCamera/3/PlayerCamera_3.INT.udn b/Engine/Documentation/Source/Programming/Tutorials/PlayerCamera/3/PlayerCamera_3.INT.udn
--- a/Engine/Documentation/Source/Programming/Tutorials/PlayerCamera/3/PlayerCamera_3.INT.udn
+++ b/Engine/Documentation/Source/Programming/Tutorials/PlayerCamera/3/PlayerCamera_3.INT.udn
@@ -21,7 +21,7 @@ Description:Implementing code to control a pawn and a camera with input.
 
 %Steps%
 
-1. Our game now has input mappings that we can use, so let's set up some member variables to store the data we receive. During our update, we'll need to know the values of our movement and mouse-looking axes, each of which are two-dimensional vectors. We'll also want to know whether we should be moving toward our zoomed-in or zoomed-out view, and how far between those two states we currently are. To accomplish this, we should add the following code to our class definition in `PawnWithCamera.h`:
+1. Our game now has input mappings that we can use, so let's set up some member variables to store the data we receive. During our update, we'll need to know the values of our movement and mouse-looking axes, each of which are two-dimensional vectors. We'll also want to know whether we should be moving toward our zoomed-in or zoomed-out view, and how far between those two states we currently are. To accomplish this, we should add the following code to our class definition in PawnWithCamera.h:
 
 		//Input variables
 		FVector2D MovementInput;
@@ -29,7 +29,7 @@ Description:Implementing code to control a pawn and a camera with input.
 		float ZoomFactor;
 		bool bZoomingIn;
 
-1. We'll need to create functions to track our input, so let's add the following to our class definition in `PawnWithCamera.h` as well:
+1. We'll need to create functions to track our input, so let's add the following to our class definition in PawnWithCamera.h as well:
 
 		//Input functions
 		void MoveForward(float AxisValue);
@@ -39,7 +39,7 @@ Description:Implementing code to control a pawn and a camera with input.
 		void ZoomIn();
 		void ZoomOut();
 
-	And we can now fill these functions out in `PawnWithCamera.cpp` with the following code:
+	And we can now fill these functions out in PawnWithCamera.cpp with the following code:
 
 		//Input functions
 		void APawnWithCamera::MoveForward(float AxisValue)
@@ -74,7 +74,7 @@ Description:Implementing code to control a pawn and a camera with input.
 
 	[REGION:note]We haven't done anything with ZoomFactor yet. That variable will be updated during our **Pawn**'s **Tick** function, since its value changes constantly over time based to the state of a button.[/REGION]
 
-1. Now that we have code that will store our input data, we just need to tell **Unreal Engine** when to call that code. Binding functions to input events for **Pawn**s is as simple as adding binding code to **APawnWithCamera::SetupPlayerInputComponent**, as follows:
+1. Now that we have code that will store our input data, we just need to tell **Unreal Engine** when to call that code. Binding functions to input events for Pawns is as simple as adding binding code to **APawnWithCamera::SetupPlayerInputComponent**, as follows:
 
 		//Hook up events for "ZoomIn"
 		InputComponent->BindAction("ZoomIn", IE_Pressed, this, &APawnWithCamera::ZoomIn);
@@ -86,7 +86,7 @@ Description:Implementing code to control a pawn and a camera with input.
 		InputComponent->BindAxis("CameraPitch", this, &APawnWithCamera::PitchCamera);
 		InputComponent->BindAxis("CameraYaw", this, &APawnWithCamera::YawCamera);
 
-1. Finally, we can use those values in our **Tick** function to update our **Pawn** and **Camera** each frame. The following code blocks should all be added to **APawnWithCamera::Tick** in `PawnWithCamera.cpp`:
+1. Finally, we can use those values in our Tick function to update our Pawn and **Camera** each frame. The following code blocks should all be added to **APawnWithCamera::Tick** in PawnWithCamera.cpp:
 
 		//Zoom in if ZoomIn button is down, zoom back out if it's not
 		{
@@ -118,7 +118,7 @@ Description:Implementing code to control a pawn and a camera with input.
 			NewRotation.Pitch = FMath::Clamp(NewRotation.Pitch + CameraInput.Y, -80.0f, -15.0f);
 			OurCameraSpringArm->SetWorldRotation(NewRotation);
 		}
-	This block of code rotates our **Pawn**'s yaw directly with the mouse's X axis, but only the camera system responds to the pitch changes from the mouse's Y axis. Rotating any **Actor** or **Actor** subclass actually rotates the root-level **Component**, which indirectly affects everything attached to it.
+	This block of code rotates our Pawn's yaw directly with the mouse's X axis, but only the camera system responds to the pitch changes from the mouse's Y axis. Rotating any **Actor** or Actor subclass actually rotates the root-level **Component**, which indirectly affects everything attached to it.
 
 		//Handle movement based on our "MoveX" and "MoveY" axes
 		{
@@ -139,7 +139,7 @@ Description:Implementing code to control a pawn and a camera with input.
 
 	[![](ClassInContentBrowser.png)](ClassInContentBrowser.png)
 
-	Feel free to add a **Static Mesh** or other visual **Component**, or play without one. You should feel a smooth acceleration and deceleration to your camera's movement as it follows you around the level, but your rotation should feel tight and instantaneous. Try changing some of the properties on the **SpringArmComponent** to see how they affect the feel of your controls, such as adding "Camera Rotation Lag" or incresing or decreasing "Camera Lag".
+	Feel free to add a **Static Mesh** or other visual Component, or play without one. You should feel a smooth acceleration and deceleration to your camera's movement as it follows you around the level, but your rotation should feel tight and instantaneous. Try changing some of the properties on the **SpringArmComponent** to see how they affect the feel of your controls, such as adding "Camera Rotation Lag" or incresing or decreasing Camera Lag.
 
 	[![](SpringArmValues.png)](SpringArmValues.png)
 
