#UE4doc #docs:
Didn't save changed before previous submit

diff --git a/Engine/Documentation/Source/Programming/Tutorials/PlayerCamera/3/PlayerCamera_3.INT.udn b/Engine/Documentation/Source/Programming/Tutorials/PlayerCamera/3/PlayerCamera_3.INT.udn
--- a/Engine/Documentation/Source/Programming/Tutorials/PlayerCamera/3/PlayerCamera_3.INT.udn
+++ b/Engine/Documentation/Source/Programming/Tutorials/PlayerCamera/3/PlayerCamera_3.INT.udn
@@ -2,6 +2,7 @@ Availability: Public
 Title:3. Write C++ Code To React To Input
 Crumbs:%ROOT%, Programming, Programming/Tutorials, Programming/Tutorials/PlayerCamera
 Description:Implementing code to control a pawn and a camera with input.
+version:4.9
 
 [VAR:Steps]
 [OBJECT:Navigation]
@@ -11,7 +12,7 @@ Description:Implementing code to control a pawn and a camera with input.
 	[PARAM:current]
 	[/PARAM]
 	[PARAM:home]
-		[Programming Tutorials Home](Programming\Tutorials\)
+		[](Programming\Tutorials\PlayerCamera)
 	[/PARAM]
 	[PARAM:next]
 		[Next Step](Programming\Tutorials\PlayerCamera\4 "%Programming\Tutorials\PlayerCamera\4:title%")
