#UE4doc #docs - adjusted bold lettering.

diff --git a/Engine/Documentation/Source/Resources/ContentExamples/Lighting/1_1/1_1.INT.udn b/Engine/Documentation/Source/Resources/ContentExamples/Lighting/1_1/1_1.INT.udn
--- a/Engine/Documentation/Source/Resources/ContentExamples/Lighting/1_1/1_1.INT.udn
+++ b/Engine/Documentation/Source/Resources/ContentExamples/Lighting/1_1/1_1.INT.udn
@@ -38,7 +38,7 @@ Related: Engine/Rendering/LightingAndShadows/LightMobility
 ![%Resources\ContentExamples\Lighting\1_1:title%](Resources\ContentExamples\Lighting\Lighting_1_1.png)
 [/REGION]
 
-This display shows a Point Light, which is the most basic type of light available in UE4. A Point Light emits light from a single point in all directions, spherically. This is most obviously seen by looking at the light's radius, which is visible when the light is selected:
+This display shows a **Point Light**, which is the most basic type of light available in UE4. A Point Light emits light from a single point in all directions, spherically. This is most obviously seen by looking at the light's radius, which is visible when the light is selected:
 
 
 ![](LightRadius_PointLight.png)(w:500)
