#UE4doc #docs -  Adjusted bold lettering.

diff --git a/Engine/Documentation/Source/Resources/ContentExamples/Lighting/1_2/1_2.INT.udn b/Engine/Documentation/Source/Resources/ContentExamples/Lighting/1_2/1_2.INT.udn
--- a/Engine/Documentation/Source/Resources/ContentExamples/Lighting/1_2/1_2.INT.udn
+++ b/Engine/Documentation/Source/Resources/ContentExamples/Lighting/1_2/1_2.INT.udn
@@ -37,11 +37,11 @@ Related: Engine/Rendering/LightingAndShadows/LightMobility
 ![%Resources\ContentExamples\Lighting\1_2:title%](Resources\ContentExamples\Lighting\Lighting_1_2.png)
 [/REGION]
 
-A Spot Light emits light from a single point in a cone shape. Users are given two cones to shape the light - the **Inner Cone Angle** and **Outer Cone Angle**. Within the **Inner Cone Angle**, the light achieves full brightness. As you go from the extent of the inner radius to the extents of the **Outer Cone Angle**, a falloff takes place, creating a penumbra, or softening around the Spot Light's disc of illumination.
+A **Spot Light** emits light from a single point in a cone shape. Users are given two cones to shape the light - the **Inner Cone Angle** and Outer Cone Angle. Within the Inner Cone Angle, the light achieves full brightness. As you go from the extent of the inner radius to the extents of the Outer Cone Angle, a falloff takes place, creating a penumbra, or softening around the Spot Light's disc of illumination.
 
 ![](ConeRadius.png)(w:500)
 
-In the image above, we can see that the Spot Light has a very narrow **Inner Cone Angle** and a much wider **Outer Cone Angle**. These are defined in the **Details** panel when the light is selected.
+In the image above, we can see that the Spot Light has a very narrow Inner Cone Angle and a much wider Outer Cone Angle. These are defined in the **Details** panel when the light is selected.
 
 
 
