#UE4doc #docs -  Adjusted bold lettering.

diff --git a/Engine/Documentation/Source/Resources/ContentExamples/Lighting/2_2/2_2.INT.udn b/Engine/Documentation/Source/Resources/ContentExamples/Lighting/2_2/2_2.INT.udn
--- a/Engine/Documentation/Source/Resources/ContentExamples/Lighting/2_2/2_2.INT.udn
+++ b/Engine/Documentation/Source/Resources/ContentExamples/Lighting/2_2/2_2.INT.udn
@@ -48,12 +48,12 @@ Stationary lights support IES profiles, described in [Example 3.1](Resources\Con
 
 ### Overlapping
 
-**Stationary Lights** are limited to a maximum of 4 being able to overlap at one time. Once you have 5 or more Stationary Lights overlapping, the one with the smallest radius will start casting dynamic shadows, which have a higher performance cost. You will also get warnings when rebuilding your lights. Here you can see 5 Stationary Lights placed close to one another. Notice that the fifth has a red **X** over it, indicating that it is now casting dynamic shadows.
+Stationary Lights are limited to a maximum of 4 being able to overlap at one time. Once you have 5 or more Stationary Lights overlapping, the one with the smallest radius will start casting dynamic shadows, which have a higher performance cost. You will also get warnings when rebuilding your lights. Here you can see 5 Stationary Lights placed close to one another. Notice that the fifth has a red **X** over it, indicating that it is now casting dynamic shadows.
 
 
 ![](OverlappingLights.png)(w:680)
 
-Light overlapping can also be visualized in Stationary Light Overlap mode, available within the Viewport's View Mode menu. This view mode shifts from green to red as you increase the number of overlapping Stationary Lights. Once you hit 5 and above, the scene turns red, indicating a problematic light. In the image below, the fifth, and problematic, overlapping light's radius is shown in red. 
+Light overlapping can also be visualized in **Stationary Light Overlap** mode, available within the Viewport's View Mode menu. This view mode shifts from green to red as you increase the number of overlapping Stationary Lights. Once you hit 5 and above, the scene turns red, indicating a problematic light. In the image below, the fifth, and problematic, overlapping light's radius is shown in red. 
 
 ![](StationaryLightOverlap.png)(w:680)
 
