#UE4doc #docs - Adjusted bold lettering.

diff --git a/Engine/Documentation/Source/Resources/ContentExamples/Lighting/2_3/2_3.INT.udn b/Engine/Documentation/Source/Resources/ContentExamples/Lighting/2_3/2_3.INT.udn
--- a/Engine/Documentation/Source/Resources/ContentExamples/Lighting/2_3/2_3.INT.udn
+++ b/Engine/Documentation/Source/Resources/ContentExamples/Lighting/2_3/2_3.INT.udn
@@ -39,7 +39,7 @@ Related: Engine/Rendering/LightingAndShadows/LightMobility
 
 [INCLUDE:Shared/Glossary/M/#Glossary_Movable_Lights]
 
-Movable lights support IES profiles, described in [Example 3.1](Resources\ContentExamples\Lighting\3_1), and Light Functions, described in [Example 3.2](Resources\ContentExamples\Lighting\3_2).
+Movable lights support IES profiles, described in **[Example 3.1](Resources\ContentExamples\Lighting\3_1)**, and Light Functions, described in **[Example 3.2](Resources\ContentExamples\Lighting\3_2)**.
 
 
 
