#UE4doc #docs - Adjusted bold lettering.

diff --git a/Engine/Documentation/Source/Resources/ContentExamples/Lighting/3_1/3_1.INT.udn b/Engine/Documentation/Source/Resources/ContentExamples/Lighting/3_1/3_1.INT.udn
--- a/Engine/Documentation/Source/Resources/ContentExamples/Lighting/3_1/3_1.INT.udn
+++ b/Engine/Documentation/Source/Resources/ContentExamples/Lighting/3_1/3_1.INT.udn
@@ -41,7 +41,7 @@ Related: Engine/Rendering/LightingAndShadows/LightMobility
 
 Illuminating Engineering Society (IES) profiles are a lighting industry standard method of diagramming the brightness and falloff of light as it exits a particular real world light fixture. They have been used in the CGI industry for many years in photometric lighting. Unreal Engine 4 can use IES profiles to get realistic lighting results based on profiles directly from light manufacturers.
 
-IES profiles render very fast and do not have a significant effect on performance, making them a more ideal option for shaping your light than using a Light Function, which is discussed in [Example 3.2](Resources\ContentExamples\Lighting\3_2). Profiles of various fixtures can be downloaded from light manufacturer websites. See the [IES Profile Documentation](Engine/Rendering\LightingAndShadows\IESLightProfiles) for more information.
+IES profiles render very fast and do not have a significant effect on performance, making them a more ideal option for shaping your light than using a Light Function, which is discussed in **[Example 3.2](Resources\ContentExamples\Lighting\3_2)**. Profiles of various fixtures can be downloaded from light manufacturer websites. See the **[IES Profile Documentation](Engine/Rendering\LightingAndShadows\IESLightProfiles)** for more information.
 
 
 
