#UE4doc #doc - Adjusted bold lettering.

diff --git a/Engine/Documentation/Source/Resources/ContentExamples/Lighting/3_2/3_2.INT.udn b/Engine/Documentation/Source/Resources/ContentExamples/Lighting/3_2/3_2.INT.udn
--- a/Engine/Documentation/Source/Resources/ContentExamples/Lighting/3_2/3_2.INT.udn
+++ b/Engine/Documentation/Source/Resources/ContentExamples/Lighting/3_2/3_2.INT.udn
@@ -41,7 +41,7 @@ Related: Engine/Rendering/LightingAndShadows/LightMobility
 
 [INCLUDE:Engine/Rendering/LightingAndShadows/LightFunctions#LightFunctions]
 
-Since you are primarily using Light Functions to change the overall shape of a light, it is possible to confuse them with IES profiles, which are shown in [Example 3.1](Resources\ContentExamples\Lighting\3_1). However, IES profiles are significantly cheaper to calculate, and should be used instead of Light Functions whenever possible.
+Since you are primarily using Light Functions to change the overall shape of a light, it is possible to confuse them with IES profiles, which are shown in **[Example 3.1](Resources\ContentExamples\Lighting\3_1)**. However, IES profiles are significantly cheaper to calculate, and should be used instead of Light Functions whenever possible.
 
 
 
