#UE4doc #doc - Adjusted bold lettering.

diff --git a/Engine/Documentation/Source/Resources/ContentExamples/Lighting/4_1/4_1.INT.udn b/Engine/Documentation/Source/Resources/ContentExamples/Lighting/4_1/4_1.INT.udn
--- a/Engine/Documentation/Source/Resources/ContentExamples/Lighting/4_1/4_1.INT.udn
+++ b/Engine/Documentation/Source/Resources/ContentExamples/Lighting/4_1/4_1.INT.udn
@@ -38,7 +38,7 @@ Related: Engine/Rendering/LightingAndShadows/LightMobility
 [/REGION]
 
 
-This example illustrates the use of a **Falloff Exponent** to control the falloff of a light. **Falloff** is defined as the rate at which light loses its energy as it travels away from its source. When not using **Inverse Square Falloff**, which is described in [Example 4.2](Resources\ContentExamples\Lighting\4_2), this number is used as an exponent, providing **Exponential Falloff**. Exponential Falloff can often get close to realism, but will generally be best for stylized visuals.
+This example illustrates the use of a **Falloff Exponent** to control the falloff of a light. Falloff is defined as the rate at which light loses its energy as it travels away from its source. When not using **Inverse Square Falloff**, which is described in **[Example 4.2](Resources\ContentExamples\Lighting\4_2)**, this number is used as an exponent, providing **Exponential Falloff**. Exponential Falloff can often get close to realism, but will generally be best for stylized visuals.
 
 
 
