#UE4doc #doc - Adjusted bold lettering.

diff --git a/Engine/Documentation/Source/Resources/ContentExamples/Lighting/5_1/5_1.INT.udn b/Engine/Documentation/Source/Resources/ContentExamples/Lighting/5_1/5_1.INT.udn
--- a/Engine/Documentation/Source/Resources/ContentExamples/Lighting/5_1/5_1.INT.udn
+++ b/Engine/Documentation/Source/Resources/ContentExamples/Lighting/5_1/5_1.INT.udn
@@ -38,10 +38,10 @@ Related: Engine/Rendering/LightingAndShadows/LightMobility
 ![%Resources\ContentExamples\Lighting\5_1:title%](Resources\ContentExamples\Lighting\Lighting_5_1.png)
 [/REGION]
 
-This example shows the effect of changing the Source Radius property on a Static Light. We have two Static Lights with differing colors on either side of a Static Object. The colors are used just to show the difference more readily.
+This example shows the effect of changing the **Source Radius** property on a **Static Light**. We have two Static Lights with differing colors on either side of a Static Object. The colors are used just to show the difference more readily.
 
 
-Source Radius simulates the size of the light source. As you may realize, small light sources, such as small LED lights, produce sharp shadows, particularly when casting very close to a surface. Larger light sources, however, have much softer or blurrier shadows. As previously mentioned in [Example 2.1](Resources/ContentExamples/Lighting/2_1), Static Lights bake their information into a Lightmap. This means that in order for the quality of shadows to remain high, particularly with softer shadows, the Lightmap Resolution property will be key.
+Source Radius simulates the size of the light source. As you may realize, small light sources, such as small LED lights, produce sharp shadows, particularly when casting very close to a surface. Larger light sources, however, have much softer or blurrier shadows. As previously mentioned in **[Example 2.1](Resources/ContentExamples/Lighting/2_1)**, Static Lights bake their information into a Lightmap. This means that in order for the quality of shadows to remain high, particularly with softer shadows, the **Lightmap Resolution** property will be key.
 
 
 
