#UE4doc #doc - Adjusted bold lettering.

diff --git a/Engine/Documentation/Source/Resources/ContentExamples/Lighting/5_2/5_2.INT.udn b/Engine/Documentation/Source/Resources/ContentExamples/Lighting/5_2/5_2.INT.udn
--- a/Engine/Documentation/Source/Resources/ContentExamples/Lighting/5_2/5_2.INT.udn
+++ b/Engine/Documentation/Source/Resources/ContentExamples/Lighting/5_2/5_2.INT.udn
@@ -37,7 +37,7 @@ Related: Engine/Rendering/LightingAndShadows/LightMobility
 ![%Resources\ContentExamples\Lighting\5_2:title%](Resources\ContentExamples\Lighting\Lighting_5_2.png)
 [/REGION]
 
-This example shows the effects of Source Radius on non-Static (Stationary or Movable) lights, particularly with how they contribute to specular reflections. Since Source Radius simulates the size of the light source, it stands to reason that the reflection of the light will change in size to correlate.
+This example shows the effects of **Source Radius** on non-Static (Stationary or Movable) lights, particularly with how they contribute to specular reflections. Since Source Radius simulates the size of the light source, it stands to reason that the reflection of the light will change in size to correlate.
 
 
 
