#UE4doc #doc - Adjusted bold lettering.

diff --git a/Engine/Documentation/Source/Resources/ContentExamples/Lighting/5_3/5_3.INT.udn b/Engine/Documentation/Source/Resources/ContentExamples/Lighting/5_3/5_3.INT.udn
--- a/Engine/Documentation/Source/Resources/ContentExamples/Lighting/5_3/5_3.INT.udn
+++ b/Engine/Documentation/Source/Resources/ContentExamples/Lighting/5_3/5_3.INT.udn
@@ -38,7 +38,7 @@ Related: Engine/Rendering/LightingAndShadows/LightMobility
 [/REGION]
 
 
-This example demonstrates the use of Source Length on reflective surfaces. The purpose of Source Length is to give lights an elongated or tube-like reflective shape, as if the light were coming form a fluorescent light fixture. It should be noted that light emission is still calculated as if it is coming from a single point. The effect exists only within reflections.
+This example demonstrates the use of **Source Length** on reflective surfaces. The purpose of Source Length is to give lights an elongated or tube-like reflective shape, as if the light were coming form a fluorescent light fixture. It should be noted that light emission is still calculated as if it is coming from a single point. The effect exists only within reflections.
 
 The Source Length runs down the Y-Axis of the light, and can be rotated by rotating the light. Also, while the light is selected, Source Length can be visualized by a yellow line stretching out from the light itself.
 
diff --git a/Engine/Documentation/Source/Resources/ContentExamples/Lighting/5_4/5_4.INT.udn b/Engine/Documentation/Source/Resources/ContentExamples/Lighting/5_4/5_4.INT.udn
--- a/Engine/Documentation/Source/Resources/ContentExamples/Lighting/5_4/5_4.INT.udn
+++ b/Engine/Documentation/Source/Resources/ContentExamples/Lighting/5_4/5_4.INT.udn
@@ -37,7 +37,7 @@ Related: Engine/Rendering/LightingAndShadows/LightMobility
 ![%Resources\ContentExamples\Lighting\5_4:title%](Resources\ContentExamples\Lighting\Lighting_5_4.png)
 [/REGION]
 
-This example shows the effect of Min Roughness on a light's reflection. Min Roughness is available in the Light category of the **Details** panel, though you will have to expand the category (click the small white triangle) to see it.
+This example shows the effect of **Min Roughness** on a light's reflection. Min Roughness is available in the Light category of the **Details** panel, though you will have to expand the category (click the small white triangle) to see it.
 
 ![](MinRoughness.png)
 
