#UE4doc #docs - Adjusted bold lettering.

diff --git a/Engine/Documentation/Source/Resources/ContentExamples/Lighting/6_2/6_2.INT.udn b/Engine/Documentation/Source/Resources/ContentExamples/Lighting/6_2/6_2.INT.udn
--- a/Engine/Documentation/Source/Resources/ContentExamples/Lighting/6_2/6_2.INT.udn
+++ b/Engine/Documentation/Source/Resources/ContentExamples/Lighting/6_2/6_2.INT.udn
@@ -37,7 +37,7 @@ Related: Engine/Rendering/LightingAndShadows/LightMobility
 ![%Resources\ContentExamples\Lighting\6_2:title%](Resources\ContentExamples\Lighting\Lighting_6_2.png)
 [/REGION]
 
-This final example shows the Shadow Filter Sharpen property, which is used to boost sharpness at the edges of shadows. It is only useful on Movable or Stationary lights, and only works with shadows from dynamic objects. You need to be close to the shadows to see the full effect. Although not frequently used, this is a nice way to give the illusion that a light source is closer or smaller, creating sharper shadows on objects.
+This final example shows the **Shadow Filter Sharpen** property, which is used to boost sharpness at the edges of shadows. It is only useful on **Movable** or **Stationary** lights, and only works with shadows from dynamic objects. You need to be close to the shadows to see the full effect. Although not frequently used, this is a nice way to give the illusion that a light source is closer or smaller, creating sharper shadows on objects.
 
 In the image below, you can see the example up close. Notice that the shadow on the right is significantly sharper than the one on the left.
 
