#UE4doc #docs - Adjusting bold lettering.

diff --git a/Engine/Documentation/Source/Resources/ContentExamples/MaterialInstances/1_1/1_1.INT.udn b/Engine/Documentation/Source/Resources/ContentExamples/MaterialInstances/1_1/1_1.INT.udn
--- a/Engine/Documentation/Source/Resources/ContentExamples/MaterialInstances/1_1/1_1.INT.udn
+++ b/Engine/Documentation/Source/Resources/ContentExamples/MaterialInstances/1_1/1_1.INT.udn
@@ -54,8 +54,8 @@ Related: Engine/Actors/DecalActor
 
 **Materials** are a collection of properties and a node graph that defines how each pixel of the material is rendered.
 
-A **Material** (or **Material Instance** for that matter) can be applied directly to a mesh, and adjustments can be made to that material at any time; adding, removing, or changing nodes. But, to truly take advantage of Unreal Engine 4 materials, a **Material** needs to be setup with **Parameters**.
+A Material (or **Material Instance** for that matter) can be applied directly to a mesh, and adjustments can be made to that material at any time; adding, removing, or changing nodes. But, to truly take advantage of Unreal Engine 4 materials, a Material needs to be setup with **Parameters**.
 
-**Parameters** are nodes in the **Material Node Graph** that allow that node's value to be changed in by a **Material Instance**.
+Parameters are nodes in the **Material Node Graph** that allow that node's value to be changed in by a Material Instance.
 
 
