UEDOC-2012

Gifs have been replaced by [REDACTED_PERSON]

diff --git a/Engine/Documentation/Source/Support/Builds/ReleaseNotes/2015/4_7/ReleaseNotes_4_7.INT.udn b/Engine/Documentation/Source/Support/Builds/ReleaseNotes/2015/4_7/ReleaseNotes_4_7.INT.udn
--- a/Engine/Documentation/Source/Support/Builds/ReleaseNotes/2015/4_7/ReleaseNotes_4_7.INT.udn
+++ b/Engine/Documentation/Source/Support/Builds/ReleaseNotes/2015/4_7/ReleaseNotes_4_7.INT.udn
@@ -118,7 +118,17 @@ Here is the fun part -- writing Blueprint script code for your new component.  I
 #### Automatic Asset Importing
 Importing meshes, sounds and textures has never been easier. Simply save your source art files under your project's Content folder, and the editor will automatically import the new files, ready to be used in UE4!
 
-![](AutoAssetImport.gif)(convert:false)
+[OBJECT:EmbeddedVideo]
+[PARAMLITERAL:width]
+640
+[/PARAMLITERAL]
+[PARAMLITERAL:height]
+360
+[/PARAMLITERAL]
+[PARAMLITERAL:videoid]
+8iXezqbc2kE
+[/PARAMLITERAL]
+[/OBJECT]
 
 If you update these files from an external application, the editor will detect that and automatically reimport your content! This makes iterating on game content much faster.
 
@@ -189,7 +199,17 @@ Clicking "VR Preview" will minimize the editor window, and pop up your game dire
 #### Leap Motion Plugin *BETA*
 UE4 now ships with support for the **Leap Motion device** through a plugin.  This allows you to use Leap's hand tracking technology directly in your project.  Their API is exposed in both Blueprints and code, so any project can take advantage of the tech.  
 
-![](LeapMotion.gif)(convert:false)
+[OBJECT:EmbeddedVideo]
+[PARAMLITERAL:width]
+640
+[/PARAMLITERAL]
+[PARAMLITERAL:height]
+360
+[/PARAMLITERAL]
+[PARAMLITERAL:videoid]
+pifYPUX4ccU
+[/PARAMLITERAL]
+[/OBJECT]
 
 For examples on how to use the plugin, please see Leap Motion's website at [](https://www.leapmotion.com).  Currently, this plugin is only available to those who build the engine from GitHub, and those who download the plugin directly from Leap's website.
 
@@ -248,7 +268,17 @@ Some other notes:
 #### Instant Bone Controller Preview
 Bone Controller nodes can now be adjusted interactively using the viewport gizmo, rather than having to recompile each time!
 
-![](InstantBonePreview.gif)(convert:false)
+[OBJECT:EmbeddedVideo]
+[PARAMLITERAL:width]
+640
+[/PARAMLITERAL]
+[PARAMLITERAL:height]
+360
+[/PARAMLITERAL]
+[PARAMLITERAL:videoid]
+VthaM3yZbVM
+[/PARAMLITERAL]
+[/OBJECT]
 
 Currently this feature works with TwoBoneIK, ModifyBone & Fabrik nodes, but we will be adding more.
 
