#UE4doc #docs - Adjusting bold lettering.

diff --git a/Engine/Documentation/Source/Resources/ContentExamples/MaterialInstances/1_2/1_2.INT.udn b/Engine/Documentation/Source/Resources/ContentExamples/MaterialInstances/1_2/1_2.INT.udn
--- a/Engine/Documentation/Source/Resources/ContentExamples/MaterialInstances/1_2/1_2.INT.udn
+++ b/Engine/Documentation/Source/Resources/ContentExamples/MaterialInstances/1_2/1_2.INT.udn
@@ -55,9 +55,9 @@ Related: Engine/Actors/DecalActor
 ![](MI_1_2_001.png)(w:900)
 </div>
 
-A **Material Instance** is a child of Material, meaning it has everything the **Parent Material** has, but every **Parameter** can be changed individually on the **Material Instance**.
+A **Material Instance** is a child of Material, meaning it has everything the **Parent Material** has, but every **Parameter** can be changed individually on the Material Instance.
 
-Why use Instances? Materials are very flexible, almost anything you can imagine, the material editor can bring to life. But, some people want to have the complex shader design done by a technical artist, then they take that material, make a **Material Instance** and simply adjust the textures and parameters exposed by the technical artist. That is what **Instances** are good for.
+Why use Instances? Materials are very flexible, almost anything you can imagine, the material editor can bring to life. But, some people want to have the complex shader design done by a technical artist, then they take that material, make a Material Instance and simply adjust the textures and parameters exposed by the technical artist. That is what Instances are good for.
 
-The second advantage of **Instances** is speed. Making changes in the material takes some time because every change requires the shader to compile. Changing the parameters of a **Material Instance** is instantaneous because the shader is already compiled (the exception being the `StaticSwitchParameter`, it is always slow as it needs to recompile the shader as they can change the shader complexity). So if you want to tweak your material, it is recommend doing it with a **Material Instance** because it is much faster.
+The second advantage of Instances is speed. Making changes in the material takes some time because every change requires the shader to compile. Changing the parameters of a Material Instance is instantaneous because the shader is already compiled (the exception being the `StaticSwitchParameter`, it is always slow as it needs to recompile the shader as they can change the shader complexity). So if you want to tweak your material, it is recommend doing it with a Material Instance because it is much faster.
 
