#UE4doc #docs - Adjusting bold lettering.

diff --git a/Engine/Documentation/Source/Resources/ContentExamples/MaterialInstances/1_3/1_3.INT.udn b/Engine/Documentation/Source/Resources/ContentExamples/MaterialInstances/1_3/1_3.INT.udn
--- a/Engine/Documentation/Source/Resources/ContentExamples/MaterialInstances/1_3/1_3.INT.udn
+++ b/Engine/Documentation/Source/Resources/ContentExamples/MaterialInstances/1_3/1_3.INT.udn
@@ -57,5 +57,5 @@ Related: Engine/Actors/DecalActor
 [/REGION]
 </div>
 
-It is possible to make an **Instance** of an **Instance**.  Every parameter that is changed on an **instance** will also change in every child, unless that parameter is overwritten. Changing a parameter is done by enabling the checkbox, if the check box is off, then the parameter will use the value of **its** parent material or material instance.
+It is possible to make an **Instance** of an Instance.  Every parameter that is changed on an instance will also change in every child, unless that parameter is overwritten. Changing a parameter is done by enabling the checkbox, if the check box is off, then the parameter will use the value of **its** parent material or material instance.
 
