#UE4doc #docs - Adjusting bold lettering.

diff --git a/Engine/Documentation/Source/Resources/ContentExamples/MaterialInstances/1_4/1_4.INT.udn b/Engine/Documentation/Source/Resources/ContentExamples/MaterialInstances/1_4/1_4.INT.udn
--- a/Engine/Documentation/Source/Resources/ContentExamples/MaterialInstances/1_4/1_4.INT.udn
+++ b/Engine/Documentation/Source/Resources/ContentExamples/MaterialInstances/1_4/1_4.INT.udn
@@ -58,5 +58,5 @@ Related: Engine/Actors/DecalActor
 
 **Materials** can get very complicated, very fast. **Groups** lets you manage how **Parameters** are listed in the **Material Instance** editor.
 
-Grouping parameters has no effect on how the **Material** is rendered, they are just used to optimize your workflow.
+Grouping parameters has no effect on how the Material is rendered, they are just used to optimize your workflow.
 
