#UE4doc #docs - removed overview header

diff --git a/Engine/Documentation/Source/Resources/ContentExamples/MaterialInstances/MaterialInstances.INT.udn b/Engine/Documentation/Source/Resources/ContentExamples/MaterialInstances/MaterialInstances.INT.udn
--- a/Engine/Documentation/Source/Resources/ContentExamples/MaterialInstances/MaterialInstances.INT.udn
+++ b/Engine/Documentation/Source/Resources/ContentExamples/MaterialInstances/MaterialInstances.INT.udn
@@ -44,7 +44,6 @@ Description:An overview of the Material Instances example level
 
 <!--Insert Path when done [RELATIVE:Resources/ContentExamples/MaterialInstances] -->
 
-## Overview
 
 [REGION:fullwidth]
 
