#UE4doc #docs - removed overview header

diff --git a/Engine/Documentation/Source/Resources/ContentExamples/MaterialNodes/MaterialNodes.INT.udn b/Engine/Documentation/Source/Resources/ContentExamples/MaterialNodes/MaterialNodes.INT.udn
--- a/Engine/Documentation/Source/Resources/ContentExamples/MaterialNodes/MaterialNodes.INT.udn
+++ b/Engine/Documentation/Source/Resources/ContentExamples/MaterialNodes/MaterialNodes.INT.udn
@@ -45,7 +45,7 @@ Description:An overview of the Material Nodes example level.
 [/OBJECT]
 [/VAR]
 
-## Overview
+
 
 [REGION:fullwidth]
 
