UEDOC-2007

Gif has been replaced with A Youtube video

diff --git a/Engine/Documentation/Source/Engine/Blueprints/HowTo/BPHT_7/BPHT_7.INT.udn b/Engine/Documentation/Source/Engine/Blueprints/HowTo/BPHT_7/BPHT_7.INT.udn
--- a/Engine/Documentation/Source/Engine/Blueprints/HowTo/BPHT_7/BPHT_7.INT.udn
+++ b/Engine/Documentation/Source/Engine/Blueprints/HowTo/BPHT_7/BPHT_7.INT.udn
@@ -23,7 +23,17 @@ Description:In this final step, examples are given how you can take this project
 
 Using what you have learned over the course of this project, try to do the following: 
 
-[![](UpdatedImage.png)(w:720)](UpdatedExample.mp4)
+[OBJECT:EmbeddedVideo]
+[PARAMLITERAL:width]
+640
+[/PARAMLITERAL]
+[PARAMLITERAL:height]
+360
+[/PARAMLITERAL]
+[PARAMLITERAL:videoid]
+yut5Feu1KxI
+[/PARAMLITERAL]
+[/OBJECT]
 
 * In the **FirstPersonCharacter** Blueprint, award a higher score for hitting the rock on the 1st versus 2nd or 3rd shot. 
 * Also in the **FirstPersonCharacter** Blueprint, set it so the player cannot shoot unless a rock has been spawned. 
