UEDOC-2003

Converted Gifs into YouTube Videos

diff --git a/Engine/Documentation/Source/Resources/Showcases/MatineeFightScene/MatineeFightScene.INT.udn b/Engine/Documentation/Source/Resources/Showcases/MatineeFightScene/MatineeFightScene.INT.udn
--- a/Engine/Documentation/Source/Resources/Showcases/MatineeFightScene/MatineeFightScene.INT.udn
+++ b/Engine/Documentation/Source/Resources/Showcases/MatineeFightScene/MatineeFightScene.INT.udn
@@ -71,7 +71,17 @@ You can also head to the [](Engine\Matinee) documentation for more details on ho
 
 Before breaking down the scene, you can view the end result below: 
 
-[![](OriginalMovieThumbnail.png)(w:640)](FightSceneOriginal.wmv)
+[OBJECT:EmbeddedVideo]
+[PARAMLITERAL:width]
+640
+[/PARAMLITERAL]
+[PARAMLITERAL:height]
+360
+[/PARAMLITERAL]
+[PARAMLITERAL:videoid]
+E0L7pBBhwUo
+[/PARAMLITERAL]
+[/OBJECT]
 
 
 * The scene consists of five characters that interact with one another during the course of the 68 second sequence. 
@@ -691,7 +701,17 @@ You could also create a **Template** using this project for new projects. Refer
 
 Below is an example remix of the original sequence: 
 
-[![](Remix_Thumbnail.png)(w:640)](FightSceneRemix.wmv)
+[OBJECT:EmbeddedVideo]
+[PARAMLITERAL:width]
+640
+[/PARAMLITERAL]
+[PARAMLITERAL:height]
+360
+[/PARAMLITERAL]
+[PARAMLITERAL:videoid]
+kzXg1NY4Aag
+[/PARAMLITERAL]
+[/OBJECT]
 
 As you can see from the video, there a number of different things you can do to modify the scene. In the video above, the following modifications were used: 
 
