#UE4 Docs: Checking in the progress with the first Android docs.

diff --git a/Engine/Documentation/Source/Platforms/Android/ReducingAPKSize/ReducingAPHSizes.INT.udn b/Engine/Documentation/Source/Platforms/Android/ReducingAPKSize/ReducingAPHSizes.INT.udn
new file mode 100644
--- /dev/null
+++ b/Engine/Documentation/Source/Platforms/Android/ReducingAPKSize/ReducingAPHSizes.INT.udn
@@ -0,0 +1,104 @@
+Availability:Docs
+Title: Reducing APK Packages size 
+Crumbs:%ROOT%, Platforms, Platforms/Android
+Description:How to reduce the size of your packaged Android game.
+Skilllevel: Intermediate
+Version: 4.8
+
+[TOC(start:2 end:3)]
+
+No matter what type of Android device your UE4 project is targeting reducing the size of the Android application package or **APK** to be under the 50MB limit for the Google play store can be challenging. 
+In the following tutorial we will cover steps you can take to help reduce your projects final APK package size to be as small as possible using nothing but the tools provided to you in the UE4 Editor. 
+
+## Creating a new empty project
+When you first start working on your Android based mobile project you might be tempted to use an existing project as a base to work from or create a new project that has the Starter Content enabled. 
+**DO NOT** do this and instead create a new, completely empty C++ or Blueprint based project and then use the [Migration Tool](Engine/Content/Browser/UserGuide/Migrate) to bring over the content that you want to use. 
+This way you ensure that the only content in your project is content that you need to be there.
+
+## Compressing Cooked Content
+The easiest and quickest way to reduce your APK package size is to tell UE4 to compress the APK packages during the packing process. 
+To turn on package compression you need to do the following in the UE4 Editor.
+
+[region:note]
+Please note that the Create compressed cooked packages option is only available in UE versions 4.8 and later.
+[/region]
+
+1. First open up the projects settings by going to the Main Toolbar and selecting the **Edit** option then selecting **Project Settings**.
+
+	![](T_Project_Settings_Menu.png)
+
+1. Under the Project section click on the **Packaging** section to show the options for how the project will be packaged.
+
+	![](T_Project_Packaging.png)
+
+1. Click on the small white triangle that is at the bottom of the Packaging settings to expose the **Advanced Project settings**.
+
+	![](T_Expose_Advanced_Props.png)
+
+1. Look for the option labeled **Create compressed cooked packages** and enable it if not already enabled.
+
+	![](T_Enable_Comp_Packages.png)
+
+Now make note of the current size of your projects APK package and then re-package your project. 
+When finished compare the sizes and you should now see, depending on the content used in the project, quite a big difference between the sizes of the two APK files. 
+In fact it is not uncommon for some projects APK packages to be reduced in size by up to 50% when create compressed cooked packages is enabled.
+
+##  Setting up a Project's Levels
+One overlooked area that can lead to a bloating of APK package sizes is failing to properly setup your projects options like which levels should be loaded by default or which levels should be used between level transitions. 
+To set which level or levels should be used for this type of interaction you can do the following.
+
+1. First on the **Main Toolbar** got to  **Edit** > **Project Settings**.
+
+	 ![](T_Project_Settings.png)
+
+1. Then under **Project** in the **Maps & Modes** section, set the map in the **Default Maps** to a map that is specific to your project.
+
+	![](T_Set_Maps.png)
+
+## Checking what Content will be Cooked
+To ensure that only the content that is related to your project is added to your APK file, you can check which content is being added to your projects APK by looking in the **Cooked** folder in your project.
+You can find the Cooked folder by going to **(ProjectName)\Saved\Cooked** and then choosing the folder with the format you cooked your project for.
+
+![](T_Cooked_Content_Folder.png)
+
+The image below shows the Texture folder from the Match 3 sample game. 
+The content in the folder has been sorted by size so it is easier to see what assets are the biggest.
+Once the biggest assets have been established you can then examine the assets inside of UE4 to see if in fact they can be reduced in size without sacrificing the integrity of the asset.
+
+![](T_Cooked_Content.png)
+
+## Removing Unused Content
+When your project is finally ready to be packaged up for the store make sure to first remove any unused or testing content by selecting it in the Content Browser and deleting it.
+You can delete content by first selecting the content you want to delete and then either press the **Delete** key on the keyboard or use the Delete option in the Right - Click menu. 
+When you do try and delete something from your project the **Delete Assets** menu will be displayed like in the following image.
+
+[region:tip]
+This is the preferred method of deleting objects in UE4 and should be used over just simply deleting/removing the assets from their location in the content folder.
+[/region]
+
+![](T_Delete_Menu.png)
+
+The Delete Assets menu will inform you if the asset you are trying to delete is referenced by another asset or not. 
+If it is not referenced you can simply delete it, however if it is referenced by other assets you can use the **Replace References** option to replace the reference to that asset to another asset that is supposed to be packaged in your project.
+
+
+## Removing Unused Plug-ins
+Disabling any unused plug-in's inside of the UE4 Editor is another way to ensure that un-need content and code do not make it into your projects final APK package file.
+This is because some Plug-in's require a certain set of base assets and code to work correctly.
+If the Plug-in is not disabled the assets and code that are required to make the plug-in work will not be removed from your project.
+While this might not reduce your projects size as much as reducing or removing a large Texture would, every little bit can help when trying to slim your project down to the 50MB size.
+To disable a plug-in inside of UE4 you need to do the following. 
+
+[region:note]
+	Make sure you fully test your project to ensure that the that the plug-ins you are disabling do not interfere with the functionality of your project.
+[/region]
+
+1. First open up the Plug-ins Manager by going to the **Main Toolbar** and clicking on the **Window** option and then selecting the **Plug-in** option.
+
+	![](T_Open_Plugins_Browser.png)
+
+1. From the Plug-in Browser select the different sections and disable any plug-in's you are not using by un-checking the check mark box next to **Enabled**.
+  When you have disabled all the Plug-ins you do not need make sure to re-package your project to see the results.
+
+	
+
