UEDOC-1914

Converted Gifs into YouTube Videos

diff --git a/Engine/Documentation/Source/Gameplay/HowTo/CharacterMovement/Blueprints/Overview.INT.udn b/Engine/Documentation/Source/Gameplay/HowTo/CharacterMovement/Blueprints/Overview.INT.udn
--- a/Engine/Documentation/Source/Gameplay/HowTo/CharacterMovement/Blueprints/Overview.INT.udn
+++ b/Engine/Documentation/Source/Gameplay/HowTo/CharacterMovement/Blueprints/Overview.INT.udn
@@ -8,7 +8,17 @@ This "How To" series covers all the steps needed to create a fully playable and
 
 For an example of the finished product, click the video sample below. 
 
-[![](CharacterMovement_Thumbnail.png)(w:640)](CharaterMovement.mp4)
+[OBJECT:EmbeddedVideo]
+[PARAMLITERAL:width]
+640
+[/PARAMLITERAL]
+[PARAMLITERAL:height]
+360
+[/PARAMLITERAL]
+[PARAMLITERAL:videoid]
+EhspjHjDmtI
+[/PARAMLITERAL]
+[/OBJECT]
  
 For this series, we are using the **Animation Starter Pack**, which is available for free on the **Marketplace**, as it contains several animations that can be used to create character movement. Once you have downloaded the Animation Pack, you can add the pack to an existing project (preferably a New Blank Project) via the **Add to project** button in the Launcher then specifying which project you want to add the assets to. 
 
