#UE4doc #docs - Adjusted bold lettering.

diff --git a/Engine/Documentation/Source/Resources/ContentExamples/MaterialNodes/1_7/1_7.INT.udn b/Engine/Documentation/Source/Resources/ContentExamples/MaterialNodes/1_7/1_7.INT.udn
--- a/Engine/Documentation/Source/Resources/ContentExamples/MaterialNodes/1_7/1_7.INT.udn
+++ b/Engine/Documentation/Source/Resources/ContentExamples/MaterialNodes/1_7/1_7.INT.udn
@@ -56,7 +56,7 @@ Related: Engine/Actors/DecalActor
 ![](1_7_Opacity_Subsurface.png)
 [/REGION]
 
-Opacity has a different behavior when used with Subsurface Materials. Instead of controlling how much you can see through an object, it instead controls how much light is allowed to scatter throughout the surface. You can see in the example how a low opacity value shows a significant amount of light moving through the surface while while a high opacity value shows very little scatter.
+**Opacity** has a different behavior when used with Subsurface Materials. Instead of controlling how much you can see through an object, it instead controls how much light is allowed to scatter throughout the surface. You can see in the example how a low opacity value shows a significant amount of light moving through the surface while while a high opacity value shows very little scatter.
 
 
 
