#UE4 Docs: aded missing new Editor Manual page

diff --git a/Engine/Documentation/Source/Engine/Editor/EditorManual.INT.udn b/Engine/Documentation/Source/Engine/Editor/EditorManual.INT.udn
new file mode 100644
--- /dev/null
+++ b/Engine/Documentation/Source/Engine/Editor/EditorManual.INT.udn
@@ -0,0 +1,70 @@
+Availability:Public
+Title: Unreal Editor Manual
+Crumbs:%ROOT%, Engine
+Description:Guide to using and configuring the Unreal Editor toolset.
+Navigation:topic
+
+[VAR:Params]
+[PARAM:image]
+	![%Engine/Editor:title%](Engine/UI/UserInterfaceTopic.png)(convert:false)
+[/PARAM]
+[PARAM:icon]
+	![](%ROOT%/editor_icon.png)(convert:false)
+[/PARAM]
+[PARAM:title]
+	%Engine/Editor:title%
+[/PARAM]
+[PARAM:description]
+	%Engine/Editor:description%
+[/PARAM]
+[PARAM:path]
+	[RELATIVE:Engine/Editor]
+[/PARAM]
+[/VAR]
+
+[VAR:HomeParams]
+[PARAM:image]
+	![%Engine/Editor:title%](Engine/UI/UserInterfaceTopic.png)(convert:false)
+[/PARAM]
+[PARAM:icon]
+	![](%ROOT%/editor_icon.png)(convert:false)
+[/PARAM]
+[PARAM:title]
+	%Engine/Editor:title%
+[/PARAM]
+[PARAM:description]
+	%Engine/Editor:description%
+[/PARAM]
+[PARAM:path]
+	[RELATIVE:Engine/Editor]
+[/PARAM]
+[/VAR]
+
+[VAR:Topic]
+[OBJECT:Topic]
+%params%
+[/OBJECT]
+[/VAR]
+
+[VAR:TopicCompact]
+[OBJECT:TopicCompact]
+%params%
+[/OBJECT]
+[/VAR]
+
+[VAR:Role]
+[OBJECT:Role]
+%homeparams%
+[/OBJECT]
+[/VAR]
+
+[REGION:topics third]
+%Engine/UI:topic%
+%GettingStarted/HowTo:topic%
+%Engine/UI/LevelEditor:topic%
+%Engine/Levels:topic%
+%Engine/Actors:topic%
+%Engine/Components:topic%
+%Engine/Content:topic%
+%Engine/Basics/Projects:topic%
+[/REGION]
\ No newline at end of file
