#UE4doc #docs - Adjusted bold lettering.

diff --git a/Engine/Documentation/Source/Resources/ContentExamples/MaterialProperties/1_1/1_1.INT.udn b/Engine/Documentation/Source/Resources/ContentExamples/MaterialProperties/1_1/1_1.INT.udn
--- a/Engine/Documentation/Source/Resources/ContentExamples/MaterialProperties/1_1/1_1.INT.udn
+++ b/Engine/Documentation/Source/Resources/ContentExamples/MaterialProperties/1_1/1_1.INT.udn
@@ -59,11 +59,11 @@ Blend mode is the very foundation of a material.
 
 | Blend Mode | Description |
 | ---- | ----- |
-| **Opaque** | For a Solid object, **Opaque** is the cheapest material. But, as its name implies, it has no support for translucency. That means it will never create overdraw, and rendering one pixel of the material will always take the same time depending on the complexity of the **Node Graph**. |
+| **Opaque** | For a Solid object, Opaque is the cheapest material. But, as its name implies, it has no support for translucency. That means it will never create overdraw, and rendering one pixel of the material will always take the same time depending on the complexity of the **Node Graph**. |
 | **Masked** |  If you need some pixels not drawn on your material, this is the way to go. The material only creates overdraw for the area that is masked out. But the material is 100\% opaque or 100\% see through. The bigger the area that is masked out the more overdraw is created. That means more expensive. |
-| **Translucent** | This is the most expensive **BlendMode**. It can look like a **Masked** and **Opaque** but will have an additional effect on performance regardless of actually being transparent or not. |
-| **Additive** | Is a cheaper alternative to **Translucent**. It does not need to sort the pixels, they just get added to the scene. But that means it can only brighten. Black pixels on the material have no effect on the scene.|
-| **Modulated** | Is the second cheapest alternative to **Translucent**, but it multiplies its Pixels with the scene. That means it can only darken the scene.|
+| **Translucent** | This is the most expensive **BlendMode**. It can look like a Masked and Opaque but will have an additional effect on performance regardless of actually being transparent or not. |
+| **Additive** | Is a cheaper alternative to Translucent. It does not need to sort the pixels, they just get added to the scene. But that means it can only brighten. Black pixels on the material have no effect on the scene.|
+| **Modulated** | Is the second cheapest alternative to Translucent, but it multiplies its Pixels with the scene. That means it can only darken the scene.|
 
 
 
