UEDOC-1911

Converted Gifs into YouTube Videos

diff --git a/Engine/Documentation/Source/Engine/Rendering/ParticleSystems/Reference/Modules/Color/ColorModules.INT.udn b/Engine/Documentation/Source/Engine/Rendering/ParticleSystems/Reference/Modules/Color/ColorModules.INT.udn
--- a/Engine/Documentation/Source/Engine/Rendering/ParticleSystems/Reference/Modules/Color/ColorModules.INT.udn
+++ b/Engine/Documentation/Source/Engine/Rendering/ParticleSystems/Reference/Modules/Color/ColorModules.INT.udn
@@ -17,7 +17,17 @@ Color modules impact the color of the emitted particles.
 
 ## Initial Color
 
-![psys_initial_color.gif](psys_initial_color.gif)(convert:false)
+[OBJECT:EmbeddedVideo]
+[PARAMLITERAL:width]
+640
+[/PARAMLITERAL]
+[PARAMLITERAL:height]
+360
+[/PARAMLITERAL]
+[PARAMLITERAL:videoid]
+WCG-PdJ41rw
+[/PARAMLITERAL]
+[/OBJECT]
 
 The Initial Color module sets the color of a particle at spawn time. The particle will retain this color over its lifetime unless another color module overrides it. It contains the following member(s):
 
@@ -57,7 +67,17 @@ In Spawn, the module retrieves the appropriate values from the distributions usi
 
 ## Color Over Life
 
-![psys_color_life.gif](psys_color_life.gif)(convert:false)
+[OBJECT:EmbeddedVideo]
+[PARAMLITERAL:width]
+640
+[/PARAMLITERAL]
+[PARAMLITERAL:height]
+360
+[/PARAMLITERAL]
+[PARAMLITERAL:videoid]
+zlZ1dvYsjrk
+[/PARAMLITERAL]
+[/OBJECT]
 
 The Color Over Life module is used to set the color of a particle over its lifetime. It contains the following member(s):
 
