UEDOC-1910

Converted gifs into YouTube Videos

diff --git a/Engine/Documentation/Source/Engine/Rendering/ParticleSystems/Reference/Modules/Collision/CollisionModules.INT.udn b/Engine/Documentation/Source/Engine/Rendering/ParticleSystems/Reference/Modules/Collision/CollisionModules.INT.udn
--- a/Engine/Documentation/Source/Engine/Rendering/ParticleSystems/Reference/Modules/Collision/CollisionModules.INT.udn
+++ b/Engine/Documentation/Source/Engine/Rendering/ParticleSystems/Reference/Modules/Collision/CollisionModules.INT.udn
@@ -11,7 +11,17 @@ These modules handle how particles will deal with Actor collisions.
 
 ### Collision
 
-![psys_collision.gif](psys_collision.gif)(convert:false)
+[OBJECT:EmbeddedVideo]
+[PARAMLITERAL:width]
+640
+[/PARAMLITERAL]
+[PARAMLITERAL:height]
+360
+[/PARAMLITERAL]
+[PARAMLITERAL:videoid]
+uD7MQf6sZZU
+[/PARAMLITERAL]
+[/OBJECT]
 
 Parameters for having the emitter handle collisions of particles. It contains the following member(s):
 
