#UE4doc #docs - Adjusted bold lettering.

diff --git a/Engine/Documentation/Source/Resources/ContentExamples/MaterialProperties/1_7/1_7.INT.udn b/Engine/Documentation/Source/Resources/ContentExamples/MaterialProperties/1_7/1_7.INT.udn
--- a/Engine/Documentation/Source/Resources/ContentExamples/MaterialProperties/1_7/1_7.INT.udn
+++ b/Engine/Documentation/Source/Resources/ContentExamples/MaterialProperties/1_7/1_7.INT.udn
@@ -58,6 +58,6 @@ Related: Engine/Rendering/Materials/MaterialProperties
 ![](MP_1_7_001.png)
 [/REGION]
 
-If the Translucent Material is set Volumetric Directional Light mode, you can set the intensity of the shading.
+If the _Translucent_ Material is set **Volumetric Directional Light** mode, you can set the intensity of the shading.
 
 
