#UE4doc #docs - Adjusted bold lettering.

diff --git a/Engine/Documentation/Source/Resources/ContentExamples/MaterialProperties/1_8/1_8.INT.udn b/Engine/Documentation/Source/Resources/ContentExamples/MaterialProperties/1_8/1_8.INT.udn
--- a/Engine/Documentation/Source/Resources/ContentExamples/MaterialProperties/1_8/1_8.INT.udn
+++ b/Engine/Documentation/Source/Resources/ContentExamples/MaterialProperties/1_8/1_8.INT.udn
@@ -62,6 +62,6 @@ Related: Engine/Rendering/Materials/MaterialProperties
 
 The **Flat** setting simply splits up the triangles. While this would not be of much use by itself, adjusting the world position of the new vertices using the normal map and displacement map, can create new surface details.
 
-**PN Triangles** smoothes the object. It does require the mesh to have at least one smoothing group to make the smoothing work. As with the **Flat** setting, textures can be used to displace the tessellated vertices giving it a very detailed look.
+**PN Triangles** smoothes the object. It does require the mesh to have at least one smoothing group to make the smoothing work. As with the Flat setting, textures can be used to displace the tessellated vertices giving it a very detailed look.
 
 
diff --git a/Engine/Documentation/Source/Resources/ContentExamples/MaterialProperties/1_9/1_9.INT.udn b/Engine/Documentation/Source/Resources/ContentExamples/MaterialProperties/1_9/1_9.INT.udn
--- a/Engine/Documentation/Source/Resources/ContentExamples/MaterialProperties/1_9/1_9.INT.udn
+++ b/Engine/Documentation/Source/Resources/ContentExamples/MaterialProperties/1_9/1_9.INT.udn
@@ -59,5 +59,5 @@ Related: Engine/Rendering/Materials/MaterialProperties
 ![](MP_1_9_001.png)
 [/REGION]
 
-We learned earlier that the **Masked BlendMode** only has pixels that are completely opaque or completely transparent. You can use a grayscale texture or a single channel from a texture to use as a mask. **Clip Value** lets you set a threshold when the pixel will be opaque or transparent. For textures that have thin sections, those areas will disappear in the distance as the mipmap forces the white out of the image or channel, adjusting **Clip Value** will allow you to adjust when that occurs.
+We learned earlier that the **Masked BlendMode** only has pixels that are completely opaque or completely transparent. You can use a grayscale texture or a single channel from a texture to use as a mask. **Clip Value** lets you set a threshold when the pixel will be opaque or transparent. For textures that have thin sections, those areas will disappear in the distance as the mipmap forces the white out of the image or channel, adjusting Clip Value will allow you to adjust when that occurs.
 
