UEDOC-1909

gifs have converted into YouTube videos

diff --git a/Engine/Documentation/Source/Engine/Rendering/ParticleSystems/Reference/Modules/Beam/BeamModules.INT.udn b/Engine/Documentation/Source/Engine/Rendering/ParticleSystems/Reference/Modules/Beam/BeamModules.INT.udn
--- a/Engine/Documentation/Source/Engine/Rendering/ParticleSystems/Reference/Modules/Beam/BeamModules.INT.udn
+++ b/Engine/Documentation/Source/Engine/Rendering/ParticleSystems/Reference/Modules/Beam/BeamModules.INT.udn
@@ -81,7 +81,17 @@ The Beam Modifier module allows either the Source or Target of the beam to be mo
 
 ## Noise
 
-![psys_beam_noise.gif](psys_beam_noise.gif)(convert:false)
+[OBJECT:EmbeddedVideo]
+[PARAMLITERAL:width]
+640
+[/PARAMLITERAL]
+[PARAMLITERAL:height]
+360
+[/PARAMLITERAL]
+[PARAMLITERAL:videoid]
+TfgyUi-QAAg
+[/PARAMLITERAL]
+[/OBJECT]
 
 The Beam Noise module allows for injecting noise along the beam. It provides the following properties:
 
