#UE4doc #docs - Adjusted bold lettering.

diff --git a/Engine/Documentation/Source/Resources/ContentExamples/Materials/Materials.INT.udn b/Engine/Documentation/Source/Resources/ContentExamples/Materials/Materials.INT.udn
--- a/Engine/Documentation/Source/Resources/ContentExamples/Materials/Materials.INT.udn
+++ b/Engine/Documentation/Source/Resources/ContentExamples/Materials/Materials.INT.udn
@@ -49,10 +49,9 @@ Related: Engine/Actors/DecalActor
 [TOC(start:2 end:2)]
 
 
-## Overview
 
 The process of determining the surface appearance of the geometry in your levels, of a Static Mesh Actor, or of a character are all handled with **Materials**. 
-While the **Material** system in Unreal Engine 4 is quite comprehensive, the **Material Content Examples** break down the systems that define a **Material** into 5 separate maps, each of which are outlined on this page.
+While the Material system in Unreal Engine 4 is quite comprehensive, the **Material Content Examples** break down the systems that define a Material into 5 separate maps, each of which are outlined on this page.
 
 
 [region:note]
@@ -64,53 +63,53 @@ You can find additional resources on how an example was constructed by clicking
 
 ![](MaterialNodes_Header.png)(w:640)
 
-The **Material_Nodes** map covers the various nodes that are used to manipulate the properties that make up a **Material's** appearance. Whether you are looking to create a glass looking surface, a bumpy road, or even have a **Material** animate to create moving clouds or water; this map covers the nodes you will use to obtain your end result. 
+The **Material_Nodes** map covers the various nodes that are used to manipulate the properties that make up a Material's appearance. Whether you are looking to create a glass looking surface, a bumpy road, or even have a Material animate to create moving clouds or water; this map covers the nodes you will use to obtain your end result. 
 
 | Example | What is Demonstrated |
 | ------- | ------------ |
-| **[1.1 BaseColor](Resources\ContentExamples\MaterialNodes\1_1)** | The **BaseColor** is the basic color or texture in a **Material**. Also shown in this example are **Material Instances** which are modified copies of a **Material**. |
-| **[1.2 Metallic](Resources\ContentExamples\MaterialNodes\1_2)** | The **Metallic** node determines how metallic the **Material** appears (e.g. a non-metal surface like plastic or wood versus a metal surface like metal or a mirror). |
-| **[1.3 Specular](Resources\ContentExamples\MaterialNodes\1_3)** | **Specular** is used to scale the current amount of specularity on non-metallic surfaces (e.g. mostly used to mask out reflections for micro shadow in a **Material**). |
-| **[1.4 Roughness](Resources\ContentExamples\MaterialNodes\1_4)** | **Roughness** affects how rough or smooth the surface of a **Material** appears. |
-| **[1.5 Emissive](Resources\ContentExamples\MaterialNodes\1_5)** | The **Emissive** property affects the light or glow of a **Material** applied to an object. |
-| **[1.6 Opacity](Resources\ContentExamples\MaterialNodes\1_6)** | **Opacity** determines the translucency of a **Material**.  |
+| **[1.1 BaseColor](Resources\ContentExamples\MaterialNodes\1_1)** | The **BaseColor** is the basic color or texture in a Material. Also shown in this example are **Material Instances** which are modified copies of a Material. |
+| **[1.2 Metallic](Resources\ContentExamples\MaterialNodes\1_2)** | The **Metallic** node determines how metallic the Material appears (e.g. a non-metal surface like plastic or wood versus a metal surface like metal or a mirror). |
+| **[1.3 Specular](Resources\ContentExamples\MaterialNodes\1_3)** | **Specular** is used to scale the current amount of specularity on non-metallic surfaces (e.g. mostly used to mask out reflections for micro shadow in a Material). |
+| **[1.4 Roughness](Resources\ContentExamples\MaterialNodes\1_4)** | **Roughness** affects how rough or smooth the surface of a Material appears. |
+| **[1.5 Emissive](Resources\ContentExamples\MaterialNodes\1_5)** | The **Emissive** property affects the light or glow of a Material applied to an object. |
+| **[1.6 Opacity](Resources\ContentExamples\MaterialNodes\1_6)** | **Opacity** determines the translucency of a Material.  |
 | **[1.7 Opacity SubSurface Material](Resources\ContentExamples\MaterialNodes\1_7)** | Instead of affecting translucency, the amount of light scattering throughout the surface is determined with **Opacity SubSurface Materials**.  |
-| **[1.8 OpacityMask](Resources\ContentExamples\MaterialNodes\1_8)** | Similar to **Opacity**, but only used in **Masked Blend Mode**, this makes an aspect of the **Material** visible or invisible as shades of gray are not used (e.g. this is perfect for complex solid surfaces like wire mesh or chain link fences). |
-| **[1.9 Normal](Resources\ContentExamples\MaterialNodes\1_9)** | By using a **Normal Map**, this allows a **Material** to take on an improved physical appearance (e.g. creating a brick wall or a cobblestone road). |
-| **[1.10 WorldPositionOffset](Resources\ContentExamples\MaterialNodes\1_10)** | Allows for the vertices of a mesh to be manipulated in world space by the **Material** (e.g. useful for making objects move, change shape or other effects). |
+| **[1.8 OpacityMask](Resources\ContentExamples\MaterialNodes\1_8)** | Similar to **Opacity**, but only used in **Masked Blend Mode**, this makes an aspect of the Material visible or invisible as shades of gray are not used (e.g. this is perfect for complex solid surfaces like wire mesh or chain link fences). |
+| **[1.9 Normal](Resources\ContentExamples\MaterialNodes\1_9)** | By using a **Normal Map**, this allows a Material to take on an improved physical appearance (e.g. creating a brick wall or a cobblestone road). |
+| **[1.10 WorldPositionOffset](Resources\ContentExamples\MaterialNodes\1_10)** | Allows for the vertices of a mesh to be manipulated in world space by the Material (e.g. useful for making objects move, change shape or other effects). |
 | **[1.11 WorldDisplacement](Resources\ContentExamples\MaterialNodes\1_11)** | Similar to 1.10, but this uses Tessellation vertices rather than the base vertices of the mesh.  |
-| **[1.12 Tesselation Multiplier](Resources\ContentExamples\MaterialNodes\1_12)** | Controlling the amount of Tessellation along the surface of a **Material** (adding more detail where needed).  |
-| **[1.13 SubSurfaceColor](Resources\ContentExamples\MaterialNodes\1_13)** | Only used when the **Shading Model** property is set to **Subsurface**, this allows you to add a color to your **Material** to simulate shifts in color when light passes through the surface (e.g. human skin may have a red subsurface color to simulate blood beneath the surface). |
-| **[1.14 Refraction](Resources\ContentExamples\MaterialNodes\1_14)** | Simulates **Refraction** in the **Material** for things like glass and water, which refract light that passes through them. |
+| **[1.12 Tesselation Multiplier](Resources\ContentExamples\MaterialNodes\1_12)** | Controlling the amount of Tessellation along the surface of a Material (adding more detail where needed).  |
+| **[1.13 SubSurfaceColor](Resources\ContentExamples\MaterialNodes\1_13)** | Only used when the **Shading Model** property is set to **Subsurface**, this allows you to add a color to your Material to simulate shifts in color when light passes through the surface (e.g. human skin may have a red subsurface color to simulate blood beneath the surface). |
+| **[1.14 Refraction](Resources\ContentExamples\MaterialNodes\1_14)** | Simulates **Refraction** in the Material for things like glass and water, which refract light that passes through them. |
 
 
 ## Material_Advanced Map
 
 ![](MaterialAdvanced_Header.png)(w:640)
 
-The **Material_Advanced** map provides four advanced techniques using **Materials**. The first, a **Master Material** demonstrates how to improve workflow by creating a complex "**Master Material**" that is used by other "instances" which inherit from the master and can then be modified to fit your needs without recreating the entire network. 
+The **Material_Advanced** map provides four advanced techniques using Materials. The first, a **Master Material** demonstrates how to improve workflow by creating a complex "Master Material" that is used by other "instances" which inherit from the master and can then be modified to fit your needs without recreating the entire network. 
 The second and third examples discuss the use of **Material Functions** and how these "snipits" of a Material Graph can be used in other Materials or used to recreate an entire Material Network. The fourth and last example discusses how **Vertex Animation** can be used to create ambient animations entirely inside a Material. 
 
 | Example | What is Demonstrated |
 | ------- | ------------ |
-| **1.1 Example for flexible Master Material** | A **Master Material** is used as the foundation for other **Material Instances** which inherit from the master but are altered individually to generate different results.  |
-| **1.2 Material Functions** | An entire **Material** network is recreated through the use of a **Material Function** (e.g. these can be used to quickly recreate complex **Materials** that can then be used and modified in any **Material** layer. |
-| **1.3 Material Layers Using Material Functions** | **Material Functions** are used in layers which are then edited to create variations of complex **Materials** (e.g. multiple rocket ships are shown each with variations of a complex **Material**). |
-| **1.4 Vertex Animation** | **Vertex Animation** is used to simulate cloth movement (e.g. an animated flag is shown where the animation is handled entirely through the **Material**). |
+| **1.1 Example for flexible Master Material** | A Master Material is used as the foundation for other Material Instances which inherit from the master but are altered individually to generate different results.  |
+| **1.2 Material Functions** | An entire Material network is recreated through the use of a Material Function (e.g. these can be used to quickly recreate complex Materials that can then be used and modified in any Material layer. |
+| **1.3 Material Layers Using Material Functions** | Material Functions are used in layers which are then edited to create variations of complex Materials (e.g. multiple rocket ships are shown each with variations of a complex Material). |
+| **1.4 Vertex Animation** | **Vertex Animation** is used to simulate cloth movement (e.g. an animated flag is shown where the animation is handled entirely through the Material). |
 
 
 ## Material_Instances Map
 
 ![](MaterialInstance_Header.png)(w:640)
 
-The **Material_Instances** map breaks down how instances are generated, how **Parameters** are passed from parent to child and how instances are managed inside the **Material Instance** editor through groups. 
+The **Material_Instances** map breaks down how instances are generated, how **Parameters** are passed from parent to child and how instances are managed inside the Material Instance editor through groups. 
 
 | Example | What is Demonstrated |
 | ------- | ------------ |
-| **[1.1 Material](Resources\ContentExamples\MaterialInstances\1_1)** | This example is similar to the **Material_Advanced** map's 1.1 example where a **Master Material** is demonstrated that can be used as instances. |
-| **[1.2 Instance](Resources\ContentExamples\MaterialInstances\1_2)** | A **Material Instance** (child of a Material) is shown where the instance has every **Parameter** of the parent but can be changed individually on the child without affect the parent (e.g. good for modifying complex **Materials** while preserving the original). |
-| **[1.3 Instance of an Instance](Resources\ContentExamples\MaterialInstances\1_3)** | **Instances** of **Instances** where each instance inherits from its parent and is passed down to its children **Materials**.  |
-| **[1.4 Groups](Resources\ContentExamples\MaterialInstances\1_4)** | While **Groups** have no effect on how the **Material** is rendered, they allow you to manage how **Parameters** are listed inside the **Material Instance** editor. |
+| **[1.1 Material](Resources\ContentExamples\MaterialInstances\1_1)** | This example is similar to the **Material_Advanced** map's 1.1 example where a Master Material is demonstrated that can be used as instances. |
+| **[1.2 Instance](Resources\ContentExamples\MaterialInstances\1_2)** | A Material Instance (child of a Material) is shown where the instance has every **Parameter** of the parent but can be changed individually on the child without affect the parent (e.g. good for modifying complex Materials while preserving the original). |
+| **[1.3 Instance of an Instance](Resources\ContentExamples\MaterialInstances\1_3)** | Instances of Instances where each instance inherits from its parent and is passed down to its children Materials.  |
+| **[1.4 Groups](Resources\ContentExamples\MaterialInstances\1_4)** | While **Groups** have no effect on how the Material is rendered, they allow you to manage how Parameters are listed inside the Material Instance editor. |
 
 
 ## Material_Properties Map
@@ -118,15 +117,15 @@ The **Material_Instances** map breaks down how instances are generated, how **Pa
 ![](MaterialsProperties_Header.png)(w:640)
 
 
-The **Material_Properties** map delves deeper into the properties that make up a **Material**. By adjusting these properties, you can affect how the **Material** appears when it is applied to an object as well as how it reacts to external components such as lights or Post Processing Volumes. 
+The **Material_Properties** map delves deeper into the properties that make up a Material. By adjusting these properties, you can affect how the Material appears when it is applied to an object as well as how it reacts to external components such as lights or Post Processing Volumes. 
 
 | Example | What is Demonstrated |
 | ------- | ------------ |
-| **[1.1 BlendMode](Resources\ContentExamples\MaterialProperties\1_1)** | Each of the effects that are generated when using the different **BlendModes** within a **Material**. |
-| **[1.2 LightingModel](Resources\ContentExamples\MaterialProperties\1_2)** | The different **Shading Models** that can be used to affect the way the **Material** reflects light.  |
+| **[1.1 BlendMode](Resources\ContentExamples\MaterialProperties\1_1)** | Each of the effects that are generated when using the different **BlendModes** within a Material. |
+| **[1.2 LightingModel](Resources\ContentExamples\MaterialProperties\1_2)** | The different **Shading Models** that can be used to affect the way the Material reflects light.  |
 | **[1.3 Translucent Lighting Mode](Resources\ContentExamples\MaterialProperties\1_3)** | When a **Translucent Material** is set to **Default Lit**, the additional options of **Volumetric Non Directional**, **Volumetric Directional**, and **Surface** become available. |
-| **[1.4 Two Sided](Resources\ContentExamples\MaterialProperties\1_4)** | Shows how a **Material** is applied to surfaces in both one sided and two sided methods.  |
-| **[1.5 Material Domain](Resources\ContentExamples\MaterialProperties\1_5)** | How **Material Domains** define the overall usage of a **Material** (e.g. **Surface**, **Deferred Decal**, **Light Function**, and **Post Process Blendable** examples are given). |
+| **[1.4 Two Sided](Resources\ContentExamples\MaterialProperties\1_4)** | Shows how a Material is applied to surfaces in both one sided and two sided methods.  |
+| **[1.5 Material Domain](Resources\ContentExamples\MaterialProperties\1_5)** | How **Material Domains** define the overall usage of a Material (e.g. **Surface**, **Deferred Decal**, **Light Function**, and **Post Process Blendable** examples are given). |
 | **[1.6 Separate Translucency](Resources\ContentExamples\MaterialProperties\1_6)** | How **Separate Translucency** shaders are not affected by **Depth of Field** (e.g. using a **Depth of Field** effect to blur out objects that have **Separate Translucency** disabled). |
 | **[1.7 Volumetric Directional Lighting Intensity](Resources\ContentExamples\MaterialProperties\1_7)** | Shows how the intensity of shading can be set if the **Translucent Material** is set to **Volumetric Directional Light** mode. |
 | **[1.8 Tesselation](Resources\ContentExamples\MaterialProperties\1_8)** | **Tesselation** can be used to improve surface details of a mesh.  |
@@ -140,14 +139,14 @@ The **Material_Properties** map delves deeper into the properties that make up a
 
 ![](Decals_Header.png)(w:640)
 
-This map presents several ways in which you can use special **Materials** called **Decals** to add detail or decorations to your levels. **Decals** can be used to create stained walls, project a neon looking glow or even be animated. This map contains several examples of how **Decals** can be used and how to control their appearance on the surface they are applied to.  
+This map presents several ways in which you can use special Materials called **Decals** to add detail or decorations to your levels. **Decals** can be used to create stained walls, project a neon looking glow or even be animated. This map contains several examples of how Decals can be used and how to control their appearance on the surface they are applied to.  
 
 | Example | What is Demonstrated |
 | ------- | ------------ |
-| **[1.1 Basic Decal](Resources\ContentExamples\Decals\1_1)** | A **Decal** is a **Material** that is projected onto a Static or Skeletal Mesh. |
-| **[1.2 Decal Blend Modes](Resources\ContentExamples\Decals\1_2)** | Multiple **Decals** are shown using different **Blend Modes** which affects how they appear on the surface they are placed (e.g. Translucent, Stain, Normal Map, and Emissive **Decals**). |
-| **[1.3 Masked Decal](Resources\ContentExamples\Decals\1_3)** | A **Decal** with both an **Opacity** and **Normal Map** texture is used to create a more solid looking effect rather than a painted on effect.  |
+| **[1.1 Basic Decal](Resources\ContentExamples\Decals\1_1)** | A **Decal** is a Material that is projected onto a Static or Skeletal Mesh. |
+| **[1.2 Decal Blend Modes](Resources\ContentExamples\Decals\1_2)** | Multiple Decals are shown using different Blend Modes which affects how they appear on the surface they are placed (e.g. Translucent, Stain, Normal Map, and Emissive Decals). |
+| **[1.3 Masked Decal](Resources\ContentExamples\Decals\1_3)** | A Decal with both an Opacity and Normal Map texture is used to create a more solid looking effect rather than a painted on effect.  |
 | **[1.4 Projection on multiple meshes](Resources\ContentExamples\Decals\1_4)** | How **Decals** are projected onto the meshes that are inside its volume. |
-| **[1.5 Recieves Decals off](Resources\ContentExamples\Decals\1_5)** | Illustrates how to ignore **Decals** by turning off the **Recieves Decals** option under **Rendering** in the **Details** panel of the **Actor**. |
-| **[1.6 Sort Order](Resources\ContentExamples\Decals\1_6)** | **Sort Order** is used to control which order **Decals** appear in when placed on top of each other.  |
-| **[1.7 Animated Decal Material](Resources\ContentExamples\Decals\1_7)** | Demonstrates how to animate a **Decal** (e.g. a spinning **Decal** is placed on a brick wall. |
+| **[1.5 Recieves Decals off](Resources\ContentExamples\Decals\1_5)** | Illustrates how to ignore Decals by turning off the **Recieves Decals** option under **Rendering** in the **Details** panel of the **Actor**. |
+| **[1.6 Sort Order](Resources\ContentExamples\Decals\1_6)** | **Sort Order** is used to control which order Decals appear in when placed on top of each other.  |
+| **[1.7 Animated Decal Material](Resources\ContentExamples\Decals\1_7)** | Demonstrates how to animate a Decal (e.g. a spinning Decal is placed on a brick wall. |
