TypoCorrection

diff --git a/Engine/Documentation/Source/Platforms/Android/ReducingAPKSize/ReducingAPHSizes.INT.udn b/Engine/Documentation/Source/Platforms/Android/ReducingAPKSize/ReducingAPHSizes.INT.udn
--- a/Engine/Documentation/Source/Platforms/Android/ReducingAPKSize/ReducingAPHSizes.INT.udn
+++ b/Engine/Documentation/Source/Platforms/Android/ReducingAPKSize/ReducingAPHSizes.INT.udn
@@ -1,4 +1,4 @@
-Availability:Docs
+﻿Availability:Docs
 Title: Reducing APK Packages size 
 Crumbs:%ROOT%, Platforms, Platforms/Android
 Description:How to reduce the size of your packaged Android game.
@@ -47,7 +47,7 @@ In fact it is not uncommon for some projects APK packages to be reduced in size
 One overlooked area that can lead to a bloating of APK package sizes is failing to properly setup your projects options like which levels should be loaded by default or which levels should be used between level transitions. 
 To set which level or levels should be used for this type of interaction you can do the following.
 
-1. First on the **Main Toolbar** got to  **Edit** > **Project Settings**.
+1. First on the **Main Toolbar** go to  **Edit** > **Project Settings**.
 
 	 ![](T_Project_Settings.png)
 
