#UE4doc #docs:
Added Engine Version field

diff --git a/Engine/Documentation/Source/Programming/Tutorials/AutoCamera/AutoCamera.INT.udn b/Engine/Documentation/Source/Programming/Tutorials/AutoCamera/AutoCamera.INT.udn
--- a/Engine/Documentation/Source/Programming/Tutorials/AutoCamera/AutoCamera.INT.udn
+++ b/Engine/Documentation/Source/Programming/Tutorials/AutoCamera/AutoCamera.INT.udn
@@ -2,6 +2,7 @@ Availability: Public
 Title:Game-Controlled Cameras
 Crumbs:%ROOT%, GettingStarted, Programming/Tutorials
 Description:Learn to activate and switch between different view perspectives.
+version:4.9
 
 [VAR:Topic]
 [OBJECT:Topic]
diff --git a/Engine/Documentation/Source/Programming/Tutorials/Components/Components.INT.udn b/Engine/Documentation/Source/Programming/Tutorials/Components/Components.INT.udn
--- a/Engine/Documentation/Source/Programming/Tutorials/Components/Components.INT.udn
+++ b/Engine/Documentation/Source/Programming/Tutorials/Components/Components.INT.udn
@@ -2,6 +2,7 @@ Availability: Public
 Title:Components And Collision
 Crumbs:%ROOT%, GettingStarted, Programming/Tutorials
 Description:Learn how Components can enable Pawns to interact with physics, use particle effects, and more.
+version:4.9
 
 [VAR:Topic]
 [OBJECT:Topic]
diff --git a/Engine/Documentation/Source/Programming/Tutorials/PlayerCamera/PlayerCamera.INT.udn b/Engine/Documentation/Source/Programming/Tutorials/PlayerCamera/PlayerCamera.INT.udn
--- a/Engine/Documentation/Source/Programming/Tutorials/PlayerCamera/PlayerCamera.INT.udn
+++ b/Engine/Documentation/Source/Programming/Tutorials/PlayerCamera/PlayerCamera.INT.udn
@@ -2,6 +2,7 @@ Availability: Public
 Title:Player-Controlled Cameras
 Crumbs:%ROOT%, GettingStarted, Programming/Tutorials
 Description:Learn to manipulate a Camera and a Pawn at the same time, using player input.
+version:4.9
 
 [VAR:Topic]
 [OBJECT:Topic]
diff --git a/Engine/Documentation/Source/Programming/Tutorials/PlayerInput/PlayerInput.INT.udn b/Engine/Documentation/Source/Programming/Tutorials/PlayerInput/PlayerInput.INT.udn
--- a/Engine/Documentation/Source/Programming/Tutorials/PlayerInput/PlayerInput.INT.udn
+++ b/Engine/Documentation/Source/Programming/Tutorials/PlayerInput/PlayerInput.INT.udn
@@ -2,6 +2,7 @@ Availability: Public
 Title:Player Input and Pawns
 Crumbs:%ROOT%, GettingStarted
 Description:React to player input by extending the Pawn class.
+version:4.9
 
 [VAR:Topic]
 [OBJECT:Topic]
