#UE4 Docs: Adding a Youtube link to show disabling plugins.

diff --git a/Engine/Documentation/Source/Platforms/Android/Android.INT.udn b/Engine/Documentation/Source/Platforms/Android/Android.INT.udn
--- a/Engine/Documentation/Source/Platforms/Android/Android.INT.udn
+++ b/Engine/Documentation/Source/Platforms/Android/Android.INT.udn
@@ -19,5 +19,6 @@ Navigation:platform
 		* [](Platforms/Android/GettingStarted "%Platforms/Android/GettingStarted:description%")
 		* [](Platforms/Android/Reference "%Platforms/Android/Reference:description%")
 		* [](Platforms/Android/DeviceCompatibility "%Platforms/Android/DeviceCompatibility:description%")
+		* [](Platforms/Android/ReducingAPKSize "%Platforms/Android/ReducingAPKSize:description%")
 	[/PARAM]
 [/OBJECT]
diff --git a/Engine/Documentation/Source/Platforms/Android/ReducingAPKSize/ReducingAPKSizes.INT.udn b/Engine/Documentation/Source/Platforms/Android/ReducingAPKSize/ReducingAPKSizes.INT.udn
--- a/Engine/Documentation/Source/Platforms/Android/ReducingAPKSize/ReducingAPKSizes.INT.udn
+++ b/Engine/Documentation/Source/Platforms/Android/ReducingAPKSize/ReducingAPKSizes.INT.udn
@@ -1,5 +1,5 @@
 ﻿Availability:Docs
-Title: Reducing APK Packages size 
+Title: Reducing APK Package size 
 Crumbs:%ROOT%, Platforms, Platforms/Android
 Description:How to reduce the size of your packaged Android game.
 Skilllevel: Intermediate
@@ -98,7 +98,17 @@ To disable a plug-in inside of UE4 you need to do the following.
 	![](T_Open_Plugins_Browser.png)
 
 1. From the Plug-in Browser select the different sections and disable any plug-in's you are not using by un-checking the check mark box next to **Enabled**.
-  When you have disabled all the Plug-ins you do not need make sure to re-package your project to see the results.
-
-	
+  When you have disabled all the Plug-ins you will need to re-start the UE4 Editor and then make sure to re-package your project.
+
+	[OBJECT:EmbeddedVideo]
+	[PARAMLITERAL:width]
+	640
+	[/PARAMLITERAL]
+	[PARAMLITERAL:height]
+	360
+	[/PARAMLITERAL]
+	[PARAMLITERAL:videoid]
+	xLJ01QEvV7Y
+	[/PARAMLITERAL]
+	[/OBJECT]
 
