#UE4doc #docs:
Added Engine Version field
Changed Home link to point to the overview for the tutorial
Change font settings for multiple terms
Replaced picture due to it being based off a project with a mismatching name

diff --git a/Engine/Documentation/Source/Programming/Tutorials/PlayerInput/1/PlayerInput_1.INT.udn b/Engine/Documentation/Source/Programming/Tutorials/PlayerInput/1/PlayerInput_1.INT.udn
--- a/Engine/Documentation/Source/Programming/Tutorials/PlayerInput/1/PlayerInput_1.INT.udn
+++ b/Engine/Documentation/Source/Programming/Tutorials/PlayerInput/1/PlayerInput_1.INT.udn
@@ -2,6 +2,7 @@ Availability: Public
 Title:1. Customize A Pawn
 Crumbs:%ROOT%, Programming, Programming/Tutorials, Programming/Tutorials/PlayerInput
 Description:Introduction to Pawns
+version:4.9
 
 [VAR:Steps]
 [OBJECT:Navigation]
@@ -10,7 +11,7 @@ Description:Introduction to Pawns
 	[PARAM:current]
 	[/PARAM]
 	[PARAM:home]
-		[Programming Tutorials Home](Programming\Tutorials\)
+		[](Programming\Tutorials\PlayerInput)
 	[/PARAM]
 	[PARAM:next]
 		[Next Step](Programming\Tutorials\PlayerInput\2 "%Programming\Tutorials\PlayerInput\2:title%")
@@ -22,27 +23,27 @@ Description:Introduction to Pawns
 
 [REGION:tip]If you are new to **Unreal Engine 4**, you might want to read our [%Programming\QuickStart:title% tutorial](Programming\QuickStart) first. For this tutorial, we will assume you are familiar with creating a project, adding C++ code to it, and compiling your code.[/REGION]
 
-1. We will begin by creating a new, Basic Code project, with starter content, named "HowTo_PlayerInput". Then, we'll add a customized **Pawn** class, which we will call "MyPawn", to the project.
+1. We will begin by creating a new, Basic Code project, with starter content, named HowTo_PlayerInput. Then, we'll add a customized **Pawn** class, which we will call MyPawn, to the project.
 
-	[REGION:note]A **Pawn** is a type of **Actor** that is designed to be controlled by human players or AI.[/REGION]
+	[REGION:note]A Pawn is a type of **Actor** that is designed to be controlled by human players or AI.[/REGION]
 
 	[![](AddPawn.png)](AddPawn.png)
 
 	[![](NameYourNewPawn.png)](NameYourNewPawn.png)
 
-1. The first thing we're going to do is set our MyPawn to respond to player input automatically upon the game starting. The **Pawn** class provides a variable we can set during initialization that handles this for us. In `MyPawn.cpp`, add the following code to **AMyPawn::AMyPawn**:
+1. The first thing we're going to do is set our MyPawn to respond to player input automatically upon the game starting. The Pawn class provides a variable we can set during initialization that handles this for us. In MyPawn.cpp, add the following code to **AMyPawn::AMyPawn**:
 
 		// Set this pawn to be controlled by the lowest-numbered player
 		AutoPossessPlayer = EAutoReceiveInput::Player0;
 
-1. Next, we'll build a few basic **Components**. If you want to learn more about adding and managing **Components** in code, as well as some common types of **Components** you will encounter, try reading our [%Programming\Tutorials\Components:title% tutorial](Programming\Tutorials\Components). In order to keep track of the **Component** we will create, we should add the following code to `MyPawn.h`, at the bottom of our class definition:
+1. Next, we'll build a few basic **Components**. If you want to learn more about adding and managing Components in code, as well as some common types of Components you will encounter, try reading our [%Programming\Tutorials\Components:title% tutorial](Programming\Tutorials\Components). In order to keep track of the Component we will create, we should add the following code to MyPawn.h, at the bottom of our class definition:
 
 		UPROPERTY(EditAnywhere)
 		USceneComponent* OurVisibleComponent;
 
 	[REGION:note]This variable is tagged as a **UPROPERTY** so that it will be visible to **Unreal Engine**. This is important because it prevents the variable from being reset when the game is launched, or when the project or level is closed and reloaded.[/REGION]
 
-	And back in `MyPawn.cpp`, we should add the following code to **AMyPawn::AMyPawn**:
+	And back in MyPawn.cpp, we should add the following code to AMyPawn::AMyPawn:
 
 		// Create a dummy root component we can attach things to.
 		RootComponent = CreateDefaultSubobject<USceneComponent>(TEXT("RootComponent"));
@@ -55,11 +56,11 @@ Description:Introduction to Pawns
 		OurCamera->SetRelativeRotation(FRotator(-45.0f, 0.0f, 0.0f));
 		OurVisibleComponent->AttachTo(RootComponent);
 
-1. We are now ready to save our changes and compile with the **Build** command in **Visual Studio** or the **Compile** button in the **Unreal Engine** editor.
+1. We are now ready to save our changes and compile with the **Build** command in **Visual Studio** or the **Compile** button in the **Unreal Edtior**.
 
 ----
 
-Now that we have a customized **Pawn** to react to our game's input, we'll need to define what that input will be. To do this, we'll configure our project's **Input Settings** in the **Unreal Engine** editor.
+Now that we have a customized Pawn to react to our game's input, we'll need to define what that input will be. To do this, we'll configure our project's **Input Settings** in the **Unreal Editor**.
 
 [OBJECT:Section]
     [PARAMLITERAL:id]
