#UE4doc #docs: ##Overview removed as according to the instructional guide.

diff --git a/Engine/Documentation/Source/Resources/Showcases/PostProcessing/GlobalIllumination/GlobalIllumination.INT.udn b/Engine/Documentation/Source/Resources/Showcases/PostProcessing/GlobalIllumination/GlobalIllumination.INT.udn
--- a/Engine/Documentation/Source/Resources/Showcases/PostProcessing/GlobalIllumination/GlobalIllumination.INT.udn
+++ b/Engine/Documentation/Source/Resources/Showcases/PostProcessing/GlobalIllumination/GlobalIllumination.INT.udn
@@ -4,9 +4,6 @@ Crumbs:%ROOT%, Resources, Resources/Showcases/PostProcessing
 Description:A rundown of how Global Illumination was employed within the Post Process Matinee example included with UE4.
 
 
-
-### Global Illumination
-
 The Global Illumination settings allow you to enhance and modify the light bounced throughout the scene via Lightmass. This gives you the power to tweak bounced light without having to adjust lights or rebuild.
 
 ![](GISettings.png)
