#UE4doc #docs: ##Overview removed as according to the instructional guide.

diff --git a/Engine/Documentation/Source/Resources/Showcases/PostProcessing/DepthOfField/DepthOfField.INT.udn b/Engine/Documentation/Source/Resources/Showcases/PostProcessing/DepthOfField/DepthOfField.INT.udn
--- a/Engine/Documentation/Source/Resources/Showcases/PostProcessing/DepthOfField/DepthOfField.INT.udn
+++ b/Engine/Documentation/Source/Resources/Showcases/PostProcessing/DepthOfField/DepthOfField.INT.udn
@@ -3,10 +3,6 @@ Title:Depth of Field
 Crumbs:%ROOT%, Resources, Resources/Showcases/PostProcessing
 Description:A rundown of how Depth of Field was employed within the Post Process Matinee example included with UE4.
 
-
-
-### Depth of Field
-
 Depth of field (DOF) allows you to simulate the blurring that takes place in the foreground and background of an image due to camera aperture size. In this instance, virtually all of the settings have been adjusted to view the final effect, and since those properties all work in close relation, we will only be demonstrating the results with and without DOF.
 
 ![](DOFSettings.png)
