#UE4doc #docs:
Added Engine Version field
Changed Home link to point to the overview for the tutorial
Changed font settings for multiple terms
Replaced a picture to change the indicator borders and arrows from green to yellow
Removed misplaced note region

diff --git a/Engine/Documentation/Source/Programming/Tutorials/PlayerInput/2/PlayerInput_2.INT.udn b/Engine/Documentation/Source/Programming/Tutorials/PlayerInput/2/PlayerInput_2.INT.udn
--- a/Engine/Documentation/Source/Programming/Tutorials/PlayerInput/2/PlayerInput_2.INT.udn
+++ b/Engine/Documentation/Source/Programming/Tutorials/PlayerInput/2/PlayerInput_2.INT.udn
@@ -2,6 +2,7 @@ Availability: Public
 Title:2. Configure Game Input
 Crumbs:%ROOT%, Programming, Programming/Tutorials, Programming/Tutorials/PlayerInput
 Description:Configuring input in Unreal Engine
+version:4.9
 
 [VAR:Steps]
 [OBJECT:Navigation]
@@ -11,7 +12,7 @@ Description:Configuring input in Unreal Engine
 	[PARAM:current]
 	[/PARAM]
 	[PARAM:home]
-		[Programming Tutorials Home](Programming\Tutorials\)
+		[](Programming\Tutorials\PlayerInput)
 	[/PARAM]
 	[PARAM:next]
 		[Next Step](Programming\Tutorials\PlayerInput\3 "%Programming\Tutorials\PlayerInput\3:title%")
@@ -21,10 +22,10 @@ Description:Configuring input in Unreal Engine
 
 %Steps%
 
-[REGION:note]There are two types of input mappings: Action and Axis.
+There are two types of input mappings: Action and Axis.
 <br>**Action Mappings** are useful to think of as "yes or no" inputs, like the buttons on a mouse or joystick. They report when they are pressed, released, double-clicked, or held down for a short time. Discrete actions like jumping, shooting, or interacting with objects are good candidates for this type of mapping.
 <br>**Axis Mappings** are continuous - think of them as "how much" inputs, like the stick on a joystick, or the position of a mouse cursor. They report their value every frame, even if they're not moving. Things with magnitude or direction, such as walking, looking around, and steering a vehicle are usually handled this way.
-<br>While input mappings can be defined directly in code, the usual method is to define them in the **Unreal Engine** editor, so that's how we'll do it in this tutorial.[/REGION]
+<br>While input mappings can be defined directly in code, the usual method is to define them in the **Unreal Engine** editor, so that's how we'll do it in this tutorial.
 
 1. In the **Unreal Engine** editor, under the **Edit** dropdown menu, we'll click the **Project Settings** option.
 
@@ -33,7 +34,7 @@ Description:Configuring input in Unreal Engine
 1. From there, we'll select the **Input** option from the **Engine** section on the left. We can then expand the **Bindings** category that appears on the right, and add one **Action Mapping** and two **Axis Mappings**.
 	
 	The plus sign next to the **Action Mapping** or **Axis Mapping** section headings will add a new mapping. The expander arrow on the left can be used to show or hide our mappings.
-	To add an additional input to a mapping, click the plus sign next to that mapping. Following are the mappings and inputs we will need. Take note of the negative values for the "S" and A" inputs.
+	To add an additional input to a mapping, click the plus sign next to that mapping. Following are the mappings and inputs we will need. Take note of the negative values for the S and A inputs.
 
 	| ----- | ----- | ----- |
 	| Action Mapping | | |
@@ -48,13 +49,13 @@ Description:Configuring input in Unreal Engine
 
 	[![](ConfigureInput.png)](ConfigureInput.png)
 
-1. Now that our input is configured, let's set up a MyPawn in our level. The MyPawn class will appear in our **Content Browser**, and is ready to be dragged into the **Level Editor**.
+1. Now that our input is configured, let's set up a MyPawn in our level. The** MyPawn** class will appear in our **Content Browser**, and is ready to be dragged into the **Level Editor**.
 
 	[![](ClassInContentBrowser.png)](ClassInContentBrowser.png)
 
 	[![](PawnInSceneView.png)](PawnInSceneView.png)
 
-1. One more step is needed to set up our **MyPawn**. We'll need to give it a **Static Mesh** so that we can see it in the game. We can do this by selecting the **MyPawn** we just created, selecting the component called "OurVisibleComponent (Inherited)" in the **Details Panel**, and assigning an asset to it through the dropdown box in the **Static Mesh** category. For this tutorial, "Shape_Cylinder" is a good asset to use.
+1. One more step is needed to set up our MyPawn. We'll need to give it a **Static Mesh** so that we can see it in the game. We can do this by selecting the MyPawn we just created, selecting the component called **OurVisibleComponent (Inherited)** in the **Details Panel**, and assigning an asset to it through the dropdown box in the **Static Mesh** category. For this tutorial, **Shape_Cylinder** is a good asset to use.
 
 	[![](StaticMesh.png)](StaticMesh.png)
 
