#UE4doc #docs: ##Overview removed as according to the instructional guide.

diff --git a/Engine/Documentation/Source/Resources/Showcases/PostProcessing/Bloom/Bloom.INT.udn b/Engine/Documentation/Source/Resources/Showcases/PostProcessing/Bloom/Bloom.INT.udn
--- a/Engine/Documentation/Source/Resources/Showcases/PostProcessing/Bloom/Bloom.INT.udn
+++ b/Engine/Documentation/Source/Resources/Showcases/PostProcessing/Bloom/Bloom.INT.udn
@@ -3,10 +3,6 @@ Title:Bloom
 Crumbs:%ROOT%, Resources, Resources/Showcases/PostProcessing
 Description:A rundown of how Bloom was employed within the Post Process Matinee example included with UE4.
 
-
-
-### Bloom
-
 [INCLUDE:Engine/Rendering/LightingAndShadows/LightsAndPostProcess#LightingAndPostProcess_Bloom] For this portion of the scene, bloom was employed in the following manner:
 
 ![](BloomSettings.png)
