#UE4doc #docs:
Added Engine Version field
Changed Home link to point to the overview for the tutorial
Changed font settings for multiple terms

diff --git a/Engine/Documentation/Source/Programming/Tutorials/PlayerInput/3/PlayerInput_3.INT.udn b/Engine/Documentation/Source/Programming/Tutorials/PlayerInput/3/PlayerInput_3.INT.udn
--- a/Engine/Documentation/Source/Programming/Tutorials/PlayerInput/3/PlayerInput_3.INT.udn
+++ b/Engine/Documentation/Source/Programming/Tutorials/PlayerInput/3/PlayerInput_3.INT.udn
@@ -2,6 +2,7 @@ Availability: Public
 Title:3. Program and Bind Game Actions
 Crumbs:%ROOT%, Programming, Programming/Tutorials, Programming/Tutorials/PlayerInput
 Description:Binding input to C++ code
+version:4.9
 
 [VAR:Steps]
 [OBJECT:Navigation]
@@ -11,7 +12,7 @@ Description:Binding input to C++ code
 	[PARAM:current]
 	[/PARAM]
 	[PARAM:home]
-		[Programming Tutorials Home](Programming\Tutorials\)
+		[](Programming\Tutorials\PlayerInput)
 	[/PARAM]
 	[PARAM:next]
 		[Next Step](Programming\Tutorials\PlayerInput\4 "%Programming\Tutorials\PlayerInput\4:title%")
@@ -21,7 +22,7 @@ Description:Binding input to C++ code
 
 %Steps%
 
-1. In **Visual Studio**, open `MyPawn.h` and add the following code to the bottom of MyPawn's class definition:
+1. In **Visual Studio**, open MyPawn.h and add the following code to the bottom of MyPawn's class definition:
 
 		//Input functions
 		void Move_XAxis(float AxisValue);
@@ -35,7 +36,7 @@ Description:Binding input to C++ code
 
 	[REGION:note]The four input functions are going to be bound to our input events. When they run, they will update the values stored in our new input variables, which MyPawn will use to determine what it should do during in the game.[/REGION]
 
-1. We'll switch over to `MyPawn.cpp` and program the four functions we just declared. Add the following code:
+1. We'll switch over to MyPawn.cpp and program the four functions we just declared. Add the following code:
 
 		void AMyPawn::Move_XAxis(float AxisValue)
 		{
@@ -100,7 +101,7 @@ Description:Binding input to C++ code
 			}
 		}
 
-1. After compling our code, we can return to the **Unreal Editor** and press **Play**. We should have control of our **Pawn** with the WASD keys, and we should be able to make it grow by holding Space Bar, and watch it shrink when we let go.
+1. After compling our code, we can return to the **Unreal Editor** and press **Play**. We should have control of our **Pawn** with the WASD keys, and we should be able to make it grow by holding **Space Bar**, and watch it shrink when we let go.
 
 [![](PlayingGame1.png)](PlayingGame1.png)
 [![](PlayingGame2.png)](PlayingGame2.png)
