#UE4doc #docs:
Added Engine Version field
Changed Home link to point to the overview for the tutorial

diff --git a/Engine/Documentation/Source/Programming/Tutorials/PlayerInput/4/PlayerInput_4.INT.udn b/Engine/Documentation/Source/Programming/Tutorials/PlayerInput/4/PlayerInput_4.INT.udn
--- a/Engine/Documentation/Source/Programming/Tutorials/PlayerInput/4/PlayerInput_4.INT.udn
+++ b/Engine/Documentation/Source/Programming/Tutorials/PlayerInput/4/PlayerInput_4.INT.udn
@@ -2,7 +2,7 @@ Availability: Public
 Title:4. On Your Own!
 Crumbs:%ROOT%, Programming, Programming/Tutorials, Programming/Tutorials/PlayerInput
 Description:Ideas for further exploration!
-
+version:4.9
 
 [VAR:Steps]
 [OBJECT:Navigation]
@@ -12,7 +12,7 @@ Description:Ideas for further exploration!
 	[PARAM:current]
 	[/PARAM]
 	[PARAM:home]
-		[Programming Tutorials Home](Programming\Tutorials\)
+		[](Programming\Tutorials\PlayerInput)
 	[/PARAM]
 	[PARAM:next]
 	[/PARAM]
