#UE4doc #docs: ##Overview removed as according to the instructional guide.

diff --git a/Engine/Documentation/Source/Resources/Showcases/MaterialLayers/MaterialLayers.INT.udn b/Engine/Documentation/Source/Resources/Showcases/MaterialLayers/MaterialLayers.INT.udn
--- a/Engine/Documentation/Source/Resources/Showcases/MaterialLayers/MaterialLayers.INT.udn
+++ b/Engine/Documentation/Source/Resources/Showcases/MaterialLayers/MaterialLayers.INT.udn
@@ -51,8 +51,6 @@ Description:An overview of the Material Layers example included with UE4.
 
 [TOC(start:2)]
 
-## Overview
-
 This example focuses on the concept of Material Layers. This is the technique of layering together a series of Material Functions, allowing you to create a single blended Material that is comprised of various others. Best of all, the separation of each of these layers can be done via a texture (per-pixel), rather than on a polygon-by-polygon basis, as would have to be done if applying multiple Materials to a single object. 
 
 The example consists of six rocket Static Meshes ready for launch on a wide platform. Each of the six rockets showcases a different Material or Material Function. The first four of these define different real-world Materials, such as plastic and various metals. The last two on the right are actually blended Material Functions in which each of the four base Materials have been blended together to give a final result.
