#UE4doc #docs:
Fixed captitalization errors in titles

diff --git a/Engine/Documentation/Source/Programming/Tutorials/AutoCamera/1/AutoCamera_1.INT.udn b/Engine/Documentation/Source/Programming/Tutorials/AutoCamera/1/AutoCamera_1.INT.udn
--- a/Engine/Documentation/Source/Programming/Tutorials/AutoCamera/1/AutoCamera_1.INT.udn
+++ b/Engine/Documentation/Source/Programming/Tutorials/AutoCamera/1/AutoCamera_1.INT.udn
@@ -1,5 +1,5 @@
 Availability: Public
-Title:1. Place Cameras In The World
+Title:1. Place Cameras in the World
 Crumbs:%ROOT%, Programming, Programming/Tutorials, Programming/Tutorials/AutoCamera
 Description:Setting up a world with multiple cameras
 version: 4.9
diff --git a/Engine/Documentation/Source/Programming/Tutorials/AutoCamera/2/AutoCamera_2.INT.udn b/Engine/Documentation/Source/Programming/Tutorials/AutoCamera/2/AutoCamera_2.INT.udn
--- a/Engine/Documentation/Source/Programming/Tutorials/AutoCamera/2/AutoCamera_2.INT.udn
+++ b/Engine/Documentation/Source/Programming/Tutorials/AutoCamera/2/AutoCamera_2.INT.udn
@@ -1,5 +1,5 @@
 Availability: Public
-Title:2. Control Camera View In C++
+Title:2. Control Camera View in C++
 Crumbs:%ROOT%, Programming, Programming/Tutorials, Programming/Tutorials/AutoCamera
 Description:Coding a class to change camera views
 version: 4.9
diff --git a/Engine/Documentation/Source/Programming/Tutorials/AutoCamera/3/AutoCamera_3.INT.udn b/Engine/Documentation/Source/Programming/Tutorials/AutoCamera/3/AutoCamera_3.INT.udn
--- a/Engine/Documentation/Source/Programming/Tutorials/AutoCamera/3/AutoCamera_3.INT.udn
+++ b/Engine/Documentation/Source/Programming/Tutorials/AutoCamera/3/AutoCamera_3.INT.udn
@@ -1,5 +1,5 @@
 Availability: Public
-Title:3. Place A Camera Director In The World
+Title:3. Place a Camera Director in the World
 Crumbs:%ROOT%, Programming, Programming/Tutorials, Programming/Tutorials/AutoCamera
 Description:Placing a Camera Director in the world and assigning Cameras to it
 version:4.9
diff --git a/Engine/Documentation/Source/Programming/Tutorials/Components/1/Components_1.INT.udn b/Engine/Documentation/Source/Programming/Tutorials/Components/1/Components_1.INT.udn
--- a/Engine/Documentation/Source/Programming/Tutorials/Components/1/Components_1.INT.udn
+++ b/Engine/Documentation/Source/Programming/Tutorials/Components/1/Components_1.INT.udn
@@ -1,5 +1,5 @@
 Availability: Public
-Title:1. Creating And Attaching Components
+Title:1. Creating and Attaching Components
 Crumbs:%ROOT%, Programming, Programming/Tutorials, Programming/Tutorials/Components
 Description:Creating a variety of components and attaching them to a customized pawn.
 version:4.9
diff --git a/Engine/Documentation/Source/Programming/Tutorials/Components/2/Components_2.INT.udn b/Engine/Documentation/Source/Programming/Tutorials/Components/2/Components_2.INT.udn
--- a/Engine/Documentation/Source/Programming/Tutorials/Components/2/Components_2.INT.udn
+++ b/Engine/Documentation/Source/Programming/Tutorials/Components/2/Components_2.INT.udn
@@ -1,5 +1,5 @@
 Availability: Public
-Title:2. Configuring input and creating a Pawn Movement Component
+Title:2. Configuring Input and Creating a Pawn Movement Component
 Crumbs:%ROOT%, Programming, Programming/Tutorials, Programming/Tutorials/Components
 Description:Setting up input and building a custom Pawn Movement Component to define how our Pawn will move.
 version:4.9
diff --git a/Engine/Documentation/Source/Programming/Tutorials/Components/3/Components_3.INT.udn b/Engine/Documentation/Source/Programming/Tutorials/Components/3/Components_3.INT.udn
--- a/Engine/Documentation/Source/Programming/Tutorials/Components/3/Components_3.INT.udn
+++ b/Engine/Documentation/Source/Programming/Tutorials/Components/3/Components_3.INT.udn
@@ -1,5 +1,5 @@
 Availability: Public
-Title:3. Coding our Pawn Movement Component's behavior
+Title:3. Coding Our Pawn Movement Component's Behavior
 Crumbs:%ROOT%, Programming, Programming/Tutorials, Programming/Tutorials/Components
 Description:Writing the code that will define how our Pawn Movement Component will determine where the Pawn it controls will go each frame.
 version:4.9
diff --git a/Engine/Documentation/Source/Programming/Tutorials/Components/4/Components_4.INT.udn b/Engine/Documentation/Source/Programming/Tutorials/Components/4/Components_4.INT.udn
--- a/Engine/Documentation/Source/Programming/Tutorials/Components/4/Components_4.INT.udn
+++ b/Engine/Documentation/Source/Programming/Tutorials/Components/4/Components_4.INT.udn
@@ -1,5 +1,5 @@
 Availability: Public
-Title:4. Using our Pawn and Components together
+Title:4. Using Our Pawn and Components Together
 Crumbs:%ROOT%, Programming, Programming/Tutorials, Programming/Tutorials/Components
 Description:Tying the functionality of the Pawn, the input configuration, and the Components together.
 version:4.9
diff --git a/Engine/Documentation/Source/Programming/Tutorials/Components/5/Components_5.INT.udn b/Engine/Documentation/Source/Programming/Tutorials/Components/5/Components_5.INT.udn
--- a/Engine/Documentation/Source/Programming/Tutorials/Components/5/Components_5.INT.udn
+++ b/Engine/Documentation/Source/Programming/Tutorials/Components/5/Components_5.INT.udn
@@ -1,5 +1,5 @@
 Availability: Public
-Title:5. Playing In Editor
+Title:5. Playing in Editor
 Crumbs:%ROOT%, Programming, Programming/Tutorials, Programming/Tutorials/Components
 Description:Testing our code by playing in the editor.
 version:4.9
diff --git a/Engine/Documentation/Source/Programming/Tutorials/Components/Components.INT.udn b/Engine/Documentation/Source/Programming/Tutorials/Components/Components.INT.udn
--- a/Engine/Documentation/Source/Programming/Tutorials/Components/Components.INT.udn
+++ b/Engine/Documentation/Source/Programming/Tutorials/Components/Components.INT.udn
@@ -1,5 +1,5 @@
 Availability: Public
-Title:Components And Collision
+Title:Components and Collision
 Crumbs:%ROOT%, GettingStarted, Programming/Tutorials
 Description:Learn how Components can enable Pawns to interact with physics, use particle effects, and more.
 version:4.9
diff --git a/Engine/Documentation/Source/Programming/Tutorials/PlayerCamera/1/PlayerCamera_1.INT.udn b/Engine/Documentation/Source/Programming/Tutorials/PlayerCamera/1/PlayerCamera_1.INT.udn
--- a/Engine/Documentation/Source/Programming/Tutorials/PlayerCamera/1/PlayerCamera_1.INT.udn
+++ b/Engine/Documentation/Source/Programming/Tutorials/PlayerCamera/1/PlayerCamera_1.INT.udn
@@ -1,5 +1,5 @@
 Availability: Public
-Title:1. Attach A Camera To A Pawn
+Title:1. Attach a Camera to a Pawn
 Crumbs:%ROOT%, Programming, Programming/Tutorials, Programming/Tutorials/PlayerCamera
 Description:Creating a spring arm and a camera.
 version:4.9
diff --git a/Engine/Documentation/Source/Programming/Tutorials/PlayerCamera/2/PlayerCamera_2.INT.udn b/Engine/Documentation/Source/Programming/Tutorials/PlayerCamera/2/PlayerCamera_2.INT.udn
--- a/Engine/Documentation/Source/Programming/Tutorials/PlayerCamera/2/PlayerCamera_2.INT.udn
+++ b/Engine/Documentation/Source/Programming/Tutorials/PlayerCamera/2/PlayerCamera_2.INT.udn
@@ -1,5 +1,5 @@
 Availability: Public
-Title:2. Configure Input To Control The Camera
+Title:2. Configure Input to Control the Camera
 Crumbs:%ROOT%, Programming, Programming/Tutorials, Programming/Tutorials/PlayerCamera
 Description:Configuring input that will be applied to a pawn and camera.
 version:4.9
diff --git a/Engine/Documentation/Source/Programming/Tutorials/PlayerCamera/3/PlayerCamera_3.INT.udn b/Engine/Documentation/Source/Programming/Tutorials/PlayerCamera/3/PlayerCamera_3.INT.udn
--- a/Engine/Documentation/Source/Programming/Tutorials/PlayerCamera/3/PlayerCamera_3.INT.udn
+++ b/Engine/Documentation/Source/Programming/Tutorials/PlayerCamera/3/PlayerCamera_3.INT.udn
@@ -1,5 +1,5 @@
 Availability: Public
-Title:3. Write C++ Code To React To Input
+Title:3. Write C++ Code to React to Input
 Crumbs:%ROOT%, Programming, Programming/Tutorials, Programming/Tutorials/PlayerCamera
 Description:Implementing code to control a pawn and a camera with input.
 version:4.9
diff --git a/Engine/Documentation/Source/Programming/Tutorials/PlayerInput/1/PlayerInput_1.INT.udn b/Engine/Documentation/Source/Programming/Tutorials/PlayerInput/1/PlayerInput_1.INT.udn
--- a/Engine/Documentation/Source/Programming/Tutorials/PlayerInput/1/PlayerInput_1.INT.udn
+++ b/Engine/Documentation/Source/Programming/Tutorials/PlayerInput/1/PlayerInput_1.INT.udn
@@ -1,5 +1,5 @@
 Availability: Public
-Title:1. Customize A Pawn
+Title:1. Customize a Pawn
 Crumbs:%ROOT%, Programming, Programming/Tutorials, Programming/Tutorials/PlayerInput
 Description:Introduction to Pawns
 version:4.9
