#UE4doc #docs: ##Overview removed as according to the instructional guide.

diff --git a/Engine/Documentation/Source/Resources/SampleGames/SwingNinja/Gameplay/Gameplay.INT.udn b/Engine/Documentation/Source/Resources/SampleGames/SwingNinja/Gameplay/Gameplay.INT.udn
--- a/Engine/Documentation/Source/Resources/SampleGames/SwingNinja/Gameplay/Gameplay.INT.udn
+++ b/Engine/Documentation/Source/Resources/SampleGames/SwingNinja/Gameplay/Gameplay.INT.udn
@@ -7,8 +7,6 @@ Description:Documentation covering the gameplay setup for the Swing Ninja Game.
 
 [REGION:fullwidth]
 
-## Overview
-
 Swing Ninja is a side-scrolling physics-based game. As such, there were two key elements to making it work properly:
 
 * Side-scrolling game behavior
diff --git a/Engine/Documentation/Source/Resources/Showcases/BasicMaterials/BasicMaterials.INT.udn b/Engine/Documentation/Source/Resources/Showcases/BasicMaterials/BasicMaterials.INT.udn
--- a/Engine/Documentation/Source/Resources/Showcases/BasicMaterials/BasicMaterials.INT.udn
+++ b/Engine/Documentation/Source/Resources/Showcases/BasicMaterials/BasicMaterials.INT.udn
@@ -54,9 +54,6 @@ Description:An overview of the Basic Materials example included with UE4.
 
 [TOC(start:2)]
 
-
-## Overview
-
 This example overviews some of the most elementary and critical concepts required to make use of Unreal Engine 4's powerful Materials system. Included in the showcase are the following features:
 
 * [Basic Lit Materials](#SimpleMaterials)
