#UE4doc #docs: ##Overview removed as according to the instructional guide.

diff --git a/Engine/Documentation/Source/Resources/SampleGames/SwingNinja/CustomTools/CustomTools.INT.udn b/Engine/Documentation/Source/Resources/SampleGames/SwingNinja/CustomTools/CustomTools.INT.udn
--- a/Engine/Documentation/Source/Resources/SampleGames/SwingNinja/CustomTools/CustomTools.INT.udn
+++ b/Engine/Documentation/Source/Resources/SampleGames/SwingNinja/CustomTools/CustomTools.INT.udn
@@ -11,7 +11,6 @@ Description:Documentation covering the custom tools created for the Swing Ninja
 
 [TOC(start:2 end:3)]
 [REGION:fullwidth]
-## Overview
 
 One of the most useful features of the Blueprint visual scripting system is how it can be used by level designers to construct customized tools which allow them to perform necessary design operations more quickly. To help make the creation of Swing Ninja more efficient, certain design tools were built using Blueprints. Below is an overview of each tool, why it was created, and how it was built.
 
