#UE4doc #docs: ##Overview removed as according to the instructional guide.

diff --git a/Engine/Documentation/Source/Resources/ContentExamples/WeatherSelector/WeatherSelector.INT.udn b/Engine/Documentation/Source/Resources/ContentExamples/WeatherSelector/WeatherSelector.INT.udn
--- a/Engine/Documentation/Source/Resources/ContentExamples/WeatherSelector/WeatherSelector.INT.udn
+++ b/Engine/Documentation/Source/Resources/ContentExamples/WeatherSelector/WeatherSelector.INT.udn
@@ -42,8 +42,6 @@ Description:An overview of the Weather Selector example level
 [/OBJECT]
 [/VAR]
 
-## Overview
-
 [REGION:fullwidth]
 
 [/REGION]
