#UE4doc #docs:
Added Engine Version field

diff --git a/Engine/Documentation/Source/Programming/Tutorials/UMG/UMG.INT.udn b/Engine/Documentation/Source/Programming/Tutorials/UMG/UMG.INT.udn
--- a/Engine/Documentation/Source/Programming/Tutorials/UMG/UMG.INT.udn
+++ b/Engine/Documentation/Source/Programming/Tutorials/UMG/UMG.INT.udn
@@ -2,6 +2,7 @@ Availability: Public
 Title:User Interface With UMG
 Crumbs:%ROOT%, GettingStarted
 Description:Create a simple menu system using UMG.
+version:4.9
 
 [VAR:Topic]
 [OBJECT:Topic]
