#UE4doc #docs:
Added Engine Version field
Changed Home link to point to the overview for the tutorial
Changed font settings for multiple terms
Replaced a couple pictures to change the indicator borders and arrows from green to yellow

diff --git a/Engine/Documentation/Source/Programming/Tutorials/UMG/1/UMG_1.INT.udn b/Engine/Documentation/Source/Programming/Tutorials/UMG/1/UMG_1.INT.udn
--- a/Engine/Documentation/Source/Programming/Tutorials/UMG/1/UMG_1.INT.udn
+++ b/Engine/Documentation/Source/Programming/Tutorials/UMG/1/UMG_1.INT.udn
@@ -2,6 +2,7 @@ Availability: Public
 Title:1. Set Up Module Dependencies For UMG
 Crumbs:%ROOT%, Programming, Programming/Tutorials, Programming/Tutorials/UMG
 Description:Modifying our project to include the modules needed to build menus.
+version:4.9
 
 [VAR:Steps]
 [OBJECT:Navigation]
@@ -10,7 +11,7 @@ Description:Modifying our project to include the modules needed to build menus.
 	[PARAM:current]
 	[/PARAM]
 	[PARAM:home]
-		[Programming Tutorials Home](Programming\Tutorials\)
+		[](Programming\Tutorials\UMG)
 	[/PARAM]
 	[PARAM:next]
 		[Next Step](Programming\Tutorials\UMG\2 "%Programming\Tutorials\UMG\2:title%")
@@ -26,7 +27,7 @@ Description:Modifying our project to include the modules needed to build menus.
 	* Select the C++ tab.
 	* Select Basic Code.
 	* Leave the default values of Desktop/Console and With Starter Content.
-	* We will use the name "HowTo_UMG".
+	* We will use the name HowTo_UMG.
 
 	[![](UMG_ProjectSetup.png)](UMG_ProjectSetup.png)
 
@@ -36,11 +37,11 @@ Description:Modifying our project to include the modules needed to build menus.
 
 	[![](OpenVisualStudio.png)](OpenVisualStudio.png)
 
-1. **UMG** is dependent on a few modules, so we'll need to add those modules to "HowTo_UMG.Build.cs".
+1. UMG is dependent on a few modules, so we'll need to add those modules to HowTo_UMG.Build.cs.
 
 	[![](BuildCS.png)](BuildCS.png)
 
-	In "HowTo_UMG.Build.CS", we need to add "UMG" to the list of included public modules, and "Slate" and "SlateCore" to the list of included private modules. We'll need to modify the first line in HowTo_UMG's constructor to read as follows:
+	In HowTo_UMG.Build.CS, we need to add UMG to the list of included public modules, and Slate and SlateCore to the list of included private modules. We'll need to modify the first line in HowTo_UMG's constructor to read as follows:
 		PublicDependencyModuleNames.AddRange(new string[] { "Core", "CoreUObject", "Engine", "InputCore", "UMG" });
 
 	And then we can uncomment the following line:
@@ -51,7 +52,7 @@ Description:Modifying our project to include the modules needed to build menus.
 
 ----
 
-With **UMG** set up, we can add code to our project's custom **Game Mode** to create and display our game's menus.
+With UMG set up, we can add code to our project's custom **Game Mode** to create and display our game's menus.
 
 [OBJECT:Section]
     [PARAMLITERAL:id]
