UEDOC-1845

Replaced gif with Youtube video

diff --git a/Engine/Documentation/Source/Engine/Actors/CameraActors/CameraActors.INT.udn b/Engine/Documentation/Source/Engine/Actors/CameraActors/CameraActors.INT.udn
--- a/Engine/Documentation/Source/Engine/Actors/CameraActors/CameraActors.INT.udn
+++ b/Engine/Documentation/Source/Engine/Actors/CameraActors/CameraActors.INT.udn
@@ -15,7 +15,17 @@ If you select a Camera Actor (or a Blueprint that contains a Camera Component),
 
 The Camera Actor can be found in the **Modes** panel under the **Basic** category. To place one into the world, simply drag it from the **Modes** panel into the game world.
 
-![](cameraPlacedAnim.gif)(convert:false)
+[OBJECT:EmbeddedVideo]
+[PARAMLITERAL:width]
+640
+[/PARAMLITERAL]
+[PARAMLITERAL:height]
+360
+[/PARAMLITERAL]
+[PARAMLITERAL:videoid]
+HMx63-uevs4
+[/PARAMLITERAL]
+[/OBJECT]
 
 ## Using Camera Actors
 
