#UE4doc #docs: fixed reiterated phrase: that the that the -> that the. Altered images to match the instructional guide: green -> yellow.

diff --git a/Engine/Documentation/Source/Platforms/Android/ReducingAPKSize/ReducingAPKSizes.INT.udn b/Engine/Documentation/Source/Platforms/Android/ReducingAPKSize/ReducingAPKSizes.INT.udn
--- a/Engine/Documentation/Source/Platforms/Android/ReducingAPKSize/ReducingAPKSizes.INT.udn
+++ b/Engine/Documentation/Source/Platforms/Android/ReducingAPKSize/ReducingAPKSizes.INT.udn
@@ -90,7 +90,7 @@ While this might not reduce your projects size as much as reducing or removing a
 To disable a plug-in inside of UE4 you need to do the following. 
 
 [region:note]
-	Make sure you fully test your project to ensure that the that the plug-ins you are disabling do not interfere with the functionality of your project.
+	Make sure you fully test your project to ensure that the plug-ins you are disabling do not interfere with the functionality of your project.
 [/region]
 
 1. First open up the Plug-ins Manager by going to the **Main Toolbar** and clicking on the **Window** option and then selecting the **Plug-in** option.
