#UE4doc #docs Deprecated out-dated documentation. Currently updating this document for Windows, Mac, and Linux.

diff --git a/Engine/Documentation/Source/GettingStarted/DownloadingUnrealEngine/GitHub/GitHub.INT.udn b/Engine/Documentation/Source/GettingStarted/DownloadingUnrealEngine/GitHub/GitHub.INT.udn
--- a/Engine/Documentation/Source/GettingStarted/DownloadingUnrealEngine/GitHub/GitHub.INT.udn
+++ b/Engine/Documentation/Source/GettingStarted/DownloadingUnrealEngine/GitHub/GitHub.INT.udn
@@ -3,12 +3,60 @@ Title:Downloading Unreal Engine from GitHub
 Crumbs:%ROOT%, GettingStarted, GettingStarted/DownloadingUnrealEngine
 Description:Downloading the Unreal Engine source code from GitHub.
 
+%Globals:osselect_linux%
+
+[OBJECT:ToggleButtonContent]
+	[PARAMLITERAL:category]
+	OS
+	[/PARAMLITERAL]
+	[PARAMLITERAL:id]
+	windows
+	[/PARAMLITERAL]
+	[PARAMLITERAL:active]
+	active_button_content
+	[/PARAMLITERAL]
+	[PARAM:content]
+	Test Windows
+	[/PARAM]
+[/OBJECT]
+
+[OBJECT:ToggleButtonContent]
+	[PARAMLITERAL:category]
+	OS
+	[/PARAMLITERAL]
+	[PARAMLITERAL:id]
+	mac
+	[/PARAMLITERAL]
+	[PARAMLITERAL:active]
+	active_button_content
+	[/PARAMLITERAL]
+	[PARAM:content]
+	Test Macintosh
+	[/PARAM]
+[/OBJECT]
+
+[OBJECT:ToggleButtonContent]
+	[PARAMLITERAL:category]
+	OS
+	[/PARAMLITERAL]
+	[PARAMLITERAL:id]
+	linux
+	[/PARAMLITERAL]
+	[PARAMLITERAL:active]
+	active_button_content
+	[/PARAMLITERAL]
+	[PARAM:content]
+	Test Linux
+	[/PARAM]
+[/OBJECT]
+
+<!--
+
 There are two ways to get the source code for Unreal Engine from GitHub:
 
 * Create a fork and clone the repository. 
 * Download the source and unzip it to a folder.
 
-
 ## Fork and Clone
 
 Forking the Unreal Engine repository gives you your own Unreal Engine repository on GitHub that you can clone, 
@@ -81,6 +129,11 @@ Download one of the following:
 
 Unpack the downloaded archive to a location of your choosing on your computer.
 
+## Next Steps
+
+Now, you're ready to [build the engine](Programming/Development/BuildingUnrealEngine).
+-->
+<!-- This section is deprecated per UE 4.8
 ## Installing Dependencies
 
 In addition to the source code, there are dependencies that need to be installed in order to compile and run the 
@@ -95,3 +148,4 @@ In version 4.5.1 and prior, the `Setup.bat` does not exist and the depencies mus
 ## Next Steps
 
 After installing the dependencies, you're ready to [build the engine](Programming/Development/BuildingUnrealEngine).
+-->
\ No newline at end of file
