#UE4doc #docs Currently updating outdated 'Downloading UE4 Source Code' documentation.

diff --git a/Engine/Documentation/Source/GettingStarted/DownloadingUnrealEngine/DownloadingUnrealEngine4.INT.udn b/Engine/Documentation/Source/GettingStarted/DownloadingUnrealEngine/DownloadingUnrealEngine4.INT.udn
--- a/Engine/Documentation/Source/GettingStarted/DownloadingUnrealEngine/DownloadingUnrealEngine4.INT.udn
+++ b/Engine/Documentation/Source/GettingStarted/DownloadingUnrealEngine/DownloadingUnrealEngine4.INT.udn
@@ -11,13 +11,47 @@ Description:Steps for connecting to the source code repository and downloading t
 Unreal Engine 4 is available to be downloaded by licensees directly from Epic's Perforce server. Before you can connect to the server, a secure VPN connection must be established. The login credentials for both the VPN and Perforce server are provided to your team's technical lead by Epic Games in your UDN welcome email. Only one set of login credentials is provided to each team.
 [/PUBLISH:Licensee]
 [PUBLISH:Rocket]
-Unreal Engine 4 source code is available to subscribers to be downloaded from the Unreal Engine GitHub repository. With access to the source code, you can modify the engine and tools to customize them to meet the needs of your project. You can build the editor for Windows and Mac and compile games for both Android and iOS. Source code for all tools is included as well, such as Unreal Lightmass and Unreal Frontend.
+Source code for Unreal Engine 4 (UE4) can be downloaded by our subscribers from the Unreal Engine GitHub repository. With access to our source code, subscribers can customize the engine and tools to meet their needs. Our subscribers are building the editor in Windows, Mac, Linux, and compiling games for both Android and iOS. Source code for all of our tools (such as Unreal Lightmass and Frontend) is included as well.
+[REGION:warning]
+Before you can access the [repository](https://github.com/EpicGames/UnrealEngine), you must be an Unreal Engine subscriber with a GitHub account.  
+*	If you haven't created a GitHub account, go to [GitHub's](https://github.com/) page to sign-up.   
+*	For instructions on associating your GitHub account with your Unreal Engine account, visit the [UE4 on GitHub](https://www.unrealengine.com/ue4-on-github) page.
+[/REGION]
+
+## Source Code Branches
+
+Once you've accessed the repository, you'll notice that we've published UE4's source code in three rolling branches, namely:
+
+*	**Release Branch**
+*	**Promoted Branch**
+*	**Master Branch**
+
+### Release Branch
+
+The [release branch](https://github.com/EpicGames/UnrealEngine/tree/release) (pictured below) is extensively tested by our QA team and makes a great starting point for learning UE4 or making Unreal games. We work hard to stabilize this branch with the goal of releasing new builds every one to two months.
+
+
+
+### Promoted Branch
+
+The [promoted branch](https://github.com/EpicGames/UnrealEngine/tree/promoted) (pictured below) is updated with builds for use by our designers and artists. We try to update this branch daily so that our subscribers have access to UE4's latest stable features. This branch is a great place to learn how feature development works in UE4.
 
-[REGION:note]
-Before you can access the [repository](https://github.com/EpicGames/UnrealEngine), you must be an Unreal Engine subscriber and associate your GitHub account with your Unreal Engine account.
 
-For instructions on associating your account, visit the [UE4 on GitHub](https://www.unrealengine.com/ue4-on-github) page.
+
+### Master Branch
+
+[REGION:warning]
+Because the master branch contains cutting-edge technology, UE4 isn't guaranteed to compile from source.
 [/REGION]
+
+The [master branch](https://github.com/EpicGames/UnrealEngine/tree/master) (pictured below) tracks [live changes](https://github.com/EpicGames/UnrealEngine/commits/master) being made by our engine team. Our experienced subscribers typically work from this branch when they want to work in tandem with UE4's developers.
+
+
+
+[REGION:note]
+Other branches may appear from time-to-time as we work to stabilize new releases or hotfixes.
+[/REGION]
+
 [/PUBLISH:Rocket]
 
 [PUBLISH:Licensee]
@@ -25,6 +59,9 @@ For instructions on associating your account, visit the [UE4 on GitHub](https://
 [/PUBLISH:Licensee]
 
 [PUBLISH:Rocket]
+
+## Getting Up and Running
+
 [INCLUDE:GettingStarted/DownloadingUnrealEngine/GitHub]
 [/PUBLISH:Rocket]
 
