UEDOC-1890

Repaced gifs with youtube videos

diff --git a/Engine/Documentation/Source/Engine/Rendering/ParticleSystems/Reference/Modules/Rotation/RotationModules.INT.udn b/Engine/Documentation/Source/Engine/Rendering/ParticleSystems/Reference/Modules/Rotation/RotationModules.INT.udn
--- a/Engine/Documentation/Source/Engine/Rendering/ParticleSystems/Reference/Modules/Rotation/RotationModules.INT.udn
+++ b/Engine/Documentation/Source/Engine/Rendering/ParticleSystems/Reference/Modules/Rotation/RotationModules.INT.udn
@@ -34,7 +34,17 @@ The Init Mesh Rot (Seed) module applies an initial rotation to mesh particles cr
 
 ## Initial Rotation
 
-![psys_initial_rotation.gif](psys_initial_rotation.gif)(convert:false)
+[OBJECT:EmbeddedVideo]
+[PARAMLITERAL:width]
+640
+[/PARAMLITERAL]
+[PARAMLITERAL:height]
+360
+[/PARAMLITERAL]
+[PARAMLITERAL:videoid]
+TkMzhGUYjnI
+[/PARAMLITERAL]
+[/OBJECT]
 
 The Initial Rotation module is used to set the initial rotation of a particle at spawn time. It contains the following member(s):
 
@@ -60,7 +70,17 @@ The Init Rotation (Seeded) module is identical to the [Initial Rotation](#Initia
 
 ## Rotation/Life
 
-![psys_rotation_life.gif](psys_rotation_life.gif)(convert:false)
+[OBJECT:EmbeddedVideo]
+[PARAMLITERAL:width]
+640
+[/PARAMLITERAL]
+[PARAMLITERAL:height]
+360
+[/PARAMLITERAL]
+[PARAMLITERAL:videoid]
+M-mHMXH2y7Q
+[/PARAMLITERAL]
+[/OBJECT]
 
 The Rotation/Life module is used to set the rotation of particles over their lifetime. It contains the following member(s):
 
