#UE4doc #docs:
Changes didn't save before last submit

diff --git a/Engine/Documentation/Source/Programming/Tutorials/UMG/1/UMG_1.INT.udn b/Engine/Documentation/Source/Programming/Tutorials/UMG/1/UMG_1.INT.udn
--- a/Engine/Documentation/Source/Programming/Tutorials/UMG/1/UMG_1.INT.udn
+++ b/Engine/Documentation/Source/Programming/Tutorials/UMG/1/UMG_1.INT.udn
@@ -1,5 +1,5 @@
 Availability: Public
-Title:1. Set Up Module Dependencies For UMG
+Title:1. Set Up Module Dependencies for UMG
 Crumbs:%ROOT%, Programming, Programming/Tutorials, Programming/Tutorials/UMG
 Description:Modifying our project to include the modules needed to build menus.
 version:4.9
