#UE4doc #docs:
Added Engine Version field
Changed Home link to point to the overview for the tutorial
Changed font settings for multiple terms

diff --git a/Engine/Documentation/Source/Programming/Tutorials/UMG/2/UMG_2.INT.udn b/Engine/Documentation/Source/Programming/Tutorials/UMG/2/UMG_2.INT.udn
--- a/Engine/Documentation/Source/Programming/Tutorials/UMG/2/UMG_2.INT.udn
+++ b/Engine/Documentation/Source/Programming/Tutorials/UMG/2/UMG_2.INT.udn
@@ -2,6 +2,7 @@ Availability: Public
 Title:2. Extend Our Game Mode
 Crumbs:%ROOT%, Programming, Programming/Tutorials, Programming/Tutorials/UMG
 Description:Extending the project's Game Mode to add and remove menus from the screen.
+version:4.9
 
 [VAR:Steps]
 [OBJECT:Navigation]
@@ -11,7 +12,7 @@ Description:Extending the project's Game Mode to add and remove menus from the s
 	[PARAM:current]
 	[/PARAM]
 	[PARAM:home]
-		[Programming Tutorials Home](Programming\Tutorials\)
+		[](Programming\Tutorials\UMG)
 	[/PARAM]
 	[PARAM:next]
 		[Next Step](Programming\Tutorials\UMG\3 "%Programming\Tutorials\UMG\3:title%")
@@ -22,7 +23,7 @@ Description:Extending the project's Game Mode to add and remove menus from the s
 %Steps%
 
 
-1. The menus we create will be made from **User Widgets**. We'll need to write a function that creates and displays a new **User Widget**, and then call that function when the game starts. We'll also need to keep track of what we have created so that we can remove it later. Since each project already comes with a custom **Game Mode** class, we can simply open ours, which is defined in `HowTo_UMGGameMode.h`. The following functions and properties will need to be added to the bottom of the class:
+1. The menus we create will be made from **User Widgets**. We'll need to write a function that creates and displays a new User Widget, and then call that function when the game starts. We'll also need to keep track of what we have created so that we can remove it later. Since each project already comes with a custom **GameMode** class, we can simply open ours, which is defined in HowTo_UMGGameMode.h. The following functions and properties will need to be added to the bottom of the class:
 
 		public:
 			/** Called when the game starts. */
@@ -41,11 +42,11 @@ Description:Extending the project's Game Mode to add and remove menus from the s
 			UPROPERTY()
 			UUserWidget* CurrentWidget;
 
-1. In order to use **User Widgets** in our code,  and add the following line to the top of the "#include" section:
+1. In order to use User Widgets in our code,  and add the following line to the top of the #include section:
 
 		#include "Blueprint/UserWidget.h"
 
-1. Moving to `HowTo_UMGGameMode.cpp` now, we need to fill out the bodies of the two functions we declared. We'll start with overriding BeginPlay():
+1. Moving to HowTo_UMGGameMode.cpp now, we need to fill out the bodies of the two functions we declared. We'll start with overriding **BeginPlay()**:
 
 		void AHowTo_UMGGameMode::BeginPlay()
 		{
@@ -53,9 +54,9 @@ Description:Extending the project's Game Mode to add and remove menus from the s
 			ChangeMenuWidget(StartingWidgetClass);
 		}
 
-	[REGION:note]When overriding functions from a parent class (referenced by the word Super), as we do here with _BeginPlay_, it is often important to call the parent class' version of that function. Since our version of the function is only meant to add one step to the end of the existing procedure, we call _Super::BeginPlay_ in the first line of the function.[/REGION]
+	[REGION:note]When overriding functions from a parent class (referenced by the word Super), as we do here with BeginPlay, it is often important to call the parent class' version of that function. Since our version of the function is only meant to add one step to the end of the existing procedure, we call **Super::BeginPlay** in the first line of the function.[/REGION]
 
-1. Next, still in `HowTo_UMGGameMode.cpp`, we need to define how we change between menus. We will need to remove whatever **User Widget** we have active from the viewport, if any. Then we can create and a new **User Widget** and add it to the viewport.
+1. Next, still in HowTo_UMGGameMode.cpp, we need to define how we change between menus. We will need to remove whatever User Widget we have active from the viewport, if any. Then we can create and a new User Widget and add it to the viewport.
 
 		void AHowTo_UMGGameMode::ChangeMenuWidget(TSubclassOf<UUserWidget> NewWidgetClass)
 		{
@@ -74,7 +75,7 @@ Description:Extending the project's Game Mode to add and remove menus from the s
 			}
 		}
 
-	[REGION:note]This code will create instances of any **Widgets** we provide and put them on the screen. It also removes them, so that only one is active at a time, although **Unreal Engine** can handle displaying and interacting with many **Widgets** at once. We never need to destroy a **Widget** directly, because removing it from the viewport and clearing (or changing) all variables that reference it will cause it to be cleaned up by **Unreal Engine's** garbage collection system.[/REGION]
+	[REGION:note]This code will create instances of any **Widgets** we provide and put them on the screen. It also removes them, so that only one is active at a time, although **Unreal Engine** can handle displaying and interacting with many Widgets at once. We never need to destroy a Widget directly, because removing it from the viewport and clearing (or changing) all variables that reference it will cause it to be cleaned up by Unreal Engine's garbage collection system.[/REGION]
 
 
 ----
