UEDOC-1891

Replaced gifs with YouTube Videos

diff --git a/Engine/Documentation/Source/Engine/Rendering/ParticleSystems/Reference/Modules/RotationRate/RotationRateModules.INT.udn b/Engine/Documentation/Source/Engine/Rendering/ParticleSystems/Reference/Modules/RotationRate/RotationRateModules.INT.udn
--- a/Engine/Documentation/Source/Engine/Rendering/ParticleSystems/Reference/Modules/RotationRate/RotationRateModules.INT.udn
+++ b/Engine/Documentation/Source/Engine/Rendering/ParticleSystems/Reference/Modules/RotationRate/RotationRateModules.INT.udn
@@ -55,7 +55,17 @@ The Mesh Rot Rate Over Life module is used to set the rotation rate of a mesh-ba
 
 ## Initial Rot Rate
 
-![psys_initial_rotrate.gif](psys_initial_rotrate.gif)(convert:false)
+[OBJECT:EmbeddedVideo]
+[PARAMLITERAL:width]
+640
+[/PARAMLITERAL]
+[PARAMLITERAL:height]
+360
+[/PARAMLITERAL]
+[PARAMLITERAL:videoid]
+mLSdoDCyx4Q
+[/PARAMLITERAL]
+[/OBJECT]
 
 The Initial RotRate module is used to set the rotation rate of a particle when it is emitted. It contains the following member(s):
 
@@ -80,7 +90,17 @@ The Initial RotRate (Seeded) module is identical to the [Initial Rot Rate](#Init
 
 ## RotRate/Life
 
-![psys_rotrate_life.gif](psys_rotrate_life.gif)(convert:false)
+[OBJECT:EmbeddedVideo]
+[PARAMLITERAL:width]
+640
+[/PARAMLITERAL]
+[PARAMLITERAL:height]
+360
+[/PARAMLITERAL]
+[PARAMLITERAL:videoid]
+ELwP3O46I_k
+[/PARAMLITERAL]
+[/OBJECT]
 
 The Rot Rate/Life module is used to scale the rotation rate of a particle over its lifetime. It contains the following member(s):
 
