UEDOC-1893

Replaced gif with YouTube video

diff --git a/Engine/Documentation/Source/Engine/Rendering/ParticleSystems/Reference/Modules/VectorField/VectorField.INT.udn b/Engine/Documentation/Source/Engine/Rendering/ParticleSystems/Reference/Modules/VectorField/VectorField.INT.udn
--- a/Engine/Documentation/Source/Engine/Rendering/ParticleSystems/Reference/Modules/VectorField/VectorField.INT.udn
+++ b/Engine/Documentation/Source/Engine/Rendering/ParticleSystems/Reference/Modules/VectorField/VectorField.INT.udn
@@ -58,7 +58,17 @@ The Vector Field Initial Rotation module will allow for an initial amount of rot
 
 ## VF Rotation Rate
 
-![VFInitRotRate.gif](VFInitRotRate.gif)(convert:false)
+[OBJECT:EmbeddedVideo]
+[PARAMLITERAL:width]
+640
+[/PARAMLITERAL]
+[PARAMLITERAL:height]
+360
+[/PARAMLITERAL]
+[PARAMLITERAL:videoid]
+UmZ9q-3KP8o
+[/PARAMLITERAL]
+[/OBJECT]
 
 If you need a local vector field to have a constant rate of spin, then the VF Rotation Rate module is exactly what you need. Applying it allows for the specified vector field to have distributed rotation rate. This can lead to some highly interesting effects.
 
