UEDOC-1894

Replaced gif with YouTube video

diff --git a/Engine/Documentation/Source/Engine/Rendering/ParticleSystems/Reference/Modules/Velocity/VelocityModules.INT.udn b/Engine/Documentation/Source/Engine/Rendering/ParticleSystems/Reference/Modules/Velocity/VelocityModules.INT.udn
--- a/Engine/Documentation/Source/Engine/Rendering/ParticleSystems/Reference/Modules/Velocity/VelocityModules.INT.udn
+++ b/Engine/Documentation/Source/Engine/Rendering/ParticleSystems/Reference/Modules/Velocity/VelocityModules.INT.udn
@@ -9,7 +9,17 @@ Velocity modules exist to make adjustments to particle velocity in a variety of
 
 ## Initial Velocity
 
-![psys_initial_velocity.gif](psys_initial_velocity.gif)(convert:false)
+[OBJECT:EmbeddedVideo]
+[PARAMLITERAL:width]
+640
+[/PARAMLITERAL]
+[PARAMLITERAL:height]
+360
+[/PARAMLITERAL]
+[PARAMLITERAL:videoid]
+htluEXsHV4o
+[/PARAMLITERAL]
+[/OBJECT]
 
 The Initial Velocity module sets the initial velocity of a particle. It contains the following member(s):
 
@@ -69,7 +79,17 @@ Contributes the velocity of the parent (the particle emitter itself) to a partic
 
 ## Velocity/Life
 
-![psys_velocity_life.gif](psys_velocity_life.gif)(convert:false)
+[OBJECT:EmbeddedVideo]
+[PARAMLITERAL:width]
+640
+[/PARAMLITERAL]
+[PARAMLITERAL:height]
+360
+[/PARAMLITERAL]
+[PARAMLITERAL:videoid]
+zblNBP4eJLE
+[/PARAMLITERAL]
+[/OBJECT]
 
 The Velocity/Life module scales the velocity by a given amount over the lifetime of the particle. It contains the following member(s):
 
