UEDOC-1896

Replaced gif with Youtube video

diff --git a/Engine/Documentation/Source/Engine/Rendering/ParticleSystems/Reference/TypeData/Beam/BeamTypeData.INT.udn b/Engine/Documentation/Source/Engine/Rendering/ParticleSystems/Reference/TypeData/Beam/BeamTypeData.INT.udn
--- a/Engine/Documentation/Source/Engine/Rendering/ParticleSystems/Reference/TypeData/Beam/BeamTypeData.INT.udn
+++ b/Engine/Documentation/Source/Engine/Rendering/ParticleSystems/Reference/TypeData/Beam/BeamTypeData.INT.udn
@@ -8,7 +8,17 @@ Description:The Beam Type Data module is useful for creating effects such as bea
 The **Beam** type data module indicates that the emitter should output beams - connecting particles to form a stream between a source point, such as the emitter, and a target point, such as a particle or Actor.
 [/EXCERPT:Intro]
 
-![psys_beam.gif](psys_beam.gif)(convert:false)
+[OBJECT:EmbeddedVideo]
+[PARAMLITERAL:width]
+640
+[/PARAMLITERAL]
+[PARAMLITERAL:height]
+360
+[/PARAMLITERAL]
+[PARAMLITERAL:videoid]
+coXs-CmI_v8
+[/PARAMLITERAL]
+[/OBJECT]
 
 The Beam type data module has the following properties:
 
