UEDOC-1906

Replaced gifs with Youtube videos

diff --git a/Engine/Documentation/Source/Engine/Rendering/ParticleSystems/ParticleLights/ParticleLights.INT.udn b/Engine/Documentation/Source/Engine/Rendering/ParticleSystems/ParticleLights/ParticleLights.INT.udn
--- a/Engine/Documentation/Source/Engine/Rendering/ParticleSystems/ParticleLights/ParticleLights.INT.udn
+++ b/Engine/Documentation/Source/Engine/Rendering/ParticleSystems/ParticleLights/ParticleLights.INT.udn
@@ -54,7 +54,29 @@ The performance cost of particle lights is almost completely dependent on how mu
 
 The general strategy when setting up particle lights is to have a very small number of large particle lights, and a larger number of small particle lights.  It is important to tweak and profile with the _ProfileGPU_ console command as the cost of particle lights can easily get out of hand.  The cost shows up under _Lights_.
 
-![](PL_Sparks_2.gif)(convert:false) ![](PL_flow_2.gif)(convert:false)
+[OBJECT:EmbeddedVideo]
+[PARAMLITERAL:width]
+640
+[/PARAMLITERAL]
+[PARAMLITERAL:height]
+360
+[/PARAMLITERAL]
+[PARAMLITERAL:videoid]
+bJ1xiMPl71Y
+[/PARAMLITERAL]
+[/OBJECT]
+
+[OBJECT:EmbeddedVideo]
+[PARAMLITERAL:width]
+640
+[/PARAMLITERAL]
+[PARAMLITERAL:height]
+360
+[/PARAMLITERAL]
+[PARAMLITERAL:videoid]
+nyGqpilZe08
+[/PARAMLITERAL]
+[/OBJECT]
 
 For reference, here are some still images of the above animations with just the Particle Lights emitters.
 
