UEDOC-1907

Replaced gifs with Youtube videos

diff --git a/Engine/Documentation/Source/Engine/Rendering/ParticleSystems/Reference/Modules/Acceleration/AccelerationModules.INT.udn b/Engine/Documentation/Source/Engine/Rendering/ParticleSystems/Reference/Modules/Acceleration/AccelerationModules.INT.udn
--- a/Engine/Documentation/Source/Engine/Rendering/ParticleSystems/Reference/Modules/Acceleration/AccelerationModules.INT.udn
+++ b/Engine/Documentation/Source/Engine/Rendering/ParticleSystems/Reference/Modules/Acceleration/AccelerationModules.INT.udn
@@ -10,7 +10,17 @@ These modules provide the ability to apply acceleration, or a change in velocity
 
 ## Acceleration
 
-![psys_acceleration.gif](psys_acceleration.gif)(convert:false)
+[OBJECT:EmbeddedVideo]
+[PARAMLITERAL:width]
+640
+[/PARAMLITERAL]
+[PARAMLITERAL:height]
+360
+[/PARAMLITERAL]
+[PARAMLITERAL:videoid]
+o3TuesKDJ_A
+[/PARAMLITERAL]
+[/OBJECT]
 
 The Acceleration module sets the initial acceleration of the particle. This value is stored per-particle, and is applied each frame during the Update call. It contains the following member(s):
 
@@ -29,7 +39,17 @@ Each frame, the current and base velocity values of the particle is then updated
 ## Const Acceleration
 
 
-![psys_acceleration.gif](psys_acceleration.gif)(convert:false)
+[OBJECT:EmbeddedVideo]
+[PARAMLITERAL:width]
+640
+[/PARAMLITERAL]
+[PARAMLITERAL:height]
+360
+[/PARAMLITERAL]
+[PARAMLITERAL:videoid]
+o3TuesKDJ_A
+[/PARAMLITERAL]
+[/OBJECT]
 
 The Const Acceleration module is very similar to the Acceleration module, in that it applies an initial acceleration to the particle. The key difference, however, is that the Const Acceleration module does _not_ accept a Distribution property. 
 
@@ -48,7 +68,17 @@ Each frame, the current and base velocity values of the particle is then updated
 
 ## Drag 
 
-![phys_drag.gif](phys_drag.gif)(w:430 convert:false)
+[OBJECT:EmbeddedVideo]
+[PARAMLITERAL:width]
+640
+[/PARAMLITERAL]
+[PARAMLITERAL:height]
+360
+[/PARAMLITERAL]
+[PARAMLITERAL:videoid]
+inAil0AgElQ
+[/PARAMLITERAL]
+[/OBJECT]
 
 Drag applies a constant resistance to motion to each particle. The coefficient value is calculated at spawn time and applied to the particle at each frame.
 
