UEDOC-1908

Replaced gifs with Youtube videos

diff --git a/Engine/Documentation/Source/Engine/Rendering/ParticleSystems/Reference/Modules/Attractor/AttractorModules.INT.udn b/Engine/Documentation/Source/Engine/Rendering/ParticleSystems/Reference/Modules/Attractor/AttractorModules.INT.udn
--- a/Engine/Documentation/Source/Engine/Rendering/ParticleSystems/Reference/Modules/Attractor/AttractorModules.INT.udn
+++ b/Engine/Documentation/Source/Engine/Rendering/ParticleSystems/Reference/Modules/Attractor/AttractorModules.INT.udn
@@ -9,7 +9,17 @@ These modules implement methods of causing particles to be attracted toward a sp
 
 This is the result of using a point attractor in combination with a line attractor with its strength animated over the life of the particle to create the swirling effect.
 
-![psys_attractor.gif](psys_attractor.gif)(convert:false)
+[OBJECT:EmbeddedVideo]
+[PARAMLITERAL:width]
+640
+[/PARAMLITERAL]
+[PARAMLITERAL:height]
+360
+[/PARAMLITERAL]
+[PARAMLITERAL:videoid]
+Ntf5chUK4jM
+[/PARAMLITERAL]
+[/OBJECT]
 
 ## Attractor Line
 
@@ -27,7 +37,17 @@ The line attractor allows for having particles be drawn to a line in 3D space.
 
 ## Particle Attractor
 
-![psys_attractor_particle.gif](psys_attractor_particle.gif)(convert:false)
+[OBJECT:EmbeddedVideo]
+[PARAMLITERAL:width]
+640
+[/PARAMLITERAL]
+[PARAMLITERAL:height]
+360
+[/PARAMLITERAL]
+[PARAMLITERAL:videoid]
+6EHCXt4q-Ao
+[/PARAMLITERAL]
+[/OBJECT]
 
 The particle attractor allows for having particles be drawn to the particles emitted from another emitter in the system. It contains the following member(s):
 
@@ -60,7 +80,17 @@ The particle attractor allows for having particles be drawn to the particles emi
 
 ## Point Attractor
 
-![psys_attractor_point.gif](psys_attractor_point.gif)(convert:false)
+[OBJECT:EmbeddedVideo]
+[PARAMLITERAL:width]
+640
+[/PARAMLITERAL]
+[PARAMLITERAL:height]
+360
+[/PARAMLITERAL]
+[PARAMLITERAL:videoid]
+fHYk6juHU0c
+[/PARAMLITERAL]
+[/OBJECT]
 
 The point attractor allows for defining a point to which particles are attracted and its radius of effect. It contains the following member(s):
 
@@ -79,7 +109,17 @@ The point attractor allows for defining a point to which particles are attracted
 
 ## Point Gravity
 
-![phys_pointGravity.gif](phys_pointGravity.gif)(w:420 convert:false)
+[OBJECT:EmbeddedVideo]
+[PARAMLITERAL:width]
+640
+[/PARAMLITERAL]
+[PARAMLITERAL:height]
+360
+[/PARAMLITERAL]
+[PARAMLITERAL:videoid]
+uqXArTDw9as
+[/PARAMLITERAL]
+[/OBJECT]
 
 The Point Gravity module attracts particles towards a singular point in 3D space, using a given strength value. It is similar to the Newton field in Maya's particle system. It can be used to pull particles toward a source or to create orbital effects around a point.
 
