UEDOC-1808

Updated image to new image provided by [REDACTED_PERSON]

diff --git a/Engine/Documentation/Source/Engine/UI/LevelEditor/Modes/MeshPaintMode/VertexColor/MaterialSetup/4Way/4Way.INT.udn b/Engine/Documentation/Source/Engine/UI/LevelEditor/Modes/MeshPaintMode/VertexColor/MaterialSetup/4Way/4Way.INT.udn
--- a/Engine/Documentation/Source/Engine/UI/LevelEditor/Modes/MeshPaintMode/VertexColor/MaterialSetup/4Way/4Way.INT.udn
+++ b/Engine/Documentation/Source/Engine/UI/LevelEditor/Modes/MeshPaintMode/VertexColor/MaterialSetup/4Way/4Way.INT.udn
@@ -8,7 +8,7 @@ Description:Guide to setting up materials for 4-way (ARGB) texture blending.
 
 To handle 4 textures in the blend weights setup, you need to create a material similar to this:
 
-![](4-way_material.png)
+![](4-way_Material_Update.png)
 
 This works very similarly to the 3-way setup, but brings alpha into the mix to give a 4th blending option. Here is what is happening:
 
