#UE4doc #docs:
Added Engine Version field
Changed Home link to point to the overview for the tutorial
Changed font settings for multiple terms
Replaced a few pictures to change the indicator borders and arrows from green to yellow
Changed a few steps to include more information about the latest version of the editor

diff --git a/Engine/Documentation/Source/Programming/Tutorials/UMG/3/UMG_3.INT.udn b/Engine/Documentation/Source/Programming/Tutorials/UMG/3/UMG_3.INT.udn
--- a/Engine/Documentation/Source/Programming/Tutorials/UMG/3/UMG_3.INT.udn
+++ b/Engine/Documentation/Source/Programming/Tutorials/UMG/3/UMG_3.INT.udn
@@ -2,6 +2,7 @@ Availability: Public
 Title:3. Create Menu Widget Blueprints
 Crumbs:%ROOT%, Programming, Programming/Tutorials, Programming/Tutorials/UMG
 Description:Create widget blueprints to be used as menu screens.
+version:4.9
 
 [VAR:Steps]
 [OBJECT:Navigation]
@@ -11,7 +12,7 @@ Description:Create widget blueprints to be used as menu screens.
 	[PARAM:current]
 	[/PARAM]
 	[PARAM:home]
-		[Programming Tutorials Home](Programming\Tutorials\)
+		[](Programming\Tutorials\UMG)
 	[/PARAM]
 	[PARAM:next]
 		[Next Step](Programming\Tutorials\UMG\4 "%Programming\Tutorials\UMG\4:title%")
@@ -26,53 +27,56 @@ Description:Create widget blueprints to be used as menu screens.
 
 	[![](CompileFromEditor.png)](CompileFromEditor.png)
 
-1. We will now create the **User Widgets** that our **Game Mode** will use as menus. This is done with the "Add New" button in the **Content Browser**. The **Widget Blueprint** class is found in the "User Interface" category. We need to create two of these, one named "MainMenu" and one named "NewGameMenu". Our game will begin at the Main Menu, and will have an option to proceed to the New Game Menu.
+1. We will now create the User Widgets that our **GameMode** will use as menus. This is done with the **Add New** button in the **Content Browser**. The **Widget Blueprint** class is found in the **User Interface** category. We need to create two of these, one named MainMenu and one named NewGameMenu. Our game will begin at the Main Menu, and will have an option to proceed to the New Game Menu.
 
 	[![](CreateWidgetBP.png)](CreateWidgetBP.png)
 
-1. Double-clicking the "MainMenu" **Widget** we've just created will take us to the **Blueprint Designer**, where we can create our menu layout.
+1. Double-clicking the MainMenu **Widget** we've just created will take us to the **Blueprint Designer**, where we can create our menu layout.
 
-1. Let's drag a **Button** from the "Common" section of the **Palette Panel** to the **Graph**. This **Button** will eventually be used to open the New Game Menu.
+1. Let's drag a **Button** and a **Text** from the **Common** section of the **Palette Panel** to the **Graph**. This Button will eventually be used to open the New Game Menu.
 
 	[![](AddButtonAndTextBlock.png)](AddButton.png)
 
-1. The first step in making our layout look right is to adjust the location and size of our **Button**. We should make the following changes:
+1. The first step in making our layout look right is to adjust the location and size of our Button. We should make the following changes:
 	* Set the size to 200x200.
 	* Set the position to (200, 100).
-	* Rename it "NewGameButton", just to make it easier to recoginze when we hook up functionality to it later on.
+	* Rename it NewGameButton, just to make it easier to recoginze when we hook up functionality to it later on.
 
 	[![](ButtonDetails.png)](ButtonDetails.png)
 
-1. Since we're not drawing a custom set of images for the **Button**, we can label it by dropping a **Text Block** onto it and making the following changes:
-	* Set the text to say "New Game".
-	* Change the Visibility to Hit Test Invisible. This will prevent the **Text Block** from intercepting mouse-clicks that are intended for the **Button** underneath.
-	* Set the name to "NewGameText". This is not needed, but it's a good habit.
+1. Since we're not drawing a custom set of images for the Button, we can label it by dragging and dropping the Text Block onto it and making the following changes:
+	* Set the text to say New Game.
+	* Change the Visibility to Hit Test Invisible. This will prevent the Text Block from intercepting mouse-clicks that are intended for the Button underneath.
+	* Set the name to NewGameText. This is not needed, but it's a good habit.
 
 	[![](TextBlockDetails.png)](TextBlockDetails.png)
 
-1. Next, we'll want to make a "Quit" feature with a second **Button** and **Text Block**. Set those up in the same way as the "New Game" **Button** and **Text Block**.
+1. Next, we'll want to make a "Quit" feature with a second Button and Text Block. Set those up in the same way as the New Game Button and Text Block except the following changes.
+	* Set the name of the Button to QuitButton
+	* Set the position of the button to 600, 100
+	* Set the name of the Text Block to QuitText
 
-1. After that, we can add **Events** to our **Buttons** so that we can run code when a **Button** is clicked. This is done by locating and pressing the "+" next to the appropriate **Event** name in the **Details Panel**. In this case, "OnClicked" is the only **Event** available. Create this event for both the NewGameButton amd QuitButton **Widgets**.
+1. After that, we can add **Events** to our Buttons so that we can run code when a Button is clicked. This is done by locating and pressing the + next to the appropriate Event name in the **Details Panel**. In this case, **OnClicked** is the event we want to use. Create this event for both the NewGameButton amd QuitButton Widgets.
 
 	[![](CreateOnClickedEvent.png)](CreateOnClickedEvent.png)
 
 	[REGION:note]Designers can build functionality with **Blueprint** scripting here, or C++ programmers can connect nodes that call exposed functions.[/REGION]
 
-1. For the **Event** called "OnClicked(NewGameButton)", we'll want to:
-	* Connect a _ChangeMenuWidget_ node to use the function we added to our **Game Mode** earlier.
-	* Set the New Widget Class field on the _ChangeMenuWidget_ node to the "NewGameMenu" asset.
+1. For the Event called **OnClicked(NewGameButton)**, we'll want to:
+	* Connect a **ChangeMenuWidget** node to use the function we added to our GameMode earlier.
+	* Set the **New Widget Class** field on the **ChangeMenuWidget** node to the NewGameMenu asset.
 
 	[![](OnClicked_NewGame.png)](OnClicked_NewGame.png)
 
-1. For the "OnClicked(QuitButton)" **Event**, we'll want to:
-	* Connect a "Quit Game" node.
+1. For the **OnClicked(QuitButton)** Event, we'll want to:
+	* Connect a **Quit Game** node.
 
 	[![](OnClicked_Quit.png)](OnClicked_Quit.png)
 
 
 ----
 
-With our main menu built, we can set up a **Game Mode** asset that will load it as soon as the level starts.
+With our main menu built, we can set up a GameMode asset that will load it as soon as the level starts.
 
 [OBJECT:Section]
     [PARAMLITERAL:id]
