#UE4doc #docs:
Added Engine Version field
Changed Home link to point to the overview for the tutorial
Changed font settings for multiple terms
Replaced a picture to change the indicator borders and arrows from green to yellow

diff --git a/Engine/Documentation/Source/Programming/Tutorials/UMG/4/UMG_4.INT.udn b/Engine/Documentation/Source/Programming/Tutorials/UMG/4/UMG_4.INT.udn
--- a/Engine/Documentation/Source/Programming/Tutorials/UMG/4/UMG_4.INT.udn
+++ b/Engine/Documentation/Source/Programming/Tutorials/UMG/4/UMG_4.INT.udn
@@ -2,6 +2,7 @@ Availability: Public
 Title:4. Configure Our Game Mode
 Crumbs:%ROOT%, Programming, Programming/Tutorials, Programming/Tutorials/UMG
 Description:Set up our Game Mode to use our Main Menu Widget.
+version:4.9
 
 [VAR:Steps]
 [OBJECT:Navigation]
@@ -11,7 +12,7 @@ Description:Set up our Game Mode to use our Main Menu Widget.
 	[PARAM:current]
 	[/PARAM]
 	[PARAM:home]
-		[Programming Tutorials Home](Programming\Tutorials\)
+		[](Programming\Tutorials\UMG)
 	[/PARAM]
 	[PARAM:next]
 		[Next Step](Programming\Tutorials\UMG\5 "%Programming\Tutorials\UMG\5:title%")
@@ -22,56 +23,56 @@ Description:Set up our Game Mode to use our Main Menu Widget.
 %Steps%
 
 
-1. In the **Content Browser**, we will add a **Blueprint Class** based on our project's **Game Mode**. This makes it possible to set the exposed variables on those two classes to whatever values we want. To do this:
+1. In the **Content Browser**, we will add a **Blueprint Class** based on our project's **GameMode**. This makes it possible to set the exposed variables on those two classes to whatever values we want. To do this:
 
 	* Click the **Add** button in the **Content Browser**.
 
 	[![](AddNewBPClass.png)](AddNewBPClass.png)
 
-	* Pick HowTo_UMGGameMode as the parent class. It will be listed in the "All Classes" section.
+	* Pick **HowTo_UMGGameMode** as the parent class. It will be listed in the **All Classes** section.
 
 	[![](PickParentClassForGameModeBP.png)](PickParentClassForGameModeBP.png)
 
-	* Name the resulting **Blueprint** asset "MenuGameMode".
+	* Name the resulting Blueprint asset MenuGameMode.
 
-1. In order to see our mouse cursor in-game, we'll need to create a **Blueprint** of the **Player Controller** as we did with our **Game Mode**.
+1. In order to see our mouse cursor in-game, we'll need to create a Blueprint of the **PlayerController** as we did with our GameMode.
 
 	* Click the **Add** button in the **Content Browser** again.
 	
-	* Select **Player Controller** from the "Common Classes" section.
+	* Select **Player Controller** from the **Common Classes** section.
 
-	* Name the **Blueprint** "MenuPlayerController".
+	* Name the Blueprint MenuPlayerController.
 
-1. Edit "MenuPlayerController".
+1. Edit MenuPlayerController.
 
-	* Check the "Show Mouse Cursor" box.
+	* Check the **Show Mouse Cursor** box.
 
 	[![](GamePlayerController.png)](GamePlayerController.png)
 
-1. Edit "MenuGameMode".
+1. Edit MenuGameMode.
 
-	* The Starting Widget Class must be set to the "MainMenu" asset in order to bring up our menu when the game begins.
+	* The **Starting Widget Class** must be set to the MainMenu asset in order to bring up our menu when the game begins.
 
-	* The Default Pawn Class should be set to **Pawn** instead of **DefaultPawn** so that our player will not be able to fly around while in the menu.
+	* The **Default Pawn Class** should be set to **Pawn** instead of **DefaultPawn** so that our player will not be able to fly around while in the menu.
 
-	* The Player Controller Class should be set to the "MenuPlayerController" asset we created so that the mouse cursor will be shown in-game.
+	* The **Player Controller Class** should be set to the MenuPlayerController asset we created so that the mouse cursor will be shown in-game.
 
 	[![](ConfigureGameMode.png)](ConfigureGameMode.png)
 
-1. In order for our **Blueprint** to be used, we must return to the **Level Editor** window and change the **World Settings** for our current **Level** via the **Settings** button.
+1. In order for our Blueprint to be used, we must return to the **Level Editor** window and change the **World Settings** for our current **Level** via the **Settings** button.
 
 	[![](WorldSettingsBar.png)](WorldSettingsBar.png)
 
-	[REGION:note]It is also possible to set the default **Game Mode** in the **Project Settings** menu, under the **Maps and Modes** section. If you use this method, all of your levels will default to the **Game Mode** you choose, unless overridden individually. Which method you use depends on how you prefer to set up your project.[/REGION]
+	[REGION:note]It is also possible to set the default GameMode in the **Project Settings** menu, under the **Maps and Modes** section. If you use this method, all of your levels will default to the GameMode you choose, unless overridden individually. Which method you use depends on how you prefer to set up your project.[/REGION]
 
-1. The **World Settings Panel** will open up. By default, it will be docked with the **Details Panel**, but it can be moved elsewhere. We need to set the Game Mode Override field to our "MenuGameMode" asset.
+1. The **World Settings Panel** will open up. By default, it will be docked with the **Details Panel**, but it can be moved elsewhere. We need to set the **Game Mode Override** field to our MenuGameMode asset.
 
 	[![](WorldSettings.png)](WorldSettings.png)
 
 
 ----
 
-Our custom **Game Mode** asset is now in effect on our level, configured to load our Main Menu, and use our **Player Controller** that shows the mouse cursor. If we run the game now, our Quit button will work as expected, and our New Game button will take us to an empty menu screen. Our next step will be to set up the New Game menu.
+Our custom GameMode asset is now in effect on our level, configured to load our Main Menu, and use our Player Controller that shows the mouse cursor. If we run the game now, our Quit button will work as expected, and our New Game button will take us to an empty menu screen. Our next step will be to set up the New Game menu.
 
 [OBJECT:Section]
     [PARAMLITERAL:id]
