#UE4doc #docs: Changed how images were aligned to reduce unnecessary whitespace. ##Overview removed as according to the instructional guide.

diff --git a/Engine/Documentation/Source/Resources/ContentExamples/MaterialProperties/MaterialProperties.INT.udn b/Engine/Documentation/Source/Resources/ContentExamples/MaterialProperties/MaterialProperties.INT.udn
--- a/Engine/Documentation/Source/Resources/ContentExamples/MaterialProperties/MaterialProperties.INT.udn
+++ b/Engine/Documentation/Source/Resources/ContentExamples/MaterialProperties/MaterialProperties.INT.udn
@@ -44,8 +44,6 @@ Description:An overview of the Material Properties example level.
 
 <!--Insert path when done [RELATIVE:Resources/ContentExamples/MaterialProperties]-->
 
-## Overview
-
 [REGION:fullwidth]
 
 [/REGION]
@@ -57,7 +55,7 @@ Welcome to the **Material Advanced** example level!
 ## Examples
 
 [EXCERPT:MaterialPropertiesContentEX_001]
-[REGION:topics fifth]
+[REGION:topics third]
 %Resources/ContentExamples/MaterialProperties/1_1:TopicCompact%
 %Resources/ContentExamples/MaterialProperties/1_2:TopicCompact%
 %Resources/ContentExamples/MaterialProperties/1_3:TopicCompact%
