UEDOC-1914

Replaced gif with Youtube video

diff --git a/Engine/Documentation/Source/Engine/Rendering/ParticleSystems/Reference/Modules/Orbit/OrbitModules.INT.udn b/Engine/Documentation/Source/Engine/Rendering/ParticleSystems/Reference/Modules/Orbit/OrbitModules.INT.udn
--- a/Engine/Documentation/Source/Engine/Rendering/ParticleSystems/Reference/Modules/Orbit/OrbitModules.INT.udn
+++ b/Engine/Documentation/Source/Engine/Rendering/ParticleSystems/Reference/Modules/Orbit/OrbitModules.INT.udn
@@ -10,8 +10,17 @@ Orbit modules allow for rendering particle sprites offset/rotated away from the
 
 ## Orbit
 
-![psys_orbit.gif](psys_orbit.gif)(convert:false)
-
+[OBJECT:EmbeddedVideo]
+[PARAMLITERAL:width]
+640
+[/PARAMLITERAL]
+[PARAMLITERAL:height]
+360
+[/PARAMLITERAL]
+[PARAMLITERAL:videoid]
+sdpZ2_8GUp4
+[/PARAMLITERAL]
+[/OBJECT]
 This is the only orbit module implemented. It allows for offsetting and rotating the sprite from the actual particle position. Multiple orbit modules can be chained together in a variety of ways.
 
 | Property | Description |
