UEDOC-2017

Replaced Gifs with YouTube videos

diff --git a/Engine/Documentation/Source/Resources/Showcases/MorphTargets/MorphTargets.INT.udn b/Engine/Documentation/Source/Resources/Showcases/MorphTargets/MorphTargets.INT.udn
--- a/Engine/Documentation/Source/Resources/Showcases/MorphTargets/MorphTargets.INT.udn
+++ b/Engine/Documentation/Source/Resources/Showcases/MorphTargets/MorphTargets.INT.udn
@@ -46,7 +46,17 @@ Description:An overview of the Morph Target included with UE4.
 
 
 
-![](StillLifeTimeLapse2.gif)(convert:false)
+[OBJECT:EmbeddedVideo]
+[PARAMLITERAL:width]
+640
+[/PARAMLITERAL]
+[PARAMLITERAL:height]
+360
+[/PARAMLITERAL]
+[PARAMLITERAL:videoid]
+LYeoRx5SmH4
+[/PARAMLITERAL]
+[/OBJECT]
 
 [TOC(start:2 end:2)]
 
@@ -60,7 +70,17 @@ Each of the pieces of fruit has a morph target that allows it to blend to a rott
 In this level, we use Level Blueprint to launch a Matinee sequence that animates each one of the fruit Skeletal Meshes and their corresponding Material Instances simultaneously. The flow works something like this:
 
 
-![](AppleRotFlow.gif)(convert:false)
+[OBJECT:EmbeddedVideo]
+[PARAMLITERAL:width]
+640
+[/PARAMLITERAL]
+[PARAMLITERAL:height]
+360
+[/PARAMLITERAL]
+[PARAMLITERAL:videoid]
+kdMVBXDmPtM
+[/PARAMLITERAL]
+[/OBJECT]
 
 
 
