#UE4doc #docs:
Added Engine Version field
Changed Home link to point to the overview for the tutorial
Changed font settings for multiple terms
Replaced a picture to change the indicator borders and arrows from green to yellow
Changed a few steps to include more information / change outdated information

diff --git a/Engine/Documentation/Source/Programming/Tutorials/UMG/5/UMG_5.INT.udn b/Engine/Documentation/Source/Programming/Tutorials/UMG/5/UMG_5.INT.udn
--- a/Engine/Documentation/Source/Programming/Tutorials/UMG/5/UMG_5.INT.udn
+++ b/Engine/Documentation/Source/Programming/Tutorials/UMG/5/UMG_5.INT.udn
@@ -2,6 +2,7 @@ Availability: Public
 Title:5. Build A Second Menu
 Crumbs:%ROOT%, Programming, Programming/Tutorials, Programming/Tutorials/UMG
 Description:Build a second menu with new functionality.
+version:4.9
 
 [VAR:Steps]
 [OBJECT:Navigation]
@@ -11,7 +12,7 @@ Description:Build a second menu with new functionality.
 	[PARAM:current]
 	[/PARAM]
 	[PARAM:home]
-		[Programming Tutorials Home](Programming\Tutorials\)
+		[](Programming\Tutorials\UMG)
 	[/PARAM]
 	[PARAM:next]
 		[Next Step](Programming\Tutorials\UMG\6 "%Programming\Tutorials\UMG\6:title%")
@@ -22,42 +23,48 @@ Description:Build a second menu with new functionality.
 %Steps%
 
 
-1. In the **Content Browser**, find and open the "NewGameMenu" asset we created earlier. This menu will contain a name-entry **Text Box**, a **Button** to play the game that cannot be pressed until a name is entered, and a **Button** to return to the main menu.
+1. In the **Content Browser**, find and open the NewGameMenu asset we created earlier. This menu will contain a name-entry **Text Box**, a **Button** to play the game that cannot be pressed until a name is entered, and a Button to return to the main menu.
 
-1. To create the name entry box, we'll drag a **Text Box** (not a **Text Block**) into the layout.
+1. To create the name entry box, we'll drag a Text Box (not a **Text Block**) into the layout.
 
 	[![](CreateTextEntryBox.png)](CreateTextEntryBox.png)
 	
-1.  The **Text Box** should be configured with the following values:
-	* Position is (325, 200). This leaves room for a **Text Block** placed to the left of the **Text Box**.
+1.  The Text Box should be configured with the following values:
+	* Change the name to NameTextEntry
+	* Position is (325, 200). This leaves room for a Text Block placed to the left of the Text Box.
 	* Size is 250x40.
 	* Font Size (under the "Style" heading) is 20.
 	
 	[![](TextBoxDetails.png)](TextBoxDetails.png)
 
-1. We can create the Play Game **Button** with a **Text Block** label the same way we created the **Buttons** on the previous menu.
+1. We can create the Play Game Button with a Text Block label the same way we created the Buttons on the previous menu.
+	* For the Button: Change the name to PlayGameButton, **Position** to 200, 300, **Size** to 200, 100
+	* For the Text Block: Change the name to PlayGameText, set **Visibility** to Hit Test Visible, and drag it on top of the PlayGameButton
 
-1. The Play Game **Button** will have a special feature - it will be enabled only when the name entered in the **Text Box** is not empty. We can use UMG's "bind" feature to create a new function for the "Is Enabled" field (under the Behavior section).
+1. The Play Game Button will have a special feature - it will be enabled only when the name entered in the Text Box is not empty. We can use **Unreal Motion Graphics'** (UMG) bind feature to create a new function for the **Is Enabled** field (under the **Behavior** section).
 
 	[![](PlayGameButtonDetails.png)](PlayGameButtonDetails.png)
 
-	[REGION:note]If there were complex rules for determining what constitutes a valid player name in our game, or if we needed to save the name to a C++ variable, we might expose a **UFUNCTION** on our **Game Mode** or as a static function somewhere within our project. However, since we only care that the name string is not empty, we can script it right here in our **Widget**.[/REGION]
+	[REGION:note]If there were complex rules for determining what constitutes a valid player name in our game, or if we needed to save the name to a C++ variable, we might expose a **UFUNCTION** on our **GameMode** or as a static function somewhere within our project. However, since we only care that the name string is not empty, we can script it right here in our **Widget**.[/REGION]
 
-1. To ensure that the **Button** is enabled if and only if the **Text Box** is not empty, we can convert the text from the **Text Box** to a string and then check that its length is greater than zero. Here is how that logic would appear:
+1. To ensure that the Button is enabled if and only if the Text Box is not empty, we can convert the text from the Text Box to a string and then check that its length is greater than zero. Here is how that logic would appear:
 
 	[![](PlayGameButtonEnableFunction.png)](PlayGameButtonEnableFunction.png)
 
-1. Let's add one more **Button** so we can back out and get to our Main Menu from here. This will be just like the Play Game **Button** from our Main Menu, but it will be positioned relative to the lower-right corner instead of the upper-left. To accomplish this, click the "Anchors" dropdown in the **Details Panel** for the **Button** and find the appropriate graphical representation in the pop-up menu.
+1. Let's add one more Button so we can back out and get to our Main Menu from here. This will be just like the Play Game Button from our Main Menu, but it will be positioned relative to the lower-right corner instead of the upper-left. To accomplish this, click the **Anchors** dropdown in the **Details Panel** for the Button and find the appropriate graphical representation in the pop-up menu.
+	* Change the name to MainMenuButton
+	* Set the **Position** to -400, -200
+	* Set the **Size** to 200x100
 
 	[![](SelectAnchor.png)](SelectAnchor.png)
 
 	[REGION:tip]Positioning our anchor at the lower-right corner does not change how size and position values work, so we'll need to make our position values negative in order to be on the screen. Size values will remain positive.[/REGION]
 
-1. We will now add scripting to our new **Buttons**. The Main Menu **Button** will simply reload the Main Menu **Widget**, while the Play Game **Button** will deactivate our menu entirely by having no new **Widget** provided in the call to our _ChangeMenuWidget_ function. This is shown by the phrase "Select Class" being displayed instead of the name of an actual class or asset.
+1. We will now add scripting to our new Buttons by once again adding **OnClicked** events. The Main Menu Button will simply reload the Main Menu Widget, while the Play Game Button will deactivate our menu entirely by having no new Widget provided in the call to our **ChangeMenuWidget** function. This is shown by the phrase **Select Class** being displayed instead of the name of an actual class or asset.
 
 	[![](NewGameButtonBPs.png)](NewGameButtonBPs.png)
 
-	[REGION:note]After deactivating the menu with the Play Game **Button**, we will be unable to do anything further in our game. This is the point at which we would normally load the first level, play an introductory cutscene, or spawn and possess a **Pawn**.[/REGION]
+	[REGION:note]After deactivating the menu with the Play Game Button, we will be unable to do anything further in our game. This is the point at which we would normally load the first level, play an introductory cutscene, or spawn and possess a **Pawn**.[/REGION]
 
 
 ----
