#UE4doc #docs:
Added Engine Version field
Changed Home link to point to the overview for the tutorial
Changed font settings for multiple terms

diff --git a/Engine/Documentation/Source/Programming/Tutorials/UMG/6/UMG_6.INT.udn b/Engine/Documentation/Source/Programming/Tutorials/UMG/6/UMG_6.INT.udn
--- a/Engine/Documentation/Source/Programming/Tutorials/UMG/6/UMG_6.INT.udn
+++ b/Engine/Documentation/Source/Programming/Tutorials/UMG/6/UMG_6.INT.udn
@@ -2,6 +2,7 @@ Availability: Public
 Title:6. On Your Own!
 Crumbs:%ROOT%, Programming, Programming/Tutorials, Programming/Tutorials/UMG
 Description:Ideas for further exploration!
+version:4.9
 
 [VAR:Steps]
 [OBJECT:Navigation]
@@ -11,7 +12,7 @@ Description:Ideas for further exploration!
 	[PARAM:current]
 	[/PARAM]
 	[PARAM:home]
-		[Programming Tutorials Home](Programming\Tutorials\)
+		[](Programming\Tutorials\UMG)
 	[/PARAM]
 	[PARAM:next]
 	[/PARAM]
@@ -25,14 +26,14 @@ Using what you have learned, try to do the following:
 
 * Have multiple menus open at the same time.
 * Slide or fade **Widgets** onto the screen instead of just having them appear.
-* Look into the different sound, style, and color properties that **Widgets** support.
-* Set variables in your **Game Mode** such as difficulty level, character class, and so on from within menus.
+* Look into the different sound, style, and color properties that Widgets support.
+* Set variables in your **GameMode** such as difficulty level, character class, and so on from within menus.
 * Create an in-game menu that pops up and pauses the game when the player presses a key, then resumes gameplay when it is closed.
-* Link any functionality you've built or learned from a previous tutorial to a menu built in **UMG**.
+* Link any functionality you've built or learned from a previous tutorial to a menu built in **Unreal Motion Graphics** (UMG).
 
 As for the specifics covered in this tutorial:
 
-* To learn about designing gameplay screens in UMG and using the many **Widgets** already available, read the [%Engine/UMG/UserGuide:title%](Engine/UMG/UserGuide) or the [%Engine/UMG/UserGuide/WidgetTypeReference:title%](Engine/UMG/UserGuide/WidgetTypeReference) page.
+* To learn about designing gameplay screens in UMG and using the many Widgets already available, read the [%Engine/UMG/UserGuide:title%](Engine/UMG/UserGuide) or the [%Engine/UMG/UserGuide/WidgetTypeReference:title%](Engine/UMG/UserGuide/WidgetTypeReference) page.
 * For more information about exposing variables and functions, and about using **Timers** and **Events** in your menu, try the [%Programming/Tutorials/VariablesTimersEvents:title%](Programming/Tutorials/VariablesTimersEvents) page.
 * For further tutorials, see the [%Programming/Tutorials:title%](Programming/Tutorials) page.
 
