#UE4doc #docs:
Added Engine Version field

diff --git a/Engine/Documentation/Source/Programming/Tutorials/VariablesTimersEvents/VariablesTimersEvents.INT.udn b/Engine/Documentation/Source/Programming/Tutorials/VariablesTimersEvents/VariablesTimersEvents.INT.udn
--- a/Engine/Documentation/Source/Programming/Tutorials/VariablesTimersEvents/VariablesTimersEvents.INT.udn
+++ b/Engine/Documentation/Source/Programming/Tutorials/VariablesTimersEvents/VariablesTimersEvents.INT.udn
@@ -2,6 +2,7 @@ Availability: Public
 Title:Variables, Timers, and Events
 Crumbs:%ROOT%, GettingStarted, Programming/Tutorials
 Description:Expose variables and functions to the editor, use timers, and override C++ functions with Blueprints.
+version:4.9
 
 [VAR:Topic]
 [OBJECT:Topic]
