#UE4doc #docs - adjusted bold lettering.

diff --git a/Engine/Documentation/Source/Resources/ContentExamples/Matinee/Matinee.INT.udn b/Engine/Documentation/Source/Resources/ContentExamples/Matinee/Matinee.INT.udn
--- a/Engine/Documentation/Source/Resources/ContentExamples/Matinee/Matinee.INT.udn
+++ b/Engine/Documentation/Source/Resources/ContentExamples/Matinee/Matinee.INT.udn
@@ -47,7 +47,6 @@ Related: Engine/UI/CurveEditor
 [TOC(start:2 end:2)]
 
 
-## Overview
 
 The **Matinee** Content Example map provides several examples on how to generate cutscenes or create gameplay sequences by placing **Keyframes** along a timeline. Each **Keyframe** tells the targeted **Actor** to "do something" at that point in the timeline - be it move to and arrive at a position, change color, fire off an event in a **Blueprint**, etc. 
 
@@ -65,10 +64,10 @@ Listed below are the examples provided inside the **Matinee Basics** section:
 
 | Example | What is Demonstrated |
 | ------- | ------------ |
-| **1.1 Matinee Actor** | A placed **Matinee Actor** in a level which also describes how to open the **Matinee Editor**. |
-| **1.2 Tracks and Keys** | Shows how a **Movement Track** and **Keys** can be used to move an **Actor** around in the scene. |
-| **1.3 Track Types** | An example **Matinee** that uses multiple tracks that affect an **Actor** (e.g. the **Actor** is scaled, moved, its **Material** is changed, is set invisible, and a sound is played through a **Matinee** scene). |
-| **1.4 Key Types** | How different **Key Types** can be used to affect the properties of a track that changes over time (e.g. multiple **Actors** move from point A to point B, however their transition is modified by **Key Types**). |
+| **1.1 Matinee Actor** | A placed Matinee Actor in a level which also describes how to open the **Matinee Editor**. |
+| **1.2 Tracks and Keys** | Shows how a **Movement Track** and **Keys** can be used to move an Actor around in the scene. |
+| **1.3 Track Types** | An example Matinee that uses multiple tracks that affect an Actor (e.g. the Actor is scaled, moved, its **Material** is changed, is set invisible, and a sound is played through a Matinee scene). |
+| **1.4 Key Types** | How different **Key Types** can be used to affect the properties of a track that changes over time (e.g. multiple Actors move from point A to point B, however their transition is modified by Key Types). |
 
 
 ### Matinee Groups
@@ -94,9 +93,9 @@ Listed below are the examples provided inside the **Matinee Blueprint Interactio
 
 | Example | What is Demonstrated |
 | ------- | ------------ |
-| **3.1 Matinee Blueprint Functions** | Demonstrates the standard **Matinee Blueprint Functions** that can be called inside a **Blueprint** (e.g. play, pause, stop, rewind, reverse, loop, and set position). |
-| **3.2 Matinee Events / Event Tracks** | Shows how events can be fired inside a **Blueprint** using an **Event Track** inside **Matinee** (events that are fired can be used to execute and start any section of script). |
-| **3.3 Modifying Matinee Properties** | An example of some of the properties that can be set on the **Matinee Actor** itself (e.g. adjust the play rate, play on level load, check looping or skippable, manipulate the player by disabling movement, look input, or hiding the player character). |
+| **3.1 Matinee Blueprint Functions** | Demonstrates the standard **Matinee Blueprint Functions** that can be called inside a Blueprint (e.g. play, pause, stop, rewind, reverse, loop, and set position). |
+| **3.2 Matinee Events / Event Tracks** | Shows how events can be fired inside a Blueprint using an Event Track inside Matinee (events that are fired can be used to execute and start any section of script). |
+| **3.3 Modifying Matinee Properties** | An example of some of the properties that can be set on the Matinee Actor itself (e.g. adjust the play rate, play on level load, check looping or skippable, manipulate the player by disabling movement, look input, or hiding the player character). |
 
 
 ### Matinee Advanced
@@ -107,7 +106,7 @@ Listed below are the examples provided inside the **Matinee Advanced** section:
 
 | Example | What is Demonstrated |
 | ------- | ------------ |
-| **4.1 Curve Editor** | How to get finer control over parameters by using the **Curve Editor** inside **Matinee** (e.g. fine tune the movement of an **Actor** by adjusting the curve manually). |
+| **4.1 Curve Editor** | How to get finer control over parameters by using the **Curve Editor** inside Matinee (e.g. fine tune the movement of an Actor by adjusting the curve manually). |
 | **4.2 Split Translation and Rotation** | Additional control over movement by splitting the automatic movement coordinates into **Translation** and **Rotation** that can be adjusted manually at each key. |
 
 
