#UE4doc #docs - adjusted bold lettering.

diff --git a/Engine/Documentation/Source/Resources/ContentExamples/MorphTargets/MorphTargets.INT.udn b/Engine/Documentation/Source/Resources/ContentExamples/MorphTargets/MorphTargets.INT.udn
--- a/Engine/Documentation/Source/Resources/ContentExamples/MorphTargets/MorphTargets.INT.udn
+++ b/Engine/Documentation/Source/Resources/ContentExamples/MorphTargets/MorphTargets.INT.udn
@@ -47,8 +47,8 @@ Related: Engine/Content/FBX
 
 ## Overview
 
-The Content Examples provided inside the **MorphTargets** map demonstrate how with the use of a **Skeletal Mesh** (and additional assets outlined below in example 1.2), it is possible to use a **Morph Target** to blend between objects or scale aspects of a **Skeletal Mesh**. 
-This map contains five examples ranging from a basic **Morph Target** setup, to manipulating **Materials** to match a **Morph Target** and even scaling facial features of a **Skeletal Mesh**. 
+The Content Examples provided inside the **MorphTargets** map demonstrate how with the use of a **Skeletal Mesh** (and additional assets outlined below in example 1.2), it is possible to use a **Morph Target** to blend between objects or scale aspects of a Skeletal Mesh. 
+This map contains five examples ranging from a basic Morph Target setup, to manipulating **Materials** to match a Morph Target and even scaling facial features of a Skeletal Mesh. 
 
 
 ## MorphTargets Map
@@ -59,11 +59,11 @@ Listed below are the examples provided inside the **MorphTargets** map:
 
 | Example | What is Demonstrated |
 | ------- | ------------ |
-| **1.1 Basic Morph Target** | A very basic **Morph Target** example with an animation sequence (e.g. a cube that morphs into a sphere). |
-| **1.2 Morph Target Assets** | The assets needed to use a **Morph Target** (e.g. **FBX Skeletal Mesh**, **Skeleton**, **Physics Asset**, **Animation Sequence**, and a **Material** with **Used with Morph Targets** checked in its properties). |
+| **1.1 Basic Morph Target** | A very basic Morph Target example with an animation sequence (e.g. a cube that morphs into a sphere). |
+| **1.2 Morph Target Assets** | The assets needed to use a Morph Target (e.g. **FBX Skeletal Mesh**, **Skeleton**, **Physics Asset**, **Animation Sequence**, and a Material with **Used with Morph Targets** checked in its properties). |
 | **1.3 Morphing Material Settings** | Demonstrates how to morph between materials (e.g. a cube that morphs into a sphere and changes its material). |
-| **1.4 Vertex Normals** | Description on how custom **Vertex Normals** may be used with **Morph Targets**. |
-| **1.5 Morph Targets on Skeletal Mesh** | Shows how a **Skeletal Mesh** can use a **Morph Target** to manipulate its facial features (e.g. Owen character's nose is increase/decreased in size). |
+| **1.4 Vertex Normals** | Description on how custom **Vertex Normals** may be used with Morph Targets. |
+| **1.5 Morph Targets on Skeletal Mesh** | Shows how a Skeletal Mesh can use a Morph Target to manipulate its facial features (e.g. Owen character's nose is increase/decreased in size). |
 
 
 
