#UE4doc #docs - Changed how the links are done for Related section.

diff --git a/Engine/Documentation/Source/Resources/ContentExamples/MouseInterface/MouseInterface.INT.udn b/Engine/Documentation/Source/Resources/ContentExamples/MouseInterface/MouseInterface.INT.udn
--- a/Engine/Documentation/Source/Resources/ContentExamples/MouseInterface/MouseInterface.INT.udn
+++ b/Engine/Documentation/Source/Resources/ContentExamples/MouseInterface/MouseInterface.INT.udn
@@ -2,6 +2,11 @@ Availability: Public
 Title:Mouse Interface
 Crumbs:%ROOT%, Resources
 Description:An overview and breakdown of the Mouse Interaction.
+Related: Engine/Blueprints
+Related: Engine/Blueprints\Editor
+Related: Engine/Blueprints\Editor\Modes\Components
+Related: Engine/Blueprints\Editor\Modes\Defaults
+Related: Engine/Blueprints\Editor\Modes\GraphPanel
 
 
 [VAR:Topic]
@@ -53,8 +58,6 @@ Description:An overview and breakdown of the Mouse Interaction.
 [TOC(start:2 end:2)]
 
 
-## Overview
-
 
 The Mouse Interface example demonstrates how to utilize mouse or touch inputs in your games. The level, which is pretty simple in terms of play and execution, consists of an array of spheres with a cube that tracks back and forth behind them.
 
@@ -124,11 +127,3 @@ Interaction can be broken down as follows:
 
 
 
-## Related Topics
-
-* [](Engine/Blueprints)
-* [](Engine/Blueprints\Editor)
-* [](Engine/Blueprints\Editor\Modes\Components)
-* [](Engine/Blueprints\Editor\Modes\Defaults)
-* [](Engine/Blueprints\Editor\Modes\GraphPanel)
-
