#UE4doc #docs - removed overview header and adjusted bold lettering

diff --git a/Engine/Documentation/Source/Resources/ContentExamples/MouseInterface/BlueprintActors/BlueprintActors.INT.udn b/Engine/Documentation/Source/Resources/ContentExamples/MouseInterface/BlueprintActors/BlueprintActors.INT.udn
--- a/Engine/Documentation/Source/Resources/ContentExamples/MouseInterface/BlueprintActors/BlueprintActors.INT.udn
+++ b/Engine/Documentation/Source/Resources/ContentExamples/MouseInterface/BlueprintActors/BlueprintActors.INT.udn
@@ -7,7 +7,7 @@ Description:An overview of the Blueprint Actor setup used in the Mouse Interface
 [TOC(start:2 end:2)]
 
 
-## Overview
+
 
 The Blueprint Actors (MI_Sphere) arrayed within the scene can be moved during gameplay by clicking and dragging with the mouse. 
 
@@ -25,14 +25,14 @@ The Components List for the MI_Sphere Blueprint is exceedingly simple; really, i
 
 ## Construction Script
 
-The Construction Script for MI_Sphere just applies a new Material.
+The **Construction Script** for MI_Sphere just applies a new **Material**.
 
 ![](MI_Sphere_Construction.png)
 
 
 ## Event Graph
 
-The Event Graph for the MI_Sphere Blueprint is where most of the actual work is done. This is where the mouse over, click, and touch events are processed, as well as where the dragging behavior is established. At a high level, the network does the following:
+The **Event Graph** for the MI_Sphere Blueprint is where most of the actual work is done. This is where the mouse over, click, and touch events are processed, as well as where the dragging behavior is established. At a high level, the network does the following:
 
 * Swaps out the Materials applied to the object based on mouse over events.
 * Tells the object to start following the motion of the mouse cursor during a drag operation. This movement is calculated every tick.
@@ -43,7 +43,7 @@ The Event Graph for the MI_Sphere Blueprint is where most of the actual work is
 [/REGION]
 **Click for full size**
 
-Of particular importance in the network are the following Events, which drive the mouse interaction:
+Of particular importance in the network are the following **Events**, which drive the mouse interaction:
 
 * **OnBeginCursorOver** - This is triggered whenever the mouse cursor moves over a specific object. In this case, the sphere Static Mesh, which is one of the Blueprint's Components. 
 * **OnEndCursorOver** - This gets called once the cursor moves off the indicated object, which, once again, is the Static Mesh Component. 
