#UE4doc #docs - Adjusted Bold Lettering and removed header.

diff --git a/Engine/Documentation/Source/Resources/ContentExamples/MouseInterface/LevelBlueprint/LevelBlueprint.INT.udn b/Engine/Documentation/Source/Resources/ContentExamples/MouseInterface/LevelBlueprint/LevelBlueprint.INT.udn
--- a/Engine/Documentation/Source/Resources/ContentExamples/MouseInterface/LevelBlueprint/LevelBlueprint.INT.udn
+++ b/Engine/Documentation/Source/Resources/ContentExamples/MouseInterface/LevelBlueprint/LevelBlueprint.INT.udn
@@ -6,10 +6,8 @@ Description:An overview of the Level Blueprint setup used in the Mouse Interface
 
 [TOC(start:2 end:2)]
 
-## Overview
 
-
-You will notice that the demolet uses both Level Blueprint and Actor Blueprints scripting for the behavior setup. This is intentionally used as an example, showing how either application can make use of mouse interaction. Level Blueprint handles the moving box while Actor Blueprints are used to create the movable spheres.
+You will notice that the demolet uses both **Level Blueprint** and **Actor Blueprints** scripting for the behavior setup. This is intentionally used as an example, showing how either application can make use of mouse interaction. Level Blueprint handles the moving box while Actor Blueprints are used to create the movable spheres.
 
 The Level Script employed in this level shows how mouse interaction can be set up by allowing the player to click on a moving box. Each click either starts or stops the animation of the box as it slides back and forth across the level. 
 
@@ -28,7 +26,7 @@ You can see in the Level Blueprint a section with an **OnBeginCursorOver** and a
 
 * **OnEndCursorOver** - This gets called once the cursor moves off the indicated object, which, once again, is the Editor Cube. 
 
-These events can be created for a given object by selecting the object in the Viewport or World Outliner, then, in the Blueprint editor window, **right-click** in the Graph and choose **Add Event for YourObject > Input > Mouse Input > Add OnBeginCursorOver** or **Add OnEndCursorOver**. 
+These events can be created for a given object by selecting the object in the **Viewport** or **World Outliner**, then, in the Blueprint editor window, **right-click** in the Graph and choose **Add Event for YourObject > Input > Mouse Input > Add OnBeginCursorOver** or **Add OnEndCursorOver**. 
 
 
 ## Mouse Click Setup
@@ -37,7 +35,7 @@ The Level Blueprint also includes a setup for handling mouse clicks.
 
 ![](ClickSetup_LeveScript.png)
 
-Generally speaking, the network here simply alternates between starting and stopping the Editor Cube's Matinee by way of a FlipFlop flow control node. Note that there are two separate events that drive the network. By having both of these events, the network can be called from a mouse on a computer or by a finger touch on a mobile device.
+Generally speaking, the network here simply alternates between starting and stopping the Editor Cube's Matinee by way of a **FlipFlop flow control** node. Note that there are two separate events that drive the network. By having both of these events, the network can be called from a mouse on a computer or by a finger touch on a mobile device.
 
 * **OnInputTouchBegin** - This gets called on a Touch event on the specified object, which happens when a player touches the screen on a mobile device.
 * **OnClicked** - This event gets called when the player clicks the mouse on the specified object. 
