#UE4doc #docs - Adjusted Bold Lettering and removed header.

diff --git a/Engine/Documentation/Source/Resources/ContentExamples/MouseInterface/MouseControlSetup/MouseControlSetup.INT.udn b/Engine/Documentation/Source/Resources/ContentExamples/MouseInterface/MouseControlSetup/MouseControlSetup.INT.udn
--- a/Engine/Documentation/Source/Resources/ContentExamples/MouseInterface/MouseControlSetup/MouseControlSetup.INT.udn
+++ b/Engine/Documentation/Source/Resources/ContentExamples/MouseInterface/MouseControlSetup/MouseControlSetup.INT.udn
@@ -5,7 +5,6 @@ Description:How to set up a game for mouse control
 
 [TOC(start:2 end:2)]
 
-## Overview
 
 To make a game where you use the mouse cursor possible, there are a few things you must do when getting started. At a high level, the steps in Blueprint are as follows. We will cover them in greater detail below:
 
@@ -29,7 +28,7 @@ By default, the PlayerController class does not have mouse interfacing enabled.
 
 	[REGION:raw]![](CreateNewBlueprint.png)[/REGION]
 
-1. In the **Pick Parent Class** window, use the Search bar to find the PlayerController class and click on it.
+1. In the **Pick Parent Class** window, use the Search bar to find the **PlayerController** class and **click on it**.
 
 	![](PickParentClass_PlayerController.png)
 
@@ -45,7 +44,7 @@ By default, the PlayerController class does not have mouse interfacing enabled.
 
 	![](MouseInterfaceProperties.png)
 
-1. Save your new Blueprint.
+1. **Save** your new Blueprint.
 
 
 ## Creating a Custom GameMode
@@ -56,7 +55,7 @@ Now that we have a PlayerController that has Mouse Input enabled, our next step
 
 	[REGION:raw]![](CreateNewBlueprint.png)[/REGION]
 
-1. In the **Pick Parent Class** window, use the Search bar to find the GameMode class and click on it.
+1. In the **Pick Parent Class** window, use the Search bar to find the **GameMode** class and **click on it**.
 
 	![](GameInfoParentClass.png)
 
@@ -68,18 +67,18 @@ Now that we have a PlayerController that has Mouse Input enabled, our next step
 
 	![](EditDefaultsBlueprint2.png)
 
-1. In the Classes category, click on the dropdown next to **Player Controller Class**. You should see the **Mouse_PlayerController** class created previously. Choose it from the list.
+1. In the Classes category, **click on the dropdown** next to **Player Controller Class**. You should see the **Mouse_PlayerController** class created previously. Choose it from the list.
 
 	![](SelectingPlayerControllerClass.png)
 
-1. Save your new Blueprint.
+1. **Save** your new Blueprint.
 
 
 ## Setting World Info Properties for Your Custom Game Mode
 
 Now that we have our custom PlayerController and GameMode Blueprints completed, we just need to tell our World Settings to use them.
 
-1. Click the **Settings** button located in the toolbar. Then from the drop down select **World Settings**. This opens up the World Settings in the **Details** panel.
+1. **Click** the **Settings** button located in the toolbar. Then from the drop down **select** **World Settings**. This opens up the World Settings in the **Details** panel.
 
 	![](WorldSettings.png)
 
