#UE4doc #docs - removed overview header and bolded link.

diff --git a/Engine/Documentation/Source/Resources/ContentExamples/MouseInterface/SettingsAndCursors/SettingsAndCursors.INT.udn b/Engine/Documentation/Source/Resources/ContentExamples/MouseInterface/SettingsAndCursors/SettingsAndCursors.INT.udn
--- a/Engine/Documentation/Source/Resources/ContentExamples/MouseInterface/SettingsAndCursors/SettingsAndCursors.INT.udn
+++ b/Engine/Documentation/Source/Resources/ContentExamples/MouseInterface/SettingsAndCursors/SettingsAndCursors.INT.udn
@@ -7,7 +7,6 @@ Description:A breakdown of the properties and cursors used for mouse interaction
 [TOC(start:2 end:2)]
 
 
-## Overview
 
 The following is a breakdown of the available properties used when setting up mouse interaction, as well as a sample of the available cursors. Be aware that cursors shown will change based on your operating system.
 
@@ -22,7 +21,7 @@ Within the PlayerController class, you will find a category of properties titled
 | **Show Mouse Cursor** | This determines whether or not the mouse cursor will be visible. Note that this is pulling directly from the operating system's cursor by default. |
 | **Enable Click Events** | This activates the ability for click events to be processed. Without this, you cannot click on things. |
 | **Enable Mouse Over Events** | This activates the ability for mouse-over events to be processed. Without this, mouse-overs do not work. |
-| **Default Mouse Cursor** | This sets what the mouse cursor will look like by default. The options are shown in the [Available Cursors](#AvailableCursors) section.  |
+| **Default Mouse Cursor** | This sets what the mouse cursor will look like by default. The options are shown in the **[Available Cursors](#AvailableCursors)** section.  |
 | **Default Click Trace Channel** | This sets the default channel that will be used by the trace system to see what you have clicked on. The default setting is _Visibiltiy_. This means the trace is checking visibility status of the object and will only return _true_ if the object is found to be visible. Put another way, you can only interact with objects you can _see_ in the game.  |
 
 
