#UE4doc #docs - Adjusted Bold Lettering and removed header.

diff --git a/Engine/Documentation/Source/Resources/ContentExamples/NavMesh/NavMesh.INT.udn b/Engine/Documentation/Source/Resources/ContentExamples/NavMesh/NavMesh.INT.udn
--- a/Engine/Documentation/Source/Resources/ContentExamples/NavMesh/NavMesh.INT.udn
+++ b/Engine/Documentation/Source/Resources/ContentExamples/NavMesh/NavMesh.INT.udn
@@ -45,9 +45,8 @@ Related: Engine/Actors/Volumes
 [TOC(start:2 end:2)]
 
 
-## Overview
 
-The **Navmesh** Content Example map provides three examples of how to generate and use a **NavMesh** to enable **Pawns** the ability to pathfind their way through obstacles, over ramps or to jump off ledges. While in the editor, and with a **NavMeshBoundsVolume** placed in the level, pressing the **P** key will show/hide the area that the **NavMesh** covers. The final example in this map shows how the coverage area of a **NavMesh** is automatically updated when moving **Actors** around in the level (this can be disabled via the **Editor Settings** under **Level Editor/Miscellaneous** and unchecking the **Update Navigation Automatically** checkbox). 
+The **Navmesh** Content Example map provides three examples of how to generate and use a **NavMesh** to enable **Pawns** the ability to pathfind their way through obstacles, over ramps or to jump off ledges. While in the editor, and with a **NavMeshBoundsVolume** placed in the level, pressing the **P** key will show/hide the area that the NavMesh covers. The final example in this map shows how the coverage area of a NavMesh is automatically updated when moving **Actors** around in the level (this can be disabled via the **Editor Settings** under **Level Editor/Miscellaneous** and unchecking the **Update Navigation Automatically** checkbox). 
 
 [region:note]
 You can find additional resources for each example by clicking on the example's name.
@@ -62,9 +61,9 @@ Listed below are the examples provided inside the **Navmesh** map:
 
 | Example | What is Demonstrated |
 | ------- | ------------ |
-| **[1.1 Simple navigation around obstacles](Resources\ContentExamples\NavMesh\1_1)** | How to get a **Pawn** to navigate through a set of obstacles. |
-| **[1.2 Navigation using a Nav Link Proxy](Resources\ContentExamples\NavMesh\1_2)** | Shows a **Pawn** navigating up to and jumping off a ramp. |
-| **[1.3 Navmesh rebuild after moving Actors](Resources\ContentExamples\NavMesh\1_3)** | This example demonstrates how the **NavMesh** automatically updates when **Actors** like **Static Meshes** are moved in the editor. |
+| **[1.1 Simple navigation around obstacles](Resources\ContentExamples\NavMesh\1_1)** | How to get a Pawn to navigate through a set of obstacles. |
+| **[1.2 Navigation using a Nav Link Proxy](Resources\ContentExamples\NavMesh\1_2)** | Shows a Pawn navigating up to and jumping off a ramp. |
+| **[1.3 Navmesh rebuild after moving Actors](Resources\ContentExamples\NavMesh\1_3)** | This example demonstrates how the NavMesh automatically updates when Actors like Static Meshes are moved in the editor. |
 
 
 
