#UE4doc #docs - Adjusted Bold Lettering

diff --git a/Engine/Documentation/Source/Resources/ContentExamples/NavMesh/1_1/1_1.INT.udn b/Engine/Documentation/Source/Resources/ContentExamples/NavMesh/1_1/1_1.INT.udn
--- a/Engine/Documentation/Source/Resources/ContentExamples/NavMesh/1_1/1_1.INT.udn
+++ b/Engine/Documentation/Source/Resources/ContentExamples/NavMesh/1_1/1_1.INT.udn
@@ -50,7 +50,7 @@ Related: Engine/Actors/Volumes
 ![](Resources\ContentExamples\NavMesh\NavMesh_banner.png)
 [/REGION]
 
-This section of the **NavMesh** content example shows how a Pawn can navigate through a set of obstacles. You can try adding and moving more blocks around to make more complicated paths. The navigation system will often list things as "simple" because it is not taking in much external data, like if an Actor is walking on ice or through tar, but can quite handily navigate in 3 dimensions, accounting for bridges and overhangs along the way, to any desired location.
+This section of the **NavMesh** content example shows how a **Pawn** can navigate through a set of obstacles. You can try adding and moving more blocks around to make more complicated paths. The navigation system will often list things as "simple" because it is not taking in much external data, like if an Actor is walking on ice or through tar, but can quite handily navigate in 3 dimensions, accounting for bridges and overhangs along the way, to any desired location.
 
 [REGION:lightbox]
 [![](NM_1_1_001.png)(w:300)](NM_1_1_001.png)
