#UE4doc #docs - Adjusted Bold Lettering

diff --git a/Engine/Documentation/Source/Resources/ContentExamples/NavMesh/1_2/1_2.INT.udn b/Engine/Documentation/Source/Resources/ContentExamples/NavMesh/1_2/1_2.INT.udn
--- a/Engine/Documentation/Source/Resources/ContentExamples/NavMesh/1_2/1_2.INT.udn
+++ b/Engine/Documentation/Source/Resources/ContentExamples/NavMesh/1_2/1_2.INT.udn
@@ -50,7 +50,7 @@ Related: Engine/Actors/Volumes
 ![](Resources\ContentExamples\NavMesh\NavMesh_banner.png)
 [/REGION]
 
-**NavLinkProxy** Actors are used to specify where Pawns can jump or drop off ledges, allowing them to temporarily leave the **NavMesh** in order to cross gaps in it. A start point and end point can be positioned, and will be treated as path nodes. In this example, the Pawn moves to the start point of the **NavLinkProxy**, then drops off the ramp to reach the end point, where it resumes navigating on the **NavMesh** towards its destination. This **NavLinkProxy** is set up to be one-way, so when the Pawn turns around to go back, it navigates around the ramp instead of trying to use the NavLinkProxy again.
+**NavLinkProxy** Actors are used to specify where Pawns can jump or drop off ledges, allowing them to temporarily leave the **NavMesh** in order to cross gaps in it. A start point and end point can be positioned, and will be treated as path nodes. In this example, the Pawn moves to the start point of the NavLinkProxy, then drops off the ramp to reach the end point, where it resumes navigating on the NavMesh towards its destination. This NavLinkProxy is set up to be one-way, so when the Pawn turns around to go back, it navigates around the ramp instead of trying to use the NavLinkProxy again.
 
 [REGION:lightbox]
 [![](NavMesh_1_2_001.png)(w:225)](NavMesh_1_2_001.png)
