#UE4doc #docs - Adjusted Bold Lettering

diff --git a/Engine/Documentation/Source/Resources/ContentExamples/NavMesh/1_3/1_3.INT.udn b/Engine/Documentation/Source/Resources/ContentExamples/NavMesh/1_3/1_3.INT.udn
--- a/Engine/Documentation/Source/Resources/ContentExamples/NavMesh/1_3/1_3.INT.udn
+++ b/Engine/Documentation/Source/Resources/ContentExamples/NavMesh/1_3/1_3.INT.udn
@@ -50,7 +50,7 @@ Related: Engine/Actors/Volumes
 ![](Resources\ContentExamples\NavMesh\NavMesh_banner.png)
 [/REGION]
 
-The **NavMesh** automatically updates when Actors like **Static Meshes** are moved in the editor. In this example, you can move the block **Static Mesh** and position it so that it bridges the gap between the ramps. The **NavMesh** will automatically update and build across the newly-formed bridge, allowing the Pawn to navigate across it when you **Play** or **Simulate**.
+The **NavMesh** automatically updates when Actors like **Static Meshes** are moved in the editor. In this example, you can move the block Static Mesh and position it so that it bridges the gap between the ramps. The NavMesh will automatically update and build across the newly-formed bridge, allowing the Pawn to navigate across it when you **Play** or **Simulate**.
 
 [REGION:lightbox]
 [![](NavMesh_1_3_001.png)(w:225)](NavMesh_1_3_001.png)
