#UE4doc #docs - Adjusted Bold Lettering and removed header.

diff --git a/Engine/Documentation/Source/Resources/ContentExamples/Networking/NetworkFeatures.INT.udn b/Engine/Documentation/Source/Resources/ContentExamples/Networking/NetworkFeatures.INT.udn
--- a/Engine/Documentation/Source/Resources/ContentExamples/Networking/NetworkFeatures.INT.udn
+++ b/Engine/Documentation/Source/Resources/ContentExamples/Networking/NetworkFeatures.INT.udn
@@ -46,7 +46,7 @@ Related: Engine/UI/LevelEditor/InEditorTesting/NetworkedPlayInEditor
 [TOC(start:2 end:3)]
 
 
-## Overview
+
 The **Network_Features** map inside the **Content Examples** project illustrates how to implement network functionality entirely through Blueprints. It provides key examples in **Replication** including the **Replicates** boolean which can be set on Actors, detecting **Network Authority** and executing different logic based on that authority, usage of **Variable Replication** and how it is used during gameplay, variable replication using **RepNotify** to execute logic based on updates made to a variable that is replicated, as well as the use of **Function Replication** to handle non-critical gameplay elements or "one-off" events.
 
 The other section in the map provides insight on **Relevancy** and how to properly handle communicating information to and from the Server and Clients. Often times, it is not critical that every bit of information be passed from the server to every client (e.g. a player that is on one side of the map does not need to know about another player opening a chest on the other side of the map). There are situations when this would become relevant to that player however (e.g. if that player moved from the other side of the map to the chest that was opened), the Content Examples in this map will show how to handle situations like this as well as how to handle a player that joins an in progress session. 
@@ -62,11 +62,11 @@ The examples in the **Replication** section showcase the methods of properly com
 
  Example | What is Demonstrated |
 | ------- | ------------ |
-| **[1.1 Actor Replication](Resources\ContentExamples\Networking\1_1)** | An example of two Actors where one is **Replicated** and the other is not (e.g. an Actor that shows up on both Client/Server and one that only shows up on Server). |
-| **[1.2 Detecting Network Authority and Remote Clients in Blueprints](Resources\ContentExamples\Networking\1_2)** | An example of a Blueprint executing different logic based on **Network Authority** (e.g. Actor is displayed differently based whether its viewed on the Server or the Client). |
-| **[1.3 Variable Replication](Resources\ContentExamples\Networking\1_3)** | An example of **Replicating** a Health variable from Server to Clients (e.g. an Actor's Health is displayed and updated on a Client as the Server updates the variable). |
-| **[1.4 Variable Replication (RepNotify)](Resources\ContentExamples\Networking\1_4)** | An example of a variable that is marked as **RepNotify** which automatically updates on both Server and Client machines (e.g. A streetlight uses a **RepNotify** variable to ensure both Client and Server are viewing the same thing).  |
-| **[1.5 Function Replication (RPCs)](Resources\ContentExamples\Networking\1_5)** | An example of a **Replicated Function** is used to communicate when a button is pressed by the Server or Client. |
+| **[1.1 Actor Replication](Resources\ContentExamples\Networking\1_1)** | An example of two Actors where one is Replicated and the other is not (e.g. an Actor that shows up on both Client/Server and one that only shows up on Server). |
+| **[1.2 Detecting Network Authority and Remote Clients in Blueprints](Resources\ContentExamples\Networking\1_2)** | An example of a Blueprint executing different logic based on Network Authority (e.g. Actor is displayed differently based whether its viewed on the Server or the Client). |
+| **[1.3 Variable Replication](Resources\ContentExamples\Networking\1_3)** | An example of Replicating a Health variable from Server to Clients (e.g. an Actor's Health is displayed and updated on a Client as the Server updates the variable). |
+| **[1.4 Variable Replication (RepNotify)](Resources\ContentExamples\Networking\1_4)** | An example of a variable that is marked as RepNotify which automatically updates on both Server and Client machines (e.g. A streetlight uses a RepNotify variable to ensure both Client and Server are viewing the same thing).  |
+| **[1.5 Function Replication (RPCs)](Resources\ContentExamples\Networking\1_5)** | An example of a Replicated Function is used to communicate when a button is pressed by the Server or Client. |
 
 
 ### Relevancy
