#UE4doc #docs - Removed header.

diff --git a/Engine/Documentation/Source/Resources/ContentExamples/Networking/1_1/1_1.INT.udn b/Engine/Documentation/Source/Resources/ContentExamples/Networking/1_1/1_1.INT.udn
--- a/Engine/Documentation/Source/Resources/ContentExamples/Networking/1_1/1_1.INT.udn
+++ b/Engine/Documentation/Source/Resources/ContentExamples/Networking/1_1/1_1.INT.udn
@@ -37,7 +37,6 @@ In order to view both **Server** and **Client** perspectives, on the **Play** dr
 
 
 
-## Overview
 
 This example demonstrates the use of the **Replicates** boolean on Actors. When enabled, it allows clients to be aware of Actors spawned on the Server, from a networking perspective. 
 
