#UE4doc #docs - Adjusted Bold Lettering and removed header.

diff --git a/Engine/Documentation/Source/Resources/ContentExamples/Networking/1_2/1_2.INT.udn b/Engine/Documentation/Source/Resources/ContentExamples/Networking/1_2/1_2.INT.udn
--- a/Engine/Documentation/Source/Resources/ContentExamples/Networking/1_2/1_2.INT.udn
+++ b/Engine/Documentation/Source/Resources/ContentExamples/Networking/1_2/1_2.INT.udn
@@ -36,7 +36,6 @@ In order to view both **Server** and **Client** perspectives, on the **Play** dr
 [/REGION]
 
 
-## Overview
 
 This example demonstrates a Blueprint executing different logic depending on whether the Blueprint is running on the Network Authority (Server) or a Remote Machine (Client). 
 
@@ -49,7 +48,7 @@ The Text Render component is then updated based on whether the script is being e
 	[![](Resources/ContentExamples/Networking/1_2/1_2_Blueprint.png)(w:640)](Resources/ContentExamples/Networking/1_2/1_2_Blueprint.png)
 [/REGION]
 
-Here is the result of what the Server and Client will see by executing this script with the **Switch Has Authority** node. 
+Here is the result of what the Server and Client will see by executing this script with the Switch Has Authority node. 
 
 [REGION:fullwidth]
 ![](1_2_Play.png)
