#UE4doc #docs - Removed header.

diff --git a/Engine/Documentation/Source/Resources/ContentExamples/Networking/1_3/1_3.INT.udn b/Engine/Documentation/Source/Resources/ContentExamples/Networking/1_3/1_3.INT.udn
--- a/Engine/Documentation/Source/Resources/ContentExamples/Networking/1_3/1_3.INT.udn
+++ b/Engine/Documentation/Source/Resources/ContentExamples/Networking/1_3/1_3.INT.udn
@@ -36,7 +36,7 @@ In order to view both **Server** and **Client** perspectives, on the **Play** dr
 [/REGION]
 
 
-## Overview
+
 Variables that are important to gameplay should **ONLY** be modified on the Network Authority (Server), and then replicated to Remote Machines (Clients) on a need-to-know basis. 
 
 In this example, the ghost's health is being modified on a timer exclusively on the Network Authority. In the non-replicated ghost, the remote machines never receive updated health values. In the replicated ghost, they receive updates as the Server changes the health. 
