#UE4doc #docs - Adjusted Bold Lettering and removed header.

diff --git a/Engine/Documentation/Source/Resources/ContentExamples/Networking/1_4/1_4.INT.udn b/Engine/Documentation/Source/Resources/ContentExamples/Networking/1_4/1_4.INT.udn
--- a/Engine/Documentation/Source/Resources/ContentExamples/Networking/1_4/1_4.INT.udn
+++ b/Engine/Documentation/Source/Resources/ContentExamples/Networking/1_4/1_4.INT.udn
@@ -35,13 +35,13 @@ In order to view both **Server** and **Client** perspectives, on the **Play** dr
 [/REGION]
 
 
-## Overview
+
 Oftentimes, it is desired to execute some logic in response to the value of a replicated variable changing. Variables that are marked **RepNotify**, have a special function called automatically any time the value is changed, on both Network Authority and Remote Machines. 
 
-In this example, the Network Authority controls changing a scalar value indicating which color the streetlight should display using a **RepNotify** variable. When the value is changed, the auto-generated OnRep function is called, which in this example then updates a material parameter. 
+In this example, the Network Authority controls changing a scalar value indicating which color the streetlight should display using a RepNotify variable. When the value is changed, the auto-generated OnRep function is called, which in this example then updates a material parameter. 
 
 ## Blueprint Breakdown
-In the **BP_Stoplight_Variable_RepNotify** Blueprint's Event Graph (pictured below), you will see that the script calls for only the Network Authority to handle the light switching logic. The **Streetlight Scalar** variable representing the state of the light is specifically marked as **RepNotify**, so both the Client and Server have the **OnRep_StreetlightScalar** function called when the value changes. 
+In the **BP_Stoplight_Variable_RepNotify** Blueprint's Event Graph (pictured below), you will see that the script calls for only the Network Authority to handle the light switching logic. The **Streetlight Scalar** variable representing the state of the light is specifically marked as RepNotify, so both the Client and Server have the **OnRep_StreetlightScalar** function called when the value changes. 
 
 **Click for full view**
 [REGION:lightbox]
@@ -50,7 +50,7 @@ In the **BP_Stoplight_Variable_RepNotify** Blueprint's Event Graph (pictured bel
 
 The **OnRep_StreetlightScalar** function then handles updating the visuals of the streetlight such that it works on all machines. 
 
-Setting a variable to **RepNotify** can be done through the **Details** panel of the variable under the **Replication** drop down menu. 
+Setting a variable to RepNotify can be done through the **Details** panel of the variable under the **Replication** drop down menu. 
 
 ![](1_4_RepNotify.png)
 
