#UE4doc #docs - Removed Overview Header.

diff --git a/Engine/Documentation/Source/Resources/ContentExamples/Networking/1_5/1_5.INT.udn b/Engine/Documentation/Source/Resources/ContentExamples/Networking/1_5/1_5.INT.udn
--- a/Engine/Documentation/Source/Resources/ContentExamples/Networking/1_5/1_5.INT.udn
+++ b/Engine/Documentation/Source/Resources/ContentExamples/Networking/1_5/1_5.INT.udn
@@ -35,7 +35,7 @@ In order to view both **Server** and **Client** perspectives, on the **Play** dr
 [/REGION]
 
 
-## Overview
+
 In addition to variables, function calls can also be replicated. **Replicated Function Calls** can be set to either **Reliable** or **Unreliable**. Reliable calls are guaranteed to occur, while Unreliable calls may be dropped during heavy traffic. Most replicated functions that handle cosmetic visuals are unreliable to avoid network saturation. 
 
 There are 3 primary types of replicated functions: **Multicast**, **Server**, and **Client**. Multicast functions should be called on the server, where they are executed, and then forwarded automatically to clients. Server functions are called by a client and then only executed on the server. Client functions are called by the server and then only execute on the owning client. 
diff --git a/Engine/Documentation/Source/Resources/ContentExamples/Networking/2_1/2_1.INT.udn b/Engine/Documentation/Source/Resources/ContentExamples/Networking/2_1/2_1.INT.udn
--- a/Engine/Documentation/Source/Resources/ContentExamples/Networking/2_1/2_1.INT.udn
+++ b/Engine/Documentation/Source/Resources/ContentExamples/Networking/2_1/2_1.INT.udn
@@ -35,7 +35,6 @@ In order to view both **Server** and **Client** perspectives, on the **Play** dr
 [/REGION]
 
 
-## Overview
 In order to create properly networked gameplay through Blueprints, it is important to also understand the concept of **Network Relevancy**. 
 Simply put, it often does not make sense for a client machine to have to receive network data for every Actor, all of the time. A player on one side of a large map does not need to know everything occuring on the other side of the map from a networking perspective. 
 In that case, we say only some Actors are "relevant" to the player at a given time. Relevant Actors are usually the ones players can see, interact with, or are nearby. The player receives full network updates only for relevant Actors, as an optimization. 
