#UE4doc #docs - Removed overview header.

diff --git a/Engine/Documentation/Source/Resources/ContentExamples/Networking/2_2/2_2.INT.udn b/Engine/Documentation/Source/Resources/ContentExamples/Networking/2_2/2_2.INT.udn
--- a/Engine/Documentation/Source/Resources/ContentExamples/Networking/2_2/2_2.INT.udn
+++ b/Engine/Documentation/Source/Resources/ContentExamples/Networking/2_2/2_2.INT.udn
@@ -36,7 +36,7 @@ In order to view both **Server** and **Client** perspectives, on the **Play** dr
 [/REGION]
 
 
-## Overview
+
 This example attempts to network the chest's behavior through the use of a Replicated Multicast Function. The idea is that the chest will call the multicast when it is interacted with, which will handle playing the visuals of the chest opening for all of the clients. 
 
 However, to demonstrate relevancy issues, this chest has been modified to become network irrelevant at a very short distance. Stand the client outside of the blue ring and the chest will be irrelevant to the client, even though it is visible. Now move the server to the chest to activate it and observe what happens on the client. The client does not see anything change at all! This is expected, as the chest is not relevant to the client yet. Move the client into the blue ring, causing the chest to be come relevant. Now what happens? Still nothing, which is not what is desired. 
