#UE4doc #docs - Removed header.

diff --git a/Engine/Documentation/Source/Resources/ContentExamples/Networking/2_3/2_3.INT.udn b/Engine/Documentation/Source/Resources/ContentExamples/Networking/2_3/2_3.INT.udn
--- a/Engine/Documentation/Source/Resources/ContentExamples/Networking/2_3/2_3.INT.udn
+++ b/Engine/Documentation/Source/Resources/ContentExamples/Networking/2_3/2_3.INT.udn
@@ -36,10 +36,10 @@ In order to view both **Server** and **Client** perspectives, on the **Play** dr
 [/REGION]
 
 
-## Overview
+
 Example 2.2 failed because the replicated Multicast function did not handle preserving the state of the chest past the time it was executed. In that case, a natural solution might be to try to use a replicated variable to solve this problem instead. 
 
-In this example, a RepNotify boolean is used to save whether the chest is open or not. Again, stand the client out of the blue ring and have the server activate the chest. What happens when the client walks into the right afterwards? The client sees the chest open, but also sees the gold effects play, even though that happened awhile ago on the server!
+In this example, a **RepNotify** boolean is used to save whether the chest is open or not. Again, stand the client out of the blue ring and have the server activate the chest. What happens when the client walks into the right afterwards? The client sees the chest open, but also sees the gold effects play, even though that happened awhile ago on the server!
 
 When the chest becomes relevant to the client again, the client gets the change to the replicated variable and the corresponding RepNotify function is called. Unfortunately, the client has no way of knowing that the variable actually changed well in the past, only that it did change, which is insufficient to handle this particular case. 
 
