#UE4doc #docs - Removed header

diff --git a/Engine/Documentation/Source/Resources/ContentExamples/Networking/2_4/2_4.INT.udn b/Engine/Documentation/Source/Resources/ContentExamples/Networking/2_4/2_4.INT.udn
--- a/Engine/Documentation/Source/Resources/ContentExamples/Networking/2_4/2_4.INT.udn
+++ b/Engine/Documentation/Source/Resources/ContentExamples/Networking/2_4/2_4.INT.udn
@@ -36,7 +36,7 @@ In order to view both **Server** and **Client** perspectives, on the **Play** dr
 [/REGION]
 
 
-## Overview
+
 Both examples 2.2 and 2.3 failed beacuase neither could handle the two cases the chest examples needs: 
 
 * To play an instanteous effect at the time of interaction. 
