#UE4doc #docs - Gif needs to be converted to video

diff --git a/Engine/Documentation/Source/Engine/Physics/Vehicles/DoubleWishboneVehicle/DoubleWishboneVehicle.INT.udn b/Engine/Documentation/Source/Engine/Physics/Vehicles/DoubleWishboneVehicle/DoubleWishboneVehicle.INT.udn
--- a/Engine/Documentation/Source/Engine/Physics/Vehicles/DoubleWishboneVehicle/DoubleWishboneVehicle.INT.udn
+++ b/Engine/Documentation/Source/Engine/Physics/Vehicles/DoubleWishboneVehicle/DoubleWishboneVehicle.INT.udn
@@ -12,31 +12,31 @@ This project and tutorial were provided by community member Xenome, and are publ
 
 Thank you, Xenome!
 
-## Overview
+_________________________
 
 This tutorial will walk you through the details of the advanced vehicle template. When we are done, you should have a pretty good idea of how it works in theory, and how to build and configure something similar yourself. The main focus is going to be on the suspension.
 
-All the FBX assets used in the vehicle template are available [here for download](DoubleWishboneVehicle_Files.zip). You should be able to analyze the originals and experiment with your own designs in any modeling package that is capable of reading and writing FBX files.
+All the FBX assets used in the vehicle template are available **[here for download](DoubleWishboneVehicle_Files.zip)**. You should be able to analyze the originals and experiment with your own designs in any modeling package that is capable of reading and writing FBX files.
 
-Before we get started, you should be aware that a double wishbone suspension is quite a bit more complex to set up than the standard Unreal Engine 4 vehicle, so if you have not successfully set one of those up before, you are strongly encouraged to do that first. There are some excellent videos taking you through the standard procedure over at Epic's YouTube page [links here](https://wiki.unrealengine.com/Introduction_to_Vehicles_in_UE4_Playlist).
+Before we get started, you should be aware that a double wishbone suspension is quite a bit more complex to set up than the standard Unreal Engine 4 vehicle, so if you have not successfully set one of those up before, you are strongly encouraged to do that first. There are some excellent videos taking you through the standard procedure over at Epic's YouTube page **[links here](https://wiki.unrealengine.com/Introduction_to_Vehicles_in_UE4_Playlist)**.
 
 Once you have watched those videos and have successfully setup a base vehicle, you will be ready to look at the double wishbone design discussed here. We will start with some fundamental theory.
 
 ## Simulated Wheels Versus Real Suspension Linkages
 
-When we set up the vehicle movement component and assign wheels to it, the vertical movement of the wheels is typically updated by a **WheelHandler** node inside the vehicle's animation Blueprint. On each frame, the handler will move the wheels up and down along the local Z axis to simulate suspension, making sure that the wheels touch the ground whenever that is possible.
+When we set up the vehicle movement component and assign wheels to it, the vertical movement of the wheels is typically updated by a WheelHandler node inside the vehicle's animation Blueprint. On each frame, the handler will move the wheels up and down along the local Z axis to simulate suspension, making sure that the wheels touch the ground whenever that is possible.
 
 Furthermore, the **WheelHandler** also applies rotation to make the wheels spin around their Y axis, driven by the rotation speed of the simulated motor and gearbox. In addition to this, it will turn the wheels around their Z axis, driven by the current steering angle.
 
-The default vertical suspension movement is clamped by the Suspension Max Raise and Drop settings in your VehicleWheel-derived classes. This movement is entirely linear, as illustrated below:
+The default vertical suspension movement is clamped by the **Suspension Max Raise** and **Drop** settings in your **VehicleWheel-derived** classes. This movement is entirely linear, as illustrated below:
 
 ![](simwheeldefaultmovement.gif)(convert:false)
 
-_The standard vertical movement produced by the **WheelHandler**_
+_The standard vertical movement produced by the WheelHandler_
 
 While not entirely realistic, this type of suspension simulation is in fact perfectly fine for most vehicles, because you cannot usually see the suspension arms or the other components such as the springs and shocks. Any car model with a full body can get typically away with this.
 
-However, in a more open vehicle like a Buggy or a Formula 1 car, where the moving parts are in plain view, this movement can be problematic because there is no sensible real-world suspension design that would be able produce this kind of movement.
+However, in a more open vehicle like a Buggy or a Formula 1 car, where the moving parts are in plain view, this movement can be problematic, because there is no sensible real-world suspension design that would be able produce this kind of movement.
 
 To achieve a more realistic result, we really need a type of motion that can be expressed as a rotation around some fixed pivot point on the vehicle's body:
 
@@ -44,41 +44,41 @@ To achieve a more realistic result, we really need a type of motion that can be
 
 _Ideal vertical and lateral movement_
 
-To solve this little conundrum, we can start by assuming that wheel management falls into two different categories: one being the simulated wheel nodes controlled by the **WheelHandler**, the other being the wheels that we can actually see being rendered when we drive the car. Then we could feed data supplied to us by the **WheelHandler** to the rest of our suspension setup to get the desired results.
+To solve this little conundrum, we can start by assuming that wheel management falls into two different categories: one being the simulated wheel nodes controlled by the WheelHandler, the other being the wheels that we can actually see being rendered when we drive the car. Then we could feed data supplied to us by the WheelHandler to the rest of our suspension setup to get the desired results.
 
 In short, the visible wheel and the simulated wheel could in fact be two entirely different objects, and the simulated wheel would not even have to be visible at all.
 
-The wheel parameters we define in the VehicleWheel class can specify the collision mesh, radius and width explicitly. Those specifications do not necessarily have to correspond to any real geometry in the vehicle's skeletal mesh, so there is really no need for the simulated wheel nodes in the mesh to have any real geometry attached to them at all. Furthermore, we can specify additional wheel offsets in the wheel setup section of the vehicle Blueprint, so the pivot of our wheels do not have to be aligned with the geometrical center of the visible wheel mesh.
+The wheel parameters we define in the VehicleWheel class can specify the **collision mesh**, **radius** and **width** explicitly. Those specifications do not necessarily have to correspond to any real geometry in the vehicle's skeletal mesh, so there is really no need for the simulated wheel nodes in the mesh to have any real geometry attached to them at all. Furthermore, we can specify additional wheel offsets in the wheel setup section of the vehicle Blueprint, so the pivot of our wheels do not have to be aligned with the geometrical center of the visible wheel mesh.
 
 These are the basic insights we need to build a suspension that looks like it could realistically work from a mechanical point of view.
 
 ## What Skeletal Nodes Can Do
 
-Before we move on, we need to take a look at two animation Blueprint nodes that will be particularly useful for our setup. The **WheelHandler** node has already been covered, and we will definitely want to use that, but there are other powerful tools that are crucial for this type of rigging: the **CopyBone** and **LookAt** nodes. They can both be found in the Skeletal Control category in the animation Blueprint editor, and they perform much the same tasks as standard position, rotation and aim constraints typically do in an application like Maya.
+Before we move on, we need to take a look at two animation Blueprint nodes that will be particularly useful for our setup. The WheelHandler node has already been covered, and we will definitely want to use that, but there are other powerful tools that are crucial for this type of rigging: the **CopyBone** and **LookAt** nodes. They can both be found in the Skeletal Control category in the animation Blueprint editor, and they perform much the same tasks as standard position, rotation and aim constraints typically do in an application like Maya.
 
 ### Copy Bone
 
-As its name suggests, the **CopyBone** node can copy transformation data from one bone (the source) to another bone (the target):
+As its name suggests, the CopyBone node can copy transformation data from one bone (the source) to another bone (the target):
 
 ![](node_copybone.png)
 
-_The **CopyBone** animation Blueprint node and its default settings._
+_The CopyBone animation Blueprint node and its default settings._
 
-This means that once the **WheelHandler** has updated the transformation of the simulated wheels, we could, for example, grab only the rotation values it has produced and apply them to our visible wheels. This basic step would take care of spinning and steering the wheels, which we do want, while we avoid getting the translation (position) data, which we do not want.
+This means that once the WheelHandler has updated the transformation of the simulated wheels, we could, for example, grab only the rotation values it has produced and apply them to our visible wheels. This basic step would take care of spinning and steering the wheels, which we do want, while we avoid getting the translation (position) data, which we do not want.
 
 ### Look At
 
-The **LookAt** node can rotate any given bone in such a way that one of its cardinal axes ends up aiming at another bone:
+The LookAt node can rotate any given bone in such a way that one of its cardinal axes ends up aiming at another bone:
 
 ![](node_lookat.png)
 
-_The **LookAt** animation Blueprint node and its default settings._
+_The LookAt animation Blueprint node and its default settings._
 
-This is useful for our setup, because it gives us a way to make sure that all the components of the suspension mechanism are constantly updated to point in the right direction, all driven directly or indirectly by the current position of the simulated wheels. All we have to do is to make sure that there are bones available for the **LookAt** nodes to aim at. This will be used for several bones in our animation Blueprint.
+This is useful for our setup, because it gives us a way to make sure that all the components of the suspension mechanism are constantly updated to point in the right direction, all driven directly or indirectly by the current position of the simulated wheels. All we have to do is to make sure that there are bones available for the LookAt nodes to aim at. This will be used for several bones in our animation Blueprint.
 
 ## Implementation Overview
 
-To recap: for each of the four wheels we will use one invisible wheel for the **WheelHandler** to manage, and one visible wheel that is actually rendered in game. The former is just a single bone/joint, without any part of the mesh weighted to it, while the latter is any wheel-like mesh we build in a modeling package and add to the vehicle model, which is then imported into UE4 for final setup.
+To recap: for each of the four wheels we will use one invisible wheel for the WheelHandler to manage, and one visible wheel that is actually rendered in game. The former is just a single bone/joint, without any part of the mesh weighted to it, while the latter is any wheel-like mesh we build in a modeling package and add to the vehicle model, which is then imported into UE4 for final setup.
 
 For everything to work as intended, key parts of the suspension will then copy the transformation data they need from other bones, or adjust their orientation by looking at targets that we have set up in the model.
 
@@ -98,11 +98,11 @@ The most central point where everything comes together is close to the green bol
 
 Each one of these components has a very strict responsibility and transformation range. The arms rotate only around their forward axis. The kingpin follows the position of the tip of the lower arm, but stays vertical at all times and never rotates relative to the vehicle. The hub rotates around its vertical axis only, in response to steering. The visible wheel copies its rotation from the simulated wheel, but takes its position from the tip of the lower arm. This keeps everything tightly locked together.
 
-In any reasonably complex rig, it is very important to keep the hierarchy and update order in mind. In this particular case: when the lower arm rotates, it must bring with it a child mesh (lower green bolt) that the kingpin subsequently needs in order to position itself correctly using a **CopyBone** node. The upper arm in turn needs access to the updated position of a child mesh (near the upper green bolt) of the kingpin in order to have something to aim at using a **LookAt** node, and so on. As you can see, it is essential to make sure that the different components are updated in the correct order, otherwise parts of the suspension may appear to lag behind the movement of other parts.
+In any reasonably complex rig, it is very important to keep the hierarchy and update order in mind. In this particular case: when the lower arm rotates, it must bring with it a child mesh (lower green bolt) that the kingpin subsequently needs in order to position itself correctly using a CopyBone node. The upper arm in turn needs access to the updated position of a child mesh (near the upper green bolt) of the kingpin in order to have something to aim at using a LookAt node, and so on. As you can see, it is essential to make sure that the different components are updated in the correct order, otherwise parts of the suspension may appear to lag behind the movement of other parts.
 
 All these frame-by-frame updates are handled by the animation Blueprint.
 
-When building the model in a DCC (Digital Content Creation) application, we really only have to worry about the hierarchy and the exact position and orientation of the component pivots. A sensible orientation is particularly important for components that will have their rotation updated by **LookAt** nodes.
+When building the model in a DCC (Digital Content Creation) application, we really only have to worry about the hierarchy and the exact position and orientation of the component pivots. A sensible orientation is particularly important for components that will have their rotation updated by LookAt nodes.
 
 As a general note, it can be helpful to use the constraints offered by your modeling package while designing this kind of mechanism. They will not follow the FBX into UE4 on import, but sometimes it is easier to figure out what needs to be done if you can get some basic movement going while building the model, just to ensure that different components will be able to move the way they need to without intersecting each other. However, avoid using any constraints that you will not be able to recreate inside the animation Blueprint.
 
@@ -112,11 +112,11 @@ Now let us see how this demo has been set up inside Unreal Engine. Load up the p
 
 While driving the vehicle, the following things happen in the animation Blueprint:
 
-First, the **WheelHandler** updates the position and rotation of the simulated wheels (`PhysWheel` bones). The visible wheels (`VisWheel` bones) then update their rotation by copying the rotation from the PhysWheels:
+First, the WheelHandler updates the position and rotation of the simulated wheels (`PhysWheel` bones). The visible wheels (`VisWheel` bones) then update their rotation by copying the rotation from the `PhysWheels`:
 
 ![](animBP_1.png)
 
-_The **WheelHandler** followed by Copy Rotation to the visible wheels._
+_The WheelHandler followed by Copy Rotation to the visible wheels._
 
 This gives us a visually correct spin and steering, but the wheels will not change their position. That is handled separately in a later step.
 
@@ -126,15 +126,15 @@ Next, the lower arms (`ArmLower` bones) aim at the position of the `PhysWheel`.
 
 _The ArmLower bones looking at the PhysWheel position._
 
-The result of a **LookAt** node can easily be previewed in the 3D window of the animation Blueprint editor: the current target of a selected **LookAt** node is marked with a red cross. This is easier to see if you set the preview viewport to wireframe mode.
+The result of a LookAt node can easily be previewed in the 3D window of the animation Blueprint editor: the current target of a selected LookAt node is marked with a red cross. This is easier to see if you set the preview viewport to wireframe mode.
 
-When an `ArmLower` bone rotates to match its **LookAt** target, it brings with it a child component called `POS_Hub`. This is used in the next step to correctly position the Kingpin using a **CopyBone** node which grabs only the translation of the `POS_Hub` bone:
+When an `ArmLower` bone rotates to match its LookAt target, it brings with it a child component called `POS_Hub`. This is used in the next step to correctly position the Kingpin using a CopyBone node which grabs only the translation of the `POS_Hub` bone:
 
 ![](animBP_3.png)
 
 _Setting the Kingpin position._
 
-The Hub and VisWheel will later be set to the same position using the same method.
+The Hub and `VisWheel` will later be set to the same position using the same method.
 
 The next step adjusts the upper arms (`ArmUpper` bones) by looking at their respective targets (`LAT_ArmUp` bones) which have been placed in the correct position by the fact that these targets are children of the Kingpin bones, which were updated in the previous step:
 
@@ -142,7 +142,7 @@ The next step adjusts the upper arms (`ArmUpper` bones) by looking at their resp
 
 _Orienting the upper arms._
 
-After setting the position of the hubs, there is some extra work that needs to be done on the front hubs only, to account for rotation around Z caused by steering. This is handled with **LookAt** nodes targeting `LAT_Hub` bones that are children of the frontal PhysWheels:
+After setting the position of the hubs, there is some extra work that needs to be done on the front hubs only, to account for rotation around Z caused by steering. This is handled with LookAt nodes targeting `LAT_Hub` bones that are children of the frontal PhysWheels:
 
 ![](animBP_5.png)
 
@@ -191,8 +191,8 @@ tip, and a smooth gradient between the two extremes.
 The constraints applied to these bodies are fully locked, both linearly and angularly. Since the bone chain is relatively long, small changes in orientation will still occur in response to 
 acceleration, which combined with the mass scale settings described above gives us a stable simulation with just the right kind of mildly amusing oscillation.
 
-To ensure that the simulation does not explode during extreme acceleration, all the antenna constraints have projection enabled, using an angular tolerance of 1.0, and the Position Solver 
-Iteration Count for the bodies has been bumped up to 16.
+To ensure that the simulation does not explode during extreme acceleration, all the antenna constraints have **Enable Projection** checked, using an **Angular Tolerance** of 1.0, and the **Position Solver 
+Iteration Count** for the bodies has been bumped up to 16.
 
 ![](phat_antenna_2.png)
 
@@ -205,7 +205,7 @@ Finally, all internal collision between the antenna bodies as well as the rest o
 
 Now that you know how this setup works, you will no doubt want to start building your own model to replace this prototype. Below are some pointers that will help you do that with a minimum of headaches.
 
-The first thing is to keep the names and the hierarchy pretty much identical to the prototype. Doing that ensures that the animation Blueprint can be reassigned and reused on a different Skeletal Mesh without any changes. That will save you a lot of time. All the meshes can of course be remodeled and moved to different positions. Just make sure that everything is properly aligned. Adding more mesh nodes to the model is not a problem, as long as it does not interfere with the overall hierarchy of the key suspension components.
+The first thing is to keep the names and the hierarchy pretty much identical to the prototype. Doing that ensures that the Animation Blueprint can be reassigned and reused on a different Skeletal Mesh without any changes. That will save you a lot of time. All the meshes can of course be remodeled and moved to different positions. Just make sure that everything is properly aligned. Adding more mesh nodes to the model is not a problem, as long as it does not interfere with the overall hierarchy of the key suspension components.
 
 The vehicle root should be placed at 0,0,0 in global space, and must not be rotated in any way. Making the root a simple joint object ensures that FBX export and import will work with a minimum of problems, particularly when using soft-bound components such as the deforming shock meshes.
 
