#UE4docs #docs Updating documentation on how to download UE4 from GitHub in Windows, Mac, and Linux.

diff --git a/Engine/Documentation/Source/GettingStarted/DownloadingUnrealEngine/DownloadingUnrealEngine4.INT.udn b/Engine/Documentation/Source/GettingStarted/DownloadingUnrealEngine/DownloadingUnrealEngine4.INT.udn
--- a/Engine/Documentation/Source/GettingStarted/DownloadingUnrealEngine/DownloadingUnrealEngine4.INT.udn
+++ b/Engine/Documentation/Source/GettingStarted/DownloadingUnrealEngine/DownloadingUnrealEngine4.INT.udn
@@ -28,15 +28,15 @@ Once you've accessed the repository, you'll notice that we've published UE4's so
 
 ### Release Branch
 
-The [release branch](https://github.com/EpicGames/UnrealEngine/tree/release) (pictured below) is extensively tested by our QA team and makes a great starting point for learning UE4 or making Unreal games. We work hard to stabilize this branch with the goal of releasing new builds every one to two months.
-
+![](ReleaseBranch.png)
 
+The [release branch](https://github.com/EpicGames/UnrealEngine/tree/release) (pictured above) is extensively tested by our QA team and makes a great starting point for learning UE4 or making Unreal games. We work hard to stabilize this branch with the goal of releasing new builds every one to two months.
 
 ### Promoted Branch
 
-The [promoted branch](https://github.com/EpicGames/UnrealEngine/tree/promoted) (pictured below) is updated with builds for use by our designers and artists. We try to update this branch daily so that our subscribers have access to UE4's latest stable features. This branch is a great place to learn how feature development works in UE4.
-
+![](PromotedBranch.png)
 
+The [promoted branch](https://github.com/EpicGames/UnrealEngine/tree/promoted) (pictured above) is updated with builds for use by our designers and artists. We try to update this branch daily so that our subscribers have access to UE4's latest stable features. This branch is a great place to learn how feature development works in UE4.
 
 ### Master Branch
 
@@ -44,9 +44,11 @@ The [promoted branch](https://github.com/EpicGames/UnrealEngine/tree/promoted) (
 Because the master branch contains cutting-edge technology, UE4 isn't guaranteed to compile from source.
 [/REGION]
 
-The [master branch](https://github.com/EpicGames/UnrealEngine/tree/master) (pictured below) tracks [live changes](https://github.com/EpicGames/UnrealEngine/commits/master) being made by our engine team. Our experienced subscribers typically work from this branch when they want to work in tandem with UE4's developers.
+![](MasterBranch.png)
 
+The [master branch](https://github.com/EpicGames/UnrealEngine/tree/master) (pictured above) tracks [live changes](https://github.com/EpicGames/UnrealEngine/commits/master) (pictured below) being made by our engine team. Our experienced subscribers typically work from this branch when they want to work in tandem with UE4's developers.
 
+![](LiveChanges.png)
 
 [REGION:note]
 Other branches may appear from time-to-time as we work to stabilize new releases or hotfixes.
@@ -60,7 +62,14 @@ Other branches may appear from time-to-time as we work to stabilize new releases
 
 [PUBLISH:Rocket]
 
-## Getting Up and Running
+## Downloading Unreal Engine
+
+Subscribers can download Unreal Engine's source code by:
+
+*	Downloading Unreal Engine from GitHub
+*	Downloading Unreal Engine as a Compressed File
+
+### %GettingStarted/DownloadingUnrealEngine/GitHub:title%
 
 [INCLUDE:GettingStarted/DownloadingUnrealEngine/GitHub]
 [/PUBLISH:Rocket]
diff --git a/Engine/Documentation/Source/GettingStarted/DownloadingUnrealEngine/GitHub/GitHub.INT.udn b/Engine/Documentation/Source/GettingStarted/DownloadingUnrealEngine/GitHub/GitHub.INT.udn
--- a/Engine/Documentation/Source/GettingStarted/DownloadingUnrealEngine/GitHub/GitHub.INT.udn
+++ b/Engine/Documentation/Source/GettingStarted/DownloadingUnrealEngine/GitHub/GitHub.INT.udn
@@ -16,7 +16,24 @@ Description:Downloading the Unreal Engine source code from GitHub.
 	active_button_content
 	[/PARAMLITERAL]
 	[PARAM:content]
-	Test Windows
+	1.	**Install** [GitHub for Windows](https://windows.github.com/).
+	1.	**Fork** and **Clone** our [repository](https://github.com/EpicGames/UnrealEngine) (pictured below).
+		![](UnrealRepo.png)
+		[REGION:tip]
+		*	Refer to GitHub's [Forking Projects Guide](https://guides.github.com/activities/forking/) to learn how to fork and clone our repository.
+		*	To use Git from the command line, refer to GitHub's articles on [Setting-up Git](https://help.github.com/articles/set-up-git/) and [Forking a Repository](https://help.github.com/articles/fork-a-repo/).
+		[/REGION]
+	1.	If you already haven't done so, **install** Visual Studio 2013. Although MFC libraries are included in Visual Studio's default installation settings, make sure that MFC libraries are included for ATL support.
+		[REGION:note]
+		Developers build UE4 with all desktop editions of Visual Studio, including Microsoft's free version ([Visual Studio Community 2013](https://www.visualstudio.com/products/visual-studio-community-vs)).
+		[/REGION]
+	1.	**Open** your source folder in Explorer and **run Setup.bat**.
+		[REGION:note]
+		Our Setup batch file downloads binary content for the engine, installs prerequesites, and sets up Unreal file associations. Downloading engine binaries may take some time because they are currently take up three to four gigabytes of memory. Subsequent checkouts take less time to download.
+		[/REGION]
+		[REGION:warning]
+		A warning from SmartScreen may appear in Windows 8. **Click "More Info"**, then **"Run Anyway"** to continue running our Setup batch file.
+		[/REGION]
 	[/PARAM]
 [/OBJECT]
 
@@ -31,7 +48,18 @@ Description:Downloading the Unreal Engine source code from GitHub.
 	active_button_content
 	[/PARAMLITERAL]
 	[PARAM:content]
-	Test Macintosh
+	1.	**Install** [GitHub for Mac](https://mac.github.com/).
+	1.	**Fork** and **Clone** our [repository](https://github.com/EpicGames/UnrealEngine) (pictured below).
+		![](UnrealRepo.png)
+		[REGION:tip]
+		*	Refer to GitHub's [Forking Projects Guide](https://guides.github.com/activities/forking/) to learn how to fork and clone our repository.
+		*	To use Git from the command line, refer to GitHub's articles on [Setting-up Git](https://help.github.com/articles/set-up-git/) and [Forking a Repository](https://help.github.com/articles/fork-a-repo/).
+		[/REGION]
+	1.	If you haven't already done so, **install** the latest version of [Xcode](https://itunes.apple.com/us/app/xcode/id497799835).
+	1.	**Open** your source folder in Finder and **run Setup.command**.
+		[REGION:note]
+		**Close** the Terminal window after Setup.command downloads binary content for the engine, installs prequisites, and sets up Unreal file associations.
+		[/REGION]
 	[/PARAM]
 [/OBJECT]
 
@@ -46,10 +74,17 @@ Description:Downloading the Unreal Engine source code from GitHub.
 	active_button_content
 	[/PARAMLITERAL]
 	[PARAM:content]
-	Test Linux
+	1.	**Set up** [Git](https://help.github.com/articles/set-up-git/).
+	1.	**Fork** and **Clone** our [repository](https://help.github.com/articles/fork-a-repo/) (pictured below).
+		![](UnrealRepo.png)
+	1.	**Open** your source folder and **run Setup.sh** to **download** the engine's binary content.
 	[/PARAM]
 [/OBJECT]
 
+## Next Steps
+
+Now that you've downloaded and set-up Unreal Engine, you're ready to [build the engine from source](Programming/Development/BuildingUnrealEngine).
+
 <!--
 
 There are two ways to get the source code for Unreal Engine from GitHub:
