#UE4doc #docs: Removed extraneous bolding as according to the instructional guide. ##Overview removed as according to the instructional guide.

diff --git a/Engine/Documentation/Source/Resources/ContentExamples/LevelScripting/1_3/1_3.INT.udn b/Engine/Documentation/Source/Resources/ContentExamples/LevelScripting/1_3/1_3.INT.udn
--- a/Engine/Documentation/Source/Resources/ContentExamples/LevelScripting/1_3/1_3.INT.udn
+++ b/Engine/Documentation/Source/Resources/ContentExamples/LevelScripting/1_3/1_3.INT.udn
@@ -26,8 +26,6 @@ Description:An overview of the Level Scripting example level, example 1.3: Spawn
 
 ![](levelscripting_1_3.png)
 
-## Overview
-
 In this example, when the player walks into a box trigger, 8 coin pickups are spawned, each of which contains a dispatch function for when it is collected by the player. In order to open the next set of doors, the player must collect all 8 coins. 
 
 ## Level Script
