#UE4doc #docs: ##Overview removed as according to the instructional guide.

diff --git a/Engine/Documentation/Source/Resources/ContentExamples/LevelScripting/1_2/1_2.INT.udn b/Engine/Documentation/Source/Resources/ContentExamples/LevelScripting/1_2/1_2.INT.udn
--- a/Engine/Documentation/Source/Resources/ContentExamples/LevelScripting/1_2/1_2.INT.udn
+++ b/Engine/Documentation/Source/Resources/ContentExamples/LevelScripting/1_2/1_2.INT.udn
@@ -26,8 +26,6 @@ Description:An overview of the Level Scripting example level, example 1.2: Event
 
 ![](levelscripting_1_2.png)
 
-## Overview
-
 Another instance of the door Blueprint is used here, but this time it is opened or closed when the player steps on a button Blueprint. This button Blueprint uses a special dispatch event, which the level script can receive. When the button is pressed, the dispatch function is called. On its own, this does not do anything until the corresponding event is added to the level script.
 
  
