#UE4doc #docs:
Added Engine Version field
Changed Home link to point to the overview for the tutorial
Changed font settings for multiple terms
Fixed capitalization in title

diff --git a/Engine/Documentation/Source/Programming/Tutorials/VariablesTimersEvents/1/VariablesTimersEvents_1.INT.udn b/Engine/Documentation/Source/Programming/Tutorials/VariablesTimersEvents/1/VariablesTimersEvents_1.INT.udn
--- a/Engine/Documentation/Source/Programming/Tutorials/VariablesTimersEvents/1/VariablesTimersEvents_1.INT.udn
+++ b/Engine/Documentation/Source/Programming/Tutorials/VariablesTimersEvents/1/VariablesTimersEvents_1.INT.udn
@@ -1,7 +1,8 @@
 Availability: Public
-Title:1. Creating an Actor that uses a timer
+Title:1. Creating an Actor that Uses a Timer
 Crumbs:%ROOT%, Programming, Programming/Tutorials, Programming/Tutorials/VariablesTimersEvents
 Description:Starting and stopping a timer that runs code at a fixed rate.
+version:4.9
 
 [VAR:Steps]
 [OBJECT:Navigation]
@@ -10,7 +11,7 @@ Description:Starting and stopping a timer that runs code at a fixed rate.
 	[PARAM:current]
 	[/PARAM]
 	[PARAM:home]
-		[Programming Tutorials Home](Programming\Tutorials\)
+		[](Programming\Tutorials\VariablesTimersEvents)
 	[/PARAM]
 	[PARAM:next]
 		[Next Step](Programming\Tutorials\VariablesTimersEvents\2 "%Programming\Tutorials\VariablesTimersEvents\2:title%")
@@ -22,13 +23,13 @@ Description:Starting and stopping a timer that runs code at a fixed rate.
 
 [REGION:tip]If you are new to **Unreal Engine 4**, you might want to read our [%Programming\QuickStart:title% tutorial](Programming\QuickStart) first. For this tutorial, we will assume you are familiar with creating a project and adding C++ code to it.[/REGION]
 
-1. We will begin by creating a new, Basic Code project, with starter content, named "HowTo_VTE", and then adding an **Actor** class to it. We'll name it "Countdown" in this tutorial.
+1. We will begin by creating a new, Basic Code project, with starter content, named HowTo_VTE, and then adding an **Actor** class to it. We'll name it Countdown in this tutorial.
 
 	[![](ChooseParentClass.png)](ChooseParentClass.png)
 
 	[![](NameYourActor.png)](NameYourActor.png)
 
-1. We'll start by creating a simple countdown timer that we can see in-game. In `Countdown.h`, add the following lines to the end of our class definition:
+1. We'll start by creating a simple countdown timer that we can see in-game. In Countdown.h, add the following lines to the end of our class definition:
 
 		int32 CountdownTime;
 
@@ -36,7 +37,7 @@ Description:Starting and stopping a timer that runs code at a fixed rate.
 
 		void UpdateTimerDisplay();
 
-1. In `Countdown.cpp`, we can create our renderable text **Component** and initialize our countdown time to 3 seconds. We can also turn **Ticking** off for this type of **Actor**, since we will not need it. **ACountdown::ACountdown** should look like this:
+1. In Countdown.cpp, we can create our renderable text **Component** and initialize our countdown time to 3 seconds. We can also turn **Ticking** off for this type of **Actor**, since we will not need it. **ACountdown::ACountdown** should look like this:
 
 		PrimaryActorTick.bCanEverTick = false;
 
@@ -47,14 +48,14 @@ Description:Starting and stopping a timer that runs code at a fixed rate.
 
 		CountdownTime = 3;
 
-1. **ACountdown::UpdateTimerDisplay** should update our text display to show the time remaining, or zero if the time is up. This code should run when we first spawn our **ACountdown** into the game, and once per second until our CountdownTime variable hits zero.
+1. **ACountdown::UpdateTimerDisplay** should update our text display to show the time remaining, or zero if the time is up. This code should run when we first spawn our ACountdown into the game, and once per second until our CountdownTime variable hits zero.
 
 		void ACountdown::UpdateTimerDisplay()
 		{
 			CountdownText->SetText(FString::FromInt(FMath::Max(CountdownTime, 0)));
 		}
 
-1. Whenever we assign a **Timer** to run a function, we are given a **Timer Handle**. We need to hold onto that handle so that we can shut the **Timer** down when the countdown finishes. Let's add the function to count time down, and the **Timer Handle** we'll need to control it, to the class definition in `Countdown.h`. While we're there, let's also add a function to do something special when the countdown ends:
+1. Whenever we assign a **Timer** to run a function, we are given a **Timer Handle**. We need to hold onto that handle so that we can shut the Timer down when the countdown finishes. Let's add the function to count time down, and the Timer Handle we'll need to control it, to the class definition in Countdown.h. While we're there, let's also add a function to do something special when the countdown ends:
 
 		void AdvanceTimer();
 
@@ -87,20 +88,20 @@ Description:Starting and stopping a timer that runs code at a fixed rate.
 		UpdateTimerDisplay();
 		GetWorldTimerManager().SetTimer(CountdownTimerHandle, this, &ACountdown::AdvanceTimer, 1.0f, true);
 
-	[REGION:note]We are updating the display in **ACountdown::BeginPlay** rather than **ACountdown::ACountdown** because values set to variables in the **Unreal Editor** will be assigned after the constructor, but before **BeginPlay**. We will want to respect those values later, when we expose CountdownTime to the editor.[/REGION]
+	[REGION:note]We are updating the display in ACountdown::BeginPlay rather than ACountdown::ACountdown because values set to variables in the **Unreal Editor** will be assigned after the constructor, but before **BeginPlay**. We will want to respect those values later, when we expose CountdownTime to the editor.[/REGION]
 
 1. Let's check our progress so far by going to the **Unreal Editor** and pressing **Compile**.
 
 	[![](CompileFromEditor.png)](CompileFromEditor.png)
 
-	We can then drop our updated **ACountdown** class from the **Content Browser** into the **Level Editor**.
+	We can then drop our updated ACountdown class from the **Content Browser** into the **Level Editor**.
 
 	[![](ClassInContentBrowser.png)](ClassInContentBrowser.png)
 
 	[![](LevelEditorText.png)](LevelEditorText.png)
 
-	[REGION:note]Because we set our countdown text during **ACountdown::BeginPlay** and not **ACountdown::ACountdown**, the default "Text" is shown in the **Level Editor**.[/REGION]
-	When we press **Play**, the countdown will progress as expected, saying "3", "2", "1", and finally "GO!"
+	[REGION:note]Because we set our countdown text during ACountdown::BeginPlay and not ACountdown::ACountdown, the default Text is shown in the **Level Editor**.[/REGION]
+	When we press **Play**, the countdown will progress as expected, saying 3, 2, 1, and finally GO!
 
 ----
 
