#UE4doc #docs: Removed extraneous bolding as according to the instructional guide. ##Overview removed as according to the instructional guide.

diff --git a/Engine/Documentation/Source/Resources/ContentExamples/LevelScripting/1_1/1_1.INT.udn b/Engine/Documentation/Source/Resources/ContentExamples/LevelScripting/1_1/1_1.INT.udn
--- a/Engine/Documentation/Source/Resources/ContentExamples/LevelScripting/1_1/1_1.INT.udn
+++ b/Engine/Documentation/Source/Resources/ContentExamples/LevelScripting/1_1/1_1.INT.udn
@@ -26,8 +26,6 @@ Description:An overview of the Level Scripting example level, example 1.1: Trigg
 
 ![](levelscripting_1_1.png)
 
-## Overview
-
 A Trigger Box Actor is placed in front of a Blueprint containing a set of doors. When the player touches the trigger box, the doors open, and when the player exits the trigger box, the doors close. The door Blueprint contains all the necessary script for animating connected to custom events, which the level script can call.
 
 ## Level Script
@@ -47,7 +45,7 @@ When provided with a reference to a Blueprint Actor, seen here as the small blue
 Click for full size or **right-click** and Save As.
 [/REGION]
 
-The BP_Double_Doors Blueprint Actor contains a setup for a somewhat sophisticated door. You can see how the [Timeline](Engine/Blueprints\UserGuide\Timelines) is used to drive several things, including the location of each door, the rotation of the bolt mechanism, and the Material applied to the bolt. The network is fired whenever the Blueprint receives the **Open Doors** or **Close Doors** events, which are custom events created for this Blueprint.
+The BP_Double_Doors Blueprint Actor contains a setup for a somewhat sophisticated door. You can see how the [Timeline](Engine/Blueprints\UserGuide\Timelines) is used to drive several things, including the location of each door, the rotation of the bolt mechanism, and the Material applied to the bolt. The network is fired whenever the Blueprint receives the Open Doors or Close Doors events, which are custom events created for this Blueprint.
 
 [/EXCERPT:bpDoubleDoors]
 [/REGION]
