#UE4doc #docs:
Added Engine Version field
Changed Home link to point to the overview for the tutorial
Changed font settings for multiple terms
Edited a couple pictures to change the indicator borders and arrows from green to yellow
Fixed capitalization in the title

diff --git a/Engine/Documentation/Source/Programming/Tutorials/VariablesTimersEvents/2/VariablesTimersEvents_2.INT.udn b/Engine/Documentation/Source/Programming/Tutorials/VariablesTimersEvents/2/VariablesTimersEvents_2.INT.udn
--- a/Engine/Documentation/Source/Programming/Tutorials/VariablesTimersEvents/2/VariablesTimersEvents_2.INT.udn
+++ b/Engine/Documentation/Source/Programming/Tutorials/VariablesTimersEvents/2/VariablesTimersEvents_2.INT.udn
@@ -1,8 +1,8 @@
 Availability: Public
-Title:2. Expose variables and functions to the editor
+Title:2. Exposing Variables and Functions to the Editor
 Crumbs:%ROOT%, Programming, Programming/Tutorials, Programming/Tutorials/VariablesTimersEvents
 Description:Using the UPROPERTY and UFUNCTION macros to expose variables and functions to the editor.
-
+version:4.9
 [VAR:Steps]
 [OBJECT:Navigation]
 	[PARAM:previous]
@@ -11,7 +11,7 @@ Description:Using the UPROPERTY and UFUNCTION macros to expose variables and fun
 	[PARAM:current]
 	[/PARAM]
 	[PARAM:home]
-		[Programming Tutorials Home](Programming\Tutorials\)
+		[](Programming\Tutorials\VariablesTimersEvents)
 	[/PARAM]
 	[PARAM:next]
 		[Next Step](Programming\Tutorials\VariablesTimersEvents\3 "%Programming\Tutorials\VariablesTimersEvents\3:title%")
@@ -21,16 +21,16 @@ Description:Using the UPROPERTY and UFUNCTION macros to expose variables and fun
 
 %Steps%
 
-1. Our countdown timer is currently hard-coded to use a value of 3 seconds. It would be more useful if we could set the countdown time to any value we want in the editor, and this is easy to do. In **Visual Studio**, we can open `Countdown.h` and find the line that says:
+1. Our countdown timer is currently hard-coded to use a value of 3 seconds. It would be more useful if we could set the countdown time to any value we want in the editor, and this is easy to do. In **Visual Studio**, we can open Countdown.h and find the line that says:
 
 		int32 CountdownTime;
 
-	In order to expose this variable to **Unreal Engine**, we need to make it a **UPROPERTY**. This enables the engine to preserve the value of the variable when launching the game or loading a saved level. The **UPROPERTY** tag, with empty parentheses, is added right above the variable it affects:
+	In order to expose this variable to **Unreal Engine**, we need to make it a **UPROPERTY**. This enables the engine to preserve the value of the variable when launching the game or loading a saved level. The UPROPERTY tag, with empty parentheses, is added right above the variable it affects:
 
 		UPROPERTY()
 		int32 CountdownTime;
 
-	**UPROPERTY** supports arguments that change how **Unreal Engine** will use the variable. Since we want our variable to be editable, we can add the **EditAnywhere** argument:
+	UPROPERTY supports arguments that change how **Unreal Engine** will use the variable. Since we want our variable to be editable, we can add the **EditAnywhere** argument:
 
 		UPROPERTY(EditAnywhere)
 		int32 CountdownTime;
@@ -45,11 +45,11 @@ Description:Using the UPROPERTY and UFUNCTION macros to expose variables and fun
 
 	[![](CommentingVariable.png)](CommentingVariable.png)
 
-	[REGION:tip]There is a lot more we can do with **UPROPERTY**, and looking into other arguments such as **BlueprintReadWrite** and **Category** might be good next steps, but we have all that we need at the moment.[/REGION]
+	[REGION:tip]There is a lot more we can do with UPROPERTY, and looking into other arguments such as **BlueprintReadWrite** and **Category** might be good next steps, but we have all that we need at the moment.[/REGION]
 	
-	When we return to the **Unreal Editor** and press **Compile**, our variable will appear in the **Details Panel** for the **ACountdown** we placed earlier, and we can test out different timer values by changing this number and pressing **Play**.
+	When we return to the **Unreal Editor** and press **Compile**, our variable will appear in the **Details Panel** for the ACountdown we placed earlier, and we can test out different timer values by changing this number and pressing **Play**.
 
-1. In addition to changing the value of the timer, let's also enable non-programming developers to change what happens when the timer is up. In **Visual Studio**, we'll open `Countdown.h` and find the following line:
+1. In addition to changing the value of the timer, let's also enable non-programming developers to change what happens when the timer is up. In **Visual Studio**, we'll open Countdown.h and find the following line:
 
 		void CountdownHasFinished();
 
@@ -59,22 +59,22 @@ Description:Using the UPROPERTY and UFUNCTION macros to expose variables and fun
 		void CountdownHasFinished();
 
 	[REGION:note]
-	Just like the **UPROPERTY** macro, we need to provide information about what can be done with it in order to enable more features and access for non-programming developers. There are three options to consider:
+	Just like the UPROPERTY macro, we need to provide information about what can be done with it in order to enable more features and access for non-programming developers. There are three options to consider:
 	
 	1. **BlueprintCallable** functions are written in C++ and can be called from the **Blueprint Graph**, but cannot be changed or overridden without editing C++ code. Functions marked this way are usually features that have been programmed for non-programmer use, but that are not supposed to be changed or wouldn't make sense to change. An easy example of this would be any kind of math function.
 
-	1. **BlueprintImplementableEvent** functions are set up in a C++ header (`.h`) file, but the body of the function is written entirely in the **Blueprint Graph**, not in C++. These are usually created to give a non-programmer the ability to create custom reactions to special situations that have no expected default action or standard behavior. An example of this might be an event that happens when a powerup touches the player's ship in a spaceship game.
+	1. **BlueprintImplementableEvent** functions are set up in a C++ header (.h) file, but the body of the function is written entirely in the Blueprint Graph, not in C++. These are usually created to give a non-programmer the ability to create custom reactions to special situations that have no expected default action or standard behavior. An example of this might be an event that happens when a powerup touches the player's ship in a spaceship game.
 
-	1. **BlueprintNativeEvent** functions are like a combination of **BlueprintCallable** and **BlueprintImplementableEvent** functions. They have default behaviors programmed in C++, but these can be supplemented or replaced by overriding in the **Blueprint Graph**. When programming these, the C++ code always goes in a virtual function with "_Implementation" added to the end of the name, as shown below. This is the most flexible option, so we will use it for this tutorial.
+	1. **BlueprintNativeEvent** functions are like a combination of BlueprintCallable and BlueprintImplementableEvent functions. They have default behaviors programmed in C++, but these can be supplemented or replaced by overriding in the Blueprint Graph. When programming these, the C++ code always goes in a virtual function with _Implementation added to the end of the name, as shown below. This is the most flexible option, so we will use it for this tutorial.
 	[/REGION]
 
-	To grant non-programmers the ability to call our C++ function and to override it with **Blueprints**, we need to make the following changes to `Countdown.h`:
+	To grant non-programmers the ability to call our C++ function and to override it with **Blueprints**, we need to make the following changes to Countdown.h:
 
 		UFUNCTION(BlueprintNativeEvent)
 		void CountdownHasFinished();
 		virtual void CountdownHasFinished_Implementation();
 
-	Then, in `Countdown.cpp`, we will need to change the line that says:
+	Then, in Countdown.cpp, we will need to change the line that says:
 
 		void ACountdown::CountdownHasFinished()
 
@@ -84,7 +84,7 @@ Description:Using the UPROPERTY and UFUNCTION macros to expose variables and fun
 
 ----
 
-We have now made a variable and a function accessible to, and alterable by, non-programmers, while providing our own default value and functionality in C++. To see how a non-programmer might use this, we'll make a **Blueprint** extension of our **ACountdown** class and modify it ourselves.
+We have now made a variable and a function accessible to, and alterable by, non-programmers, while providing our own default value and functionality in C++. To see how a non-programmer might use this, we'll make a Blueprint extension of our ACountdown class and modify it ourselves.
 
 [OBJECT:Section]
     [PARAMLITERAL:id]
