#UE4doc #docs: ##Overview removed as according to the instructional guide. Removed extraneous bolding as according to the instructional guide. Removed bolding affects multiple pages.

diff --git a/Engine/Documentation/Source/Programming/UnrealArchitecture/TMap/TMap.INT.udn b/Engine/Documentation/Source/Programming/UnrealArchitecture/TMap/TMap.INT.udn
--- a/Engine/Documentation/Source/Programming/UnrealArchitecture/TMap/TMap.INT.udn
+++ b/Engine/Documentation/Source/Programming/UnrealArchitecture/TMap/TMap.INT.udn
@@ -4,8 +4,6 @@ Description:
 Availability:Public
 Version:4.7
 
-## Overview
-
 After TArray, the most commonly used container in Unreal Engine is TMap.  This container is an associative container, which means that every key has an associated value, and you can efficiently look up a value object by its key.
 
 There are two types of map: TMap and TMultiMap.  TMap's keys are unique, and inserting a new key-value pair when the key already exists will cause the existing pair to be replaced.  TMultiMap's keys are not unique and so existing pairs will not be replaced by new additions.
diff --git a/Engine/Documentation/Source/Resources/ContentExamples/Blueprints_Advanced/2_5/2_5.INT.udn b/Engine/Documentation/Source/Resources/ContentExamples/Blueprints_Advanced/2_5/2_5.INT.udn
--- a/Engine/Documentation/Source/Resources/ContentExamples/Blueprints_Advanced/2_5/2_5.INT.udn
+++ b/Engine/Documentation/Source/Resources/ContentExamples/Blueprints_Advanced/2_5/2_5.INT.udn
@@ -54,13 +54,13 @@ Click for full size or **right-click** and Save As.
 [/REGION]
 
 
-The parent Blueprint, **BP_Pickup_Parent**, contains script for basic pickup behavior, such as a timeline animation to make the pickup spin and float, and a sequence for collecting the pickup. 
+The parent Blueprint, BP_Pickup_Parent, contains script for basic pickup behavior, such as a timeline animation to make the pickup spin and float, and a sequence for collecting the pickup. 
 The child Blueprints inherit all of this script, but have additional script for when they are collected, so that each pickup does something different while sharing the same basic features. 
 This is useful because any updates made to the parent Blueprint's script will also apply to all its child Blueprints.
 
 
-The **OnPickedUp** [event dispatcher](Engine/Blueprints\UserGuide\EventDispatcher) is called when the player's Pawn overlaps the **Pickup Trigger** _SphereComponent_. Each child _Blueprint_ binds a different 
-event to **OnPickedUp**, so that a different string prints when the player collects each pickup and **OnPickedUp** is called.
+The OnPickedUp [event dispatcher](Engine/Blueprints\UserGuide\EventDispatcher) is called when the player's Pawn overlaps the **Pickup Trigger** _SphereComponent_. Each child _Blueprint_ binds a different 
+event to OnPickedUp, so that a different string prints when the player collects each pickup and OnPickedUp is called.
 [/EXCERPT:CoinParent]
 
 ## Child Blueprints' Event Graph
@@ -75,7 +75,7 @@ Click for full size or **right-click** and Save As.
 
 
 
-As a rule, Blueprint graphs cannot contain duplicates of the same event, so if both the parent script and the child script each contain a **ReceiveBeginPlay** event, only the event in the parent script will fire. You can get around this in the child script by **right-clicking** an event and selecting **Add call to parent function**. This will create a function that will trigger the corresponding event in the parent script when it is called.
+As a rule, Blueprint graphs cannot contain duplicates of the same event, so if both the parent script and the child script each contain a ReceiveBeginPlay event, only the event in the parent script will fire. You can get around this in the child script by **right-clicking** an event and selecting **Add call to parent function**. This will create a function that will trigger the corresponding event in the parent script when it is called.
 
 [/EXCERPT:CoinChild]
 
diff --git a/Engine/Documentation/Source/Resources/ContentExamples/MaterialAdvanced/MaterialAdvanced.INT.udn b/Engine/Documentation/Source/Resources/ContentExamples/MaterialAdvanced/MaterialAdvanced.INT.udn
--- a/Engine/Documentation/Source/Resources/ContentExamples/MaterialAdvanced/MaterialAdvanced.INT.udn
+++ b/Engine/Documentation/Source/Resources/ContentExamples/MaterialAdvanced/MaterialAdvanced.INT.udn
@@ -46,8 +46,6 @@ Description:An overview of the Material Advanced example level
 
 <!--Insert path when done! [RELATIVE:Resources/ContentExamples/MaterialAdvanced]-->
 
-## Overview
-
 [REGION:fullwidth]
 
 [/REGION]
