#UE4doc #docs: ##Overview removed as according to the instructional guide.

diff --git a/Engine/Documentation/Source/Programming/UnrealArchitecture/TArrays/TArrays.INT.udn b/Engine/Documentation/Source/Programming/UnrealArchitecture/TArrays/TArrays.INT.udn
--- a/Engine/Documentation/Source/Programming/UnrealArchitecture/TArrays/TArrays.INT.udn
+++ b/Engine/Documentation/Source/Programming/UnrealArchitecture/TArrays/TArrays.INT.udn
@@ -4,8 +4,6 @@ Description:
 Availability:Public
 Version:4.7
 
-## Overview
-
 The primary go-to container class is TArray.  TArray is a class which is responsible for the ownership and organization of a sequence of other objects (called elements) of the same type.  As a TArray is a sequence, its elements have a well-defined order and its methods are used to deterministically manipulate those objects and their order.
 
 ## TArray
