#UE4doc #docs: ##Overview removed as according to the instructional guide.

diff --git a/Engine/Documentation/Source/Programming/UnrealArchitecture/Actors/Ticking/ActorTicking.INT.udn b/Engine/Documentation/Source/Programming/UnrealArchitecture/Actors/Ticking/ActorTicking.INT.udn
--- a/Engine/Documentation/Source/Programming/UnrealArchitecture/Actors/Ticking/ActorTicking.INT.udn
+++ b/Engine/Documentation/Source/Programming/UnrealArchitecture/Actors/Ticking/ActorTicking.INT.udn
@@ -6,10 +6,6 @@ Version: 4.5
 
 [TOC (start:2 end:3)]
 
-
-
-## Overview
-
 Actors are ticked once per frame with the elapsed time between frames in the order they appear in the World's Actor list (aggregate of all levels' lists). Actors are now ticked in 3 phases: before 
 asynchronous work, during asynchronous work, and after asynchronous work. Which group an Actor is ticked in is controlled by the TickGroup member. If an Actor needs to have its state updated before 
 any asynchronous work (physics, and others later), then it needs to be assigned to the TG_PreAsyncWork group. Failing to do so will cause incorrect behavior and/or off-by-one-frame problems. Actors 
