#UE4doc #docs: updated engine version to 4.9

diff --git a/Engine/Documentation/Source/Programming/UnrealArchitecture/Actors/Ticking/ActorTicking.INT.udn b/Engine/Documentation/Source/Programming/UnrealArchitecture/Actors/Ticking/ActorTicking.INT.udn
--- a/Engine/Documentation/Source/Programming/UnrealArchitecture/Actors/Ticking/ActorTicking.INT.udn
+++ b/Engine/Documentation/Source/Programming/UnrealArchitecture/Actors/Ticking/ActorTicking.INT.udn
@@ -2,7 +2,7 @@ Availability: Docs
 Title:Actor Ticking
 Crumbs:%ROOT%, Programming
 Description:Explanation of the ticking system used to update Actors each frame.
-Version: 4.5
+Version: 4.9
 
 [TOC (start:2 end:3)]
 
diff --git a/Engine/Documentation/Source/Programming/UnrealArchitecture/TArrays/TArrays.INT.udn b/Engine/Documentation/Source/Programming/UnrealArchitecture/TArrays/TArrays.INT.udn
--- a/Engine/Documentation/Source/Programming/UnrealArchitecture/TArrays/TArrays.INT.udn
+++ b/Engine/Documentation/Source/Programming/UnrealArchitecture/TArrays/TArrays.INT.udn
@@ -2,7 +2,7 @@ Title: TArray: Arrays in Unreal Engine
 Crumbs:
 Description:
 Availability:Public
-Version:4.7
+Version:4.9
 
 The primary go-to container class is TArray.  TArray is a class which is responsible for the ownership and organization of a sequence of other objects (called elements) of the same type.  As a TArray is a sequence, its elements have a well-defined order and its methods are used to deterministically manipulate those objects and their order.
 
