#UE4doc #docs: Fixed a typo: creat -> create. ##Overview removed as according to the instructional guide. Updated engine version to 4.9. Altered images to match the instructional guide: red -> yellow.

diff --git a/Engine/Documentation/Source/Programming/Slate/InGameUI/QuickStart/SlateInGameUIQuickStart.INT.udn b/Engine/Documentation/Source/Programming/Slate/InGameUI/QuickStart/SlateInGameUIQuickStart.INT.udn
--- a/Engine/Documentation/Source/Programming/Slate/InGameUI/QuickStart/SlateInGameUIQuickStart.INT.udn
+++ b/Engine/Documentation/Source/Programming/Slate/InGameUI/QuickStart/SlateInGameUIQuickStart.INT.udn
@@ -2,11 +2,11 @@ Title:Slate In-Game UI Quick Start
 Crumbs:
 Description:
 Availability:Docs
-Version:4.4
+Version:4.9
 
 ## Project Setup
 
-To start off, creat a new basic code project:
+To start off, create a new basic code project:
 
 [REGION:fullwidth]
 ![](image001.png)
