#UE4doc #docs:
Added Engine Version field
Changed Home link to point to the overview for the tutorial
Changed font settings for multiple terms
Replaced a couple pictures to change the indicator borders and arrows from green to yellow
Fixed capitalization in the title

diff --git a/Engine/Documentation/Source/Programming/Tutorials/VariablesTimersEvents/3/VariablesTimersEvents_3.INT.udn b/Engine/Documentation/Source/Programming/Tutorials/VariablesTimersEvents/3/VariablesTimersEvents_3.INT.udn
--- a/Engine/Documentation/Source/Programming/Tutorials/VariablesTimersEvents/3/VariablesTimersEvents_3.INT.udn
+++ b/Engine/Documentation/Source/Programming/Tutorials/VariablesTimersEvents/3/VariablesTimersEvents_3.INT.udn
@@ -1,8 +1,8 @@
 Availability: Public
-Title:3. Extend and override C++ with Blueprints
+Title:3. Extend and Override C++ with Blueprints
 Crumbs:%ROOT%, Programming, Programming/Tutorials, Programming/Tutorials/VariablesTimersEvents
 Description:Using Blueprints to supplement or replace C++ functionality.
-
+version:4.9
 [VAR:Steps]
 [OBJECT:Navigation]
 	[PARAM:previous]
@@ -11,7 +11,7 @@ Description:Using Blueprints to supplement or replace C++ functionality.
 	[PARAM:current]
 	[/PARAM]
 	[PARAM:home]
-		[Programming Tutorials Home](Programming\Tutorials\)
+		[](Programming\Tutorials\VariablesTimersEvents)
 	[/PARAM]
 	[PARAM:next]
 		[Next Step](Programming\Tutorials\VariablesTimersEvents\4 "%Programming\Tutorials\VariablesTimersEvents\4:title%")
@@ -21,17 +21,17 @@ Description:Using Blueprints to supplement or replace C++ functionality.
 
 %Steps%
 
-[REGION:tip]This section of the tutorial involves using **Blueprints** to extend the functionality of C++ classes. However, it is only intended as a test that our C++ code was written correctly, not as a **Blueprint** tutorial. For a proper introduction to **Blueprints**, we recommend the [%Engine\Blueprints\QuickStart:title%](Engine\Blueprints\QuickStart).[/REGION]
+[REGION:tip]This section of the tutorial involves using **Blueprints** to extend the functionality of C++ classes. However, it is only intended as a test that our C++ code was written correctly, not as a Blueprint tutorial. For a proper introduction to Blueprints, we recommend the [%Engine\Blueprints\QuickStart:title%](Engine\Blueprints\QuickStart).[/REGION]
 
-1. To change the behavior of our **ACountdown** instance, called "Countdown1", in the editor, we must first make an editable **Blueprint** version of it. To do this, we can select it from the **World Outliner** and click the **Blueprint/Add Script** button in the **Details Panel**.
+1. To change the behavior of our ACountdown instance, called Countdown1, in the editor, we must first make an editable Blueprint version of it. To do this, we can select it from the **World Outliner** and click the **Blueprint/Add Script** button in the **Details Panel**.
 
 	[![](AddScript.png)](AddScript.png)
 
-	From there, we can provide a path and name for the **Blueprint** asset that will contain our modified **ACountdown** class.
+	From there, we can provide a path and name for the Blueprint asset that will contain our modified ACountdown class.
 
 	[![](SelectBlueprintPath.png)](SelectBlueprintPath.png)
 
-	This will create a new asset that represents a **Blueprint** version of "Countdown1". It will also replace "Countdown1" with an instance of this new **Blueprint**, so that changes we make to the **Blueprint** will affect "Countdown1" in the game.
+	This will create a new asset that represents a Blueprint version of Countdown1. It will also replace Countdown1 with an instance of this new Blueprint, so that changes we make to the Blueprint will affect Countdown1 in the game.
 
 1. The **Unreal Editor** will automatically take us to our new asset in the **Content Browser**, and we can right-click it and choose "Edit..." to modify its **Blueprint Graph**, **Component** hierararchy, and **Default Values**.
 
@@ -47,27 +47,29 @@ Description:Using Blueprints to supplement or replace C++ functionality.
 
 	[![](SelectEvent.png)](SelectEvent.png)
 
-1. We can now add any additional functionality we would like by left-clicking and dragging off of the white ("execution") pin on the right side of our new node.
+1. We can now add any additional functionality we would like by left-clicking and dragging off of the white (execution) pin on the right side of our new node.
 
 	[![](DragExecPin.png)](DragExecPin.png)
 
-	When we release the left mouse button, we will be asked what function or event we would like to execute. For this tutorial, let's spawn a **Particle System** when the countdown finishes. We'll want a **Spawn Emitter At Location** node, so select that from the list. It can save time to type a partial phrase, like "spawn loc", into the search field. We can then left-click and drag the yellow "Location" pin and attach it to a **Get Actor Location** function.
+	When we release the left mouse button, we will be asked what function or event we would like to execute. For this tutorial, let's spawn a **Particle System** when the countdown finishes. We'll want a **Spawn Emitter At Location** node, so select that from the list. It can save time to type a partial phrase, like spawn loc, into the search field. We can then left-click and drag the yellow "Location" pin and attach it to a **Get Actor Location** function.
 
 	[![](GetActorLocation.png)](GetActorLocation.png)
 
-	Now we just need to select what effect we'd like to see. By clicking "Select Asset" under "Emitter Template", we can get a list of appropriate effect assets. "P_Explosion" is a good one, so we'll pick that.
+	Now we just need to select what effect we'd like to see. By clicking Select Asset under Emitter Template, we can get a list of appropriate effect assets. P_Explosion is a good one, so we'll pick that.
 
 	[![](SelectParticle.png)](SelectParticle.png)
 
+1. Click the **Compile** button at the top left of the **Blueprint Editor** to save the changes.
+
 1. If we press **Play** now, we'll see our countdown take place, and our explosion will happen when our countdown number hits zero.
 
 	[![](Explosion_Zero.png)](Explosion_Zero.png)
 
-	However, we programmed our countdown to say "GO!" at the end, not "0". This is no longer happening because we have completely replaced our C++ functionality with our **Blueprint** visual scripting. This is not what we intended to do in this case, so we need to add a call to the C++ version of this function, which can be done by right-clicking the **Countdown Has Finished** node and selecting "Add call to parent function" from the context menu.
+	However, we programmed our countdown to say GO! at the end, not 0. This is no longer happening because we have completely replaced our C++ functionality with our **Blueprint** visual scripting. This is not what we intended to do in this case, so we need to add a call to the C++ version of this function, which can be done by right-clicking the **Countdown Has Finished** node and selecting **Add call to parent function** from the context menu.
 
 	[![](CallToParent_Menu.png)](CallToParent_Menu.png)
 
-	When this is done, a node labeled **Parent: Countdown Has Finished** will be placed in the **Event Graph**. The typical place to connect a "parent" node is directly to the event node, which is what we will do here. This is not required, though, as the parent-call node is like any other and can be called anywhere we like, even multiple times.
+	When this is done, a node labeled **Parent: Countdown Has Finished** will be placed in the **Event Graph**. The typical place to connect a parent node is directly to the event node, which is what we will do here. This is not required, though, as the parent-call node is like any other and can be called anywhere we like, even multiple times.
 
 	[![](CallToParent_ConnectPins.png)](CallToParent_ConnectPins.png)
 
@@ -75,7 +77,7 @@ Description:Using Blueprints to supplement or replace C++ functionality.
 
 	[![](CallToParent_FixPins.png)](CallToParent_FixPins.png)
 
-	Now when we run our game, we should see both the word "GO!" (from our C++ code) and an explosion (from our **Blueprint Graph**) after the countdown finishes!
+	Now when we run our game, we should see both the word GO! (from our C++ code) and an explosion (from our Blueprint Graph) after the countdown finishes!
 
 	[![](Explosion_Go.png)](Explosion_Go.png)
 
