#UE4doc #docs: ##Overview removed as according to the instructional guide.

diff --git a/Engine/Documentation/Source/Engine/Rendering/Materials/Functions/Reference/Math/Math.INT.udn b/Engine/Documentation/Source/Engine/Rendering/Materials/Functions/Reference/Math/Math.INT.udn
--- a/Engine/Documentation/Source/Engine/Rendering/Materials/Functions/Reference/Math/Math.INT.udn
+++ b/Engine/Documentation/Source/Engine/Rendering/Materials/Functions/Reference/Math/Math.INT.udn
@@ -5,8 +5,6 @@ Description:Preconfigured mathematical functions such as calculation of pi, addi
 
 [TOC (start:2 end:3)]
 
-## Overview
-
 Math functions perform basic math equations on the values of pixels within a texture. 
 
 ## Math Functions
