#UE4doc #docs: fixed instances of the same word: of of -> of. Changed bolded note to [REGION:note].

diff --git a/Engine/Documentation/Source/Engine/Rendering/Materials/Functions/Reference/Misc/Misc.INT.udn b/Engine/Documentation/Source/Engine/Rendering/Materials/Functions/Reference/Misc/Misc.INT.udn
--- a/Engine/Documentation/Source/Engine/Rendering/Materials/Functions/Reference/Misc/Misc.INT.udn
+++ b/Engine/Documentation/Source/Engine/Rendering/Materials/Functions/Reference/Misc/Misc.INT.udn
@@ -27,7 +27,9 @@ The TimeWithSpeedVariable functions just like a time node, but with an optional
 
 ![Time with Speed Offset](TimeWithSpeedVariable_Demo.png)
 
-**Note: The Output Result will preview the node by default.  In this case, the node will flicker.**
+[REGION:note]
+The Output Result will preview the node by default.  In this case, the node will flicker.
+[/REGION]
 
 ### OffsetAndScaleTo1
 
@@ -72,7 +74,7 @@ This function takes in a Texture2D, such as a grid of sprite frames, and outputs
 |[REGION:tablesection]Inputs[/REGION] ||
 | **Animation Phase (0-1) (Scalar)** | If this receives a static input, then the result will give the nearest frame of the flipbook as if the frames were numbered between 0 and 1. If no input is given, time is automatically used.|
 | **Number of Rows (Scalar)** | This takes in the number of rows of the flipbook texture.|
-| **Number of Columns (Scalar)** | This takes in the number of columns of of the flipbook texture.|
+| **Number of Columns (Scalar)** | This takes in the number of columns of the flipbook texture.|
 | **Texture (Texture2D)** | Takes in a Texture2D containing a sprite sheet, or grid of frames of an animated texture.|
 | **UVs (Vector2)** | Takes in a set of UV coordinates to assist in tiling.|
 |[REGION:tablesection]Outputs[/REGION] ||
