#UE4doc #docs:
Added Engine Version field
Changed Home link to point to the overview for the tutorial
Changed font settings for multiple terms

diff --git a/Engine/Documentation/Source/Programming/Tutorials/VariablesTimersEvents/4/VariablesTimersEvents_4.INT.udn b/Engine/Documentation/Source/Programming/Tutorials/VariablesTimersEvents/4/VariablesTimersEvents_4.INT.udn
--- a/Engine/Documentation/Source/Programming/Tutorials/VariablesTimersEvents/4/VariablesTimersEvents_4.INT.udn
+++ b/Engine/Documentation/Source/Programming/Tutorials/VariablesTimersEvents/4/VariablesTimersEvents_4.INT.udn
@@ -2,6 +2,7 @@ Availability: Public
 Title:4. On Your Own!
 Crumbs:%ROOT%, Programming, Programming/Tutorials, Programming/Tutorials/VariablesTimersEvents
 Description:Binding input to C++ code
+version:4.9
 
 [VAR:Steps]
 [OBJECT:Navigation]
@@ -11,7 +12,7 @@ Description:Binding input to C++ code
 	[PARAM:current]
 	[/PARAM]
 	[PARAM:home]
-		[Programming Tutorials Home](Programming\Tutorials\)
+		[](Programming\Tutorials\VariablesTimersEvents)
 	[/PARAM]
 	[PARAM:next]
 	[/PARAM]
@@ -22,14 +23,14 @@ Description:Binding input to C++ code
 
 Using what you have learned, try to do the following:
 
-* Create an **Actor** that moves or rotates to a target transform when an **Event** is run. This could be used as a moving platform or door in a game. Make the **Event** start a **Timer** that will trigger a second **Event** which moves the **Actor** back to its original location. Use exposed variables (via **UPROPERTY**) instead of hard-coded values wherever appropriate.
-* Make a lit torch that burns out (perhaps by deactivating a fiery **Particle System Component**) by using a **Timer** handle and a few custom **Events**. For example, an _AddFuel_ **Event** could extend the burning time of the torch. A _DouseWithWater_ **Event** could shut it off immediately and prevent _AddFuel_ from working in the future. Both of these features can be written without using a **Tick**, simply by modifying a running **Timer** through its handle.
+* Create an **Actor** that moves or rotates to a target transform when an **Event** is run. This could be used as a moving platform or door in a game. Make the Event start a **Timer** that will trigger a second Event which moves the Actor back to its original location. Use exposed variables (via **UPROPERTY**) instead of hard-coded values wherever appropriate.
+* Make a lit torch that burns out (perhaps by deactivating a fiery **Particle System Component**) by using a Timer handle and a few custom Events. For example, an AddFuel Event could extend the burning time of the torch. A DouseWithWater Event could shut it off immediately and prevent AddFuel from working in the future. Both of these features can be written without using a **Tick**, simply by modifying a running Timer through its handle.
 
 As for the specifics covered in this tutorial:
 
-* For more information about **Timers**, try the [%Programming/UnrealArchitecture/Timers:title%](Programming/UnrealArchitecture/Timers) page.
-* For a complete reference using the **UPROPERTY** tag with variables in your classes or structs, look at the [%Programming/UnrealArchitecture/Reference/Properties:title%](Programming/UnrealArchitecture/Reference/Properties) page.
-* To learn more about **UFUNCTIONS** and **Event** creation, check the [%Programming/UnrealArchitecture/Reference/Functions:title%](Programming/UnrealArchitecture/Reference/Functions) page.
+* For more information about Timers, try the [%Programming/UnrealArchitecture/Timers:title%](Programming/UnrealArchitecture/Timers) page.
+* For a complete reference using the UPROPERTY tag with variables in your classes or structs, look at the [%Programming/UnrealArchitecture/Reference/Properties:title%](Programming/UnrealArchitecture/Reference/Properties) page.
+* To learn more about UFUNCTIONS and Event creation, check the [%Programming/UnrealArchitecture/Reference/Functions:title%](Programming/UnrealArchitecture/Reference/Functions) page.
 * For further tutorials, see the [%Programming/Tutorials:title%](Programming/Tutorials) page.
 
 
