#UE4doc #docs: changed bolded notes to regions in accordance to the instructional guide.

diff --git a/Engine/Documentation/Source/Engine/Rendering/Materials/Functions/Reference/Texturing/Texturing.INT.udn b/Engine/Documentation/Source/Engine/Rendering/Materials/Functions/Reference/Texturing/Texturing.INT.udn
--- a/Engine/Documentation/Source/Engine/Rendering/Materials/Functions/Reference/Texturing/Texturing.INT.udn
+++ b/Engine/Documentation/Source/Engine/Rendering/Materials/Functions/Reference/Texturing/Texturing.INT.udn
@@ -5,8 +5,6 @@ Description:A wide variety of functions designed to aid in working with textures
 
 [TOC (start:2 end:3)]
 
-## Overview
-
 The Texturing functions provide for specialized handling of texture-based actions, such as adjusting UVs of a texture, cropping textures, and many others.
 
 ## Texturing Functions
@@ -137,7 +135,9 @@ WorldAlignedNormal takes in a normal map and aligns its texture to the world, ra
 
 ![WorldAlignedNormal](WorldAlignedNormal_Demo.png)
 
-**Notice the texture is still aligned to world coordinates, even though the floor is rotated.**
+[REGION:note]
+Notice the texture is still aligned to world coordinates, even though the floor is rotated.
+[/REGION]
 
 ### WorldAlignedTexture
 
@@ -164,7 +164,9 @@ The WorldAlignedTexture function serves as a way to tile a texture across the su
 
 ![WorldAlignedTexture](WorldAlignedTexture_Demo.png)
 
-**Notice the texture is still aligned to world coordinates, even though the floor is rotated.**
+[REGION:note]
+Notice the texture is still aligned to world coordinates, even though the floor is rotated.
+[/REGION]
 
 ### 3DSandMayaUVCoordinates
 
@@ -238,7 +240,9 @@ Outputs falloff results for the 3 primary world vectors based on the camera angl
 
 ![CameraWorldBlend](CameraWorldBlend_Demo.png)
 
-**Notice that the green texture blends in wherever the surface aligns to the XY plane, with red everywhere else.**
+[REGION:note]
+Notice that the green texture blends in wherever the surface aligns to the XY plane, with red everywhere else.
+[/REGION]
 
 ### LocalSpaceSurfaceMirroring
 
