#UE4doc #docs:
Updated engine version on each page
Changed font settings on some pages

diff --git a/Engine/Documentation/Source/Programming/UnrealArchitecture/Reference/Properties/Specifiers/Interp/PropertySpecifier_Interp.INT.udn b/Engine/Documentation/Source/Programming/UnrealArchitecture/Reference/Properties/Specifiers/Interp/PropertySpecifier_Interp.INT.udn
--- a/Engine/Documentation/Source/Programming/UnrealArchitecture/Reference/Properties/Specifiers/Interp/PropertySpecifier_Interp.INT.udn
+++ b/Engine/Documentation/Source/Programming/UnrealArchitecture/Reference/Properties/Specifiers/Interp/PropertySpecifier_Interp.INT.udn
@@ -2,6 +2,5 @@ Title: Interp
 Crumbs: %ROOT%, Programming, Programming/UnrealArchitecture/Reference, Programming/UnrealArchitecture/Reference/Properties, Programming/UnrealArchitecture/Reference/Properties/Specifiers
 Description: 
 Availability:Public
-Version: 4.5
 
-Indicates that the value can be driven over time by a track in **Matinee**. 
+Indicates that the value can be driven over time by a track in Matinee. 
diff --git a/Engine/Documentation/Source/Programming/UnrealArchitecture/Reference/Properties/Specifiers/Localized/PropertySpecifier_Localized.INT.udn b/Engine/Documentation/Source/Programming/UnrealArchitecture/Reference/Properties/Specifiers/Localized/PropertySpecifier_Localized.INT.udn
--- a/Engine/Documentation/Source/Programming/UnrealArchitecture/Reference/Properties/Specifiers/Localized/PropertySpecifier_Localized.INT.udn
+++ b/Engine/Documentation/Source/Programming/UnrealArchitecture/Reference/Properties/Specifiers/Localized/PropertySpecifier_Localized.INT.udn
@@ -2,6 +2,6 @@ Title: Localized
 Crumbs: %ROOT%, Programming, Programming/UnrealArchitecture/Reference, Programming/UnrealArchitecture/Reference/Properties, Programming/UnrealArchitecture/Reference/Properties/Specifiers
 Description: 
 Availability:Public
-Version: 4.5
+Version: 4.9
 
-The value of this variable will have a localized value defined. Mostly used for strings. Implies `ReadOnly`.
+The value of this variable will have a localized value defined. Mostly used for strings. Implies ReadOnly.
diff --git a/Engine/Documentation/Source/Programming/UnrealArchitecture/Reference/Properties/Specifiers/Native/PropertySpecifier_Native.INT.udn b/Engine/Documentation/Source/Programming/UnrealArchitecture/Reference/Properties/Specifiers/Native/PropertySpecifier_Native.INT.udn
--- a/Engine/Documentation/Source/Programming/UnrealArchitecture/Reference/Properties/Specifiers/Native/PropertySpecifier_Native.INT.udn
+++ b/Engine/Documentation/Source/Programming/UnrealArchitecture/Reference/Properties/Specifiers/Native/PropertySpecifier_Native.INT.udn
@@ -2,6 +2,6 @@ Title: Native
 Crumbs: %ROOT%, Programming, Programming/UnrealArchitecture/Reference, Programming/UnrealArchitecture/Reference/Properties, Programming/UnrealArchitecture/Reference/Properties/Specifiers
 Description: 
 Availability:Public
-Version: 4.5
+Version: 4.9
 
-Property is native: C++ code is responsible for serializing it and exposing to GC.
+Property is native: C++ code is responsible for serializing it and exposing to Garbage Collection.
diff --git a/Engine/Documentation/Source/Programming/UnrealArchitecture/Reference/Properties/Specifiers/NoClear/PropertySpecifier_NoClear.INT.udn b/Engine/Documentation/Source/Programming/UnrealArchitecture/Reference/Properties/Specifiers/NoClear/PropertySpecifier_NoClear.INT.udn
--- a/Engine/Documentation/Source/Programming/UnrealArchitecture/Reference/Properties/Specifiers/NoClear/PropertySpecifier_NoClear.INT.udn
+++ b/Engine/Documentation/Source/Programming/UnrealArchitecture/Reference/Properties/Specifiers/NoClear/PropertySpecifier_NoClear.INT.udn
@@ -2,6 +2,6 @@ Title: NoClear
 Crumbs: %ROOT%, Programming, Programming/UnrealArchitecture/Reference, Programming/UnrealArchitecture/Reference/Properties, Programming/UnrealArchitecture/Reference/Properties/Specifiers
 Description: 
 Availability:Public
-Version: 4.5
+Version: 4.9
 
-Prevents this object reference from being set to None from the editor. Hides clear (and browse) button in the editor.
+Prevents this object reference from being set to none from the editor. Hides clear (and browse) button in the editor.
diff --git a/Engine/Documentation/Source/Programming/UnrealArchitecture/Reference/Properties/Specifiers/NoExport/PropertySpecifier_NoExport.INT.udn b/Engine/Documentation/Source/Programming/UnrealArchitecture/Reference/Properties/Specifiers/NoExport/PropertySpecifier_NoExport.INT.udn
--- a/Engine/Documentation/Source/Programming/UnrealArchitecture/Reference/Properties/Specifiers/NoExport/PropertySpecifier_NoExport.INT.udn
+++ b/Engine/Documentation/Source/Programming/UnrealArchitecture/Reference/Properties/Specifiers/NoExport/PropertySpecifier_NoExport.INT.udn
@@ -2,6 +2,6 @@ Title: NoExport
 Crumbs: %ROOT%, Programming, Programming/UnrealArchitecture/Reference, Programming/UnrealArchitecture/Reference/Properties, Programming/UnrealArchitecture/Reference/Properties/Specifiers
 Description: 
 Availability:Public
-Version: 4.5
+Version: 4.9
 
 Only useful for native classes. This variable should not be included in the auto-generated class declaration. 
diff --git a/Engine/Documentation/Source/Programming/UnrealArchitecture/Reference/Properties/Specifiers/NonTransactional/PropertySpecifier_NonTransactional.INT.udn b/Engine/Documentation/Source/Programming/UnrealArchitecture/Reference/Properties/Specifiers/NonTransactional/PropertySpecifier_NonTransactional.INT.udn
--- a/Engine/Documentation/Source/Programming/UnrealArchitecture/Reference/Properties/Specifiers/NonTransactional/PropertySpecifier_NonTransactional.INT.udn
+++ b/Engine/Documentation/Source/Programming/UnrealArchitecture/Reference/Properties/Specifiers/NonTransactional/PropertySpecifier_NonTransactional.INT.udn
@@ -2,6 +2,6 @@ Title: NonTransactional
 Crumbs: %ROOT%, Programming, Programming/UnrealArchitecture/Reference, Programming/UnrealArchitecture/Reference/Properties, Programming/UnrealArchitecture/Reference/Properties/Specifiers
 Description: 
 Availability:Public
-Version: 4.5
+Version: 4.9
 
 Indicates that changes to this variable value will not be included in the editor's undo/redo history. 
diff --git a/Engine/Documentation/Source/Programming/UnrealArchitecture/Reference/Properties/Specifiers/Ref/PropertySpecifier_Ref.INT.udn b/Engine/Documentation/Source/Programming/UnrealArchitecture/Reference/Properties/Specifiers/Ref/PropertySpecifier_Ref.INT.udn
--- a/Engine/Documentation/Source/Programming/UnrealArchitecture/Reference/Properties/Specifiers/Ref/PropertySpecifier_Ref.INT.udn
+++ b/Engine/Documentation/Source/Programming/UnrealArchitecture/Reference/Properties/Specifiers/Ref/PropertySpecifier_Ref.INT.udn
@@ -2,6 +2,6 @@ Title: Ref
 Crumbs: %ROOT%, Programming, Programming/UnrealArchitecture/Reference, Programming/UnrealArchitecture/Reference/Properties, Programming/UnrealArchitecture/Reference/Properties/Specifiers
 Description: 
 Availability:Public
-Version: 4.5
+Version: 4.9
 
 Value is copied out after function call. Only valid on function param declaration.
diff --git a/Engine/Documentation/Source/Programming/UnrealArchitecture/Reference/Properties/Specifiers/RepRetry/PropertySpecifier_RepRetry.INT.udn b/Engine/Documentation/Source/Programming/UnrealArchitecture/Reference/Properties/Specifiers/RepRetry/PropertySpecifier_RepRetry.INT.udn
--- a/Engine/Documentation/Source/Programming/UnrealArchitecture/Reference/Properties/Specifiers/RepRetry/PropertySpecifier_RepRetry.INT.udn
+++ b/Engine/Documentation/Source/Programming/UnrealArchitecture/Reference/Properties/Specifiers/RepRetry/PropertySpecifier_RepRetry.INT.udn
@@ -2,6 +2,6 @@ Title: RepRetry
 Crumbs: %ROOT%, Programming, Programming/UnrealArchitecture/Reference, Programming/UnrealArchitecture/Reference/Properties, Programming/UnrealArchitecture/Reference/Properties/Specifiers
 Description: 
 Availability:Public
-Version: 4.5
+Version: 4.9
 
  Only useful for struct properties. Retry replication of this property if it fails to be fully sent (e.g. object references not yet available to serialize over the network). For simple references, this is the default, but for structs, this is often undesirable due to the bandwidth cost, so it is disabled unless this flag is specified. 
diff --git a/Engine/Documentation/Source/Programming/UnrealArchitecture/Reference/Properties/Specifiers/Replicated/PropertySpecifier_Replicated.INT.udn b/Engine/Documentation/Source/Programming/UnrealArchitecture/Reference/Properties/Specifiers/Replicated/PropertySpecifier_Replicated.INT.udn
--- a/Engine/Documentation/Source/Programming/UnrealArchitecture/Reference/Properties/Specifiers/Replicated/PropertySpecifier_Replicated.INT.udn
+++ b/Engine/Documentation/Source/Programming/UnrealArchitecture/Reference/Properties/Specifiers/Replicated/PropertySpecifier_Replicated.INT.udn
@@ -2,6 +2,6 @@ Title: Replicated
 Crumbs: %ROOT%, Programming, Programming/UnrealArchitecture/Reference, Programming/UnrealArchitecture/Reference/Properties, Programming/UnrealArchitecture/Reference/Properties/Specifiers
 Description: 
 Availability:Public
-Version: 4.5
+Version: 4.9
 
 The variable should be replicated over the network.
diff --git a/Engine/Documentation/Source/Programming/UnrealArchitecture/Reference/Properties/Specifiers/ReplicatedUsing/PropertySpecifier_ReplicatedUsing.INT.udn b/Engine/Documentation/Source/Programming/UnrealArchitecture/Reference/Properties/Specifiers/ReplicatedUsing/PropertySpecifier_ReplicatedUsing.INT.udn
--- a/Engine/Documentation/Source/Programming/UnrealArchitecture/Reference/Properties/Specifiers/ReplicatedUsing/PropertySpecifier_ReplicatedUsing.INT.udn
+++ b/Engine/Documentation/Source/Programming/UnrealArchitecture/Reference/Properties/Specifiers/ReplicatedUsing/PropertySpecifier_ReplicatedUsing.INT.udn
@@ -2,8 +2,8 @@ Title: ReplicatedUsing
 Crumbs: %ROOT%, Programming, Programming/UnrealArchitecture/Reference, Programming/UnrealArchitecture/Reference/Properties, Programming/UnrealArchitecture/Reference/Properties/Specifiers
 Description: 
 Availability:Public
-Version: 4.5
+Version: 4.9
 
-The variable should be replicated over the network, executing the `Callback` function when it is received.
+The variable should be replicated over the network, executing the **Callback** function when it is received.
 
 * Usage: ReplicatedUsing=FunctionName
diff --git a/Engine/Documentation/Source/Programming/UnrealArchitecture/Reference/Properties/Specifiers/SaveGame/PropertySpecifier_SaveGame.INT.udn b/Engine/Documentation/Source/Programming/UnrealArchitecture/Reference/Properties/Specifiers/SaveGame/PropertySpecifier_SaveGame.INT.udn
--- a/Engine/Documentation/Source/Programming/UnrealArchitecture/Reference/Properties/Specifiers/SaveGame/PropertySpecifier_SaveGame.INT.udn
+++ b/Engine/Documentation/Source/Programming/UnrealArchitecture/Reference/Properties/Specifiers/SaveGame/PropertySpecifier_SaveGame.INT.udn
@@ -2,7 +2,7 @@ Title: SaveGame
 Crumbs: %ROOT%, Programming, Programming/UnrealArchitecture/Reference, Programming/UnrealArchitecture/Reference/Properties, Programming/UnrealArchitecture/Reference/Properties/Specifiers
 Description: 
 Availability:Public
-Version: 4.5
+Version: 4.9
 
-The SaveGame specifier is a simple way to include fields explicitly for a checkpoint/save system at the property level. You set the flag on all fields you want to save and then use a proxy archiver to read/write
+The **SaveGame** specifier is a simple way to include fields explicitly for a checkpoint/save system at the property level. You set the flag on all fields you want to save and then use a proxy archiver to read/write
 it. 
diff --git a/Engine/Documentation/Source/Programming/UnrealArchitecture/Reference/Properties/Specifiers/SerializeText/PropertySpecifier_SerializeText.INT.udn b/Engine/Documentation/Source/Programming/UnrealArchitecture/Reference/Properties/Specifiers/SerializeText/PropertySpecifier_SerializeText.INT.udn
--- a/Engine/Documentation/Source/Programming/UnrealArchitecture/Reference/Properties/Specifiers/SerializeText/PropertySpecifier_SerializeText.INT.udn
+++ b/Engine/Documentation/Source/Programming/UnrealArchitecture/Reference/Properties/Specifiers/SerializeText/PropertySpecifier_SerializeText.INT.udn
@@ -2,6 +2,6 @@ Title: SerializeText
 Crumbs: %ROOT%, Programming, Programming/UnrealArchitecture/Reference, Programming/UnrealArchitecture/Reference/Properties, Programming/UnrealArchitecture/Reference/Properties/Specifiers
 Description: 
 Availability:Public
-Version: 4.5
+Version: 4.9
 
 Native property should be serialized as text (ImportText, ExportText).
diff --git a/Engine/Documentation/Source/Programming/UnrealArchitecture/Reference/Properties/Specifiers/SimpleDisplay/PropertySpecifier_SimpleDisplay.INT.udn b/Engine/Documentation/Source/Programming/UnrealArchitecture/Reference/Properties/Specifiers/SimpleDisplay/PropertySpecifier_SimpleDisplay.INT.udn
--- a/Engine/Documentation/Source/Programming/UnrealArchitecture/Reference/Properties/Specifiers/SimpleDisplay/PropertySpecifier_SimpleDisplay.INT.udn
+++ b/Engine/Documentation/Source/Programming/UnrealArchitecture/Reference/Properties/Specifiers/SimpleDisplay/PropertySpecifier_SimpleDisplay.INT.udn
@@ -2,6 +2,6 @@ Title: SimpleDisplay
 Crumbs: %ROOT%, Programming, Programming/UnrealArchitecture/Reference, Programming/UnrealArchitecture/Reference/Properties, Programming/UnrealArchitecture/Reference/Properties/Specifiers
 Description: 
 Availability:Public
-Version: 4.5
+Version: 4.9
 
-Properties appear visible by default in a **Details** panel.
+Properties appear visible by default in the **Details** panel.
diff --git a/Engine/Documentation/Source/Programming/UnrealArchitecture/Reference/Properties/Specifiers/Transient/PropertySpecifier_Transient.INT.udn b/Engine/Documentation/Source/Programming/UnrealArchitecture/Reference/Properties/Specifiers/Transient/PropertySpecifier_Transient.INT.udn
--- a/Engine/Documentation/Source/Programming/UnrealArchitecture/Reference/Properties/Specifiers/Transient/PropertySpecifier_Transient.INT.udn
+++ b/Engine/Documentation/Source/Programming/UnrealArchitecture/Reference/Properties/Specifiers/Transient/PropertySpecifier_Transient.INT.udn
@@ -2,6 +2,6 @@ Title: Transient
 Crumbs: %ROOT%, Programming, Programming/UnrealArchitecture/Reference, Programming/UnrealArchitecture/Reference/Properties, Programming/UnrealArchitecture/Reference/Properties/Specifiers
 Description: 
 Availability:Public
-Version: 4.5
+Version: 4.9
 
 Property is transient: should not be saved, zero-filled at load time.
diff --git a/Engine/Documentation/Source/Programming/UnrealArchitecture/Reference/Properties/Specifiers/VisibleAnywhere/PropertySpecifier_VisibleAnywhere.INT.udn b/Engine/Documentation/Source/Programming/UnrealArchitecture/Reference/Properties/Specifiers/VisibleAnywhere/PropertySpecifier_VisibleAnywhere.INT.udn
--- a/Engine/Documentation/Source/Programming/UnrealArchitecture/Reference/Properties/Specifiers/VisibleAnywhere/PropertySpecifier_VisibleAnywhere.INT.udn
+++ b/Engine/Documentation/Source/Programming/UnrealArchitecture/Reference/Properties/Specifiers/VisibleAnywhere/PropertySpecifier_VisibleAnywhere.INT.udn
@@ -2,6 +2,6 @@ Title: VisibleAnywhere
 Crumbs: %ROOT%, Programming, Programming/UnrealArchitecture/Reference, Programming/UnrealArchitecture/Reference/Properties, Programming/UnrealArchitecture/Reference/Properties/Specifiers
 Description: 
 Availability:Public
-Version: 4.5
+Version: 4.9
 
 Indicates that this property is visible in property windows, but cannot be edited at all.
diff --git a/Engine/Documentation/Source/Programming/UnrealArchitecture/Reference/Properties/Specifiers/VisibleDefaultsOnly/PropertySpecifier_VisibleDefaultsOnly.INT.udn b/Engine/Documentation/Source/Programming/UnrealArchitecture/Reference/Properties/Specifiers/VisibleDefaultsOnly/PropertySpecifier_VisibleDefaultsOnly.INT.udn
--- a/Engine/Documentation/Source/Programming/UnrealArchitecture/Reference/Properties/Specifiers/VisibleDefaultsOnly/PropertySpecifier_VisibleDefaultsOnly.INT.udn
+++ b/Engine/Documentation/Source/Programming/UnrealArchitecture/Reference/Properties/Specifiers/VisibleDefaultsOnly/PropertySpecifier_VisibleDefaultsOnly.INT.udn
@@ -2,6 +2,6 @@ Title: VisibleDefaultsOnly
 Crumbs: %ROOT%, Programming, Programming/UnrealArchitecture/Reference, Programming/UnrealArchitecture/Reference/Properties, Programming/UnrealArchitecture/Reference/Properties/Specifiers
 Description: 
 Availability:Public
-Version: 4.5
+Version: 4.9
 
 Indicates that this property is only visible in property windows for archetypes, and cannot be edited.
diff --git a/Engine/Documentation/Source/Programming/UnrealArchitecture/Reference/Properties/Specifiers/VisibleInstanceOnly/PropertySpecifier_VisibleInstanceOnly.INT.udn b/Engine/Documentation/Source/Programming/UnrealArchitecture/Reference/Properties/Specifiers/VisibleInstanceOnly/PropertySpecifier_VisibleInstanceOnly.INT.udn
--- a/Engine/Documentation/Source/Programming/UnrealArchitecture/Reference/Properties/Specifiers/VisibleInstanceOnly/PropertySpecifier_VisibleInstanceOnly.INT.udn
+++ b/Engine/Documentation/Source/Programming/UnrealArchitecture/Reference/Properties/Specifiers/VisibleInstanceOnly/PropertySpecifier_VisibleInstanceOnly.INT.udn
@@ -2,6 +2,6 @@ Title: VisibleInstanceOnly
 Crumbs: %ROOT%, Programming, Programming/UnrealArchitecture/Reference, Programming/UnrealArchitecture/Reference/Properties, Programming/UnrealArchitecture/Reference/Properties/Specifiers
 Description: 
 Availability:Public
-Version: 4.5
+Version: 4.9
 
 Indicates that this property is only visible in property windows for instances, not for archetypes, and cannot be edited.
diff --git a/Engine/Documentation/Source/Programming/UnrealArchitecture/Reference/Structs/Specifiers/Atomic/StructSpecifier_Atomic.INT.udn b/Engine/Documentation/Source/Programming/UnrealArchitecture/Reference/Structs/Specifiers/Atomic/StructSpecifier_Atomic.INT.udn
--- a/Engine/Documentation/Source/Programming/UnrealArchitecture/Reference/Structs/Specifiers/Atomic/StructSpecifier_Atomic.INT.udn
+++ b/Engine/Documentation/Source/Programming/UnrealArchitecture/Reference/Structs/Specifiers/Atomic/StructSpecifier_Atomic.INT.udn
@@ -2,6 +2,6 @@ Title: Atomic
 Crumbs:%ROOT%, Programming,  Programming/UnrealArchitecture/Reference,  Programming/UnrealArchitecture/Reference/Structs, Programming/UnrealArchitecture/Reference/Structs/Specifiers
 Description: 
 Availability:Public
-Version: 4.5
+Version: 4.9
 
 Indicates that this struct should always be serialized as a single unit.
diff --git a/Engine/Documentation/Source/Programming/UnrealArchitecture/Reference/Structs/Specifiers/BlueprintType/StructSpecifier_BlueprintType.INT.udn b/Engine/Documentation/Source/Programming/UnrealArchitecture/Reference/Structs/Specifiers/BlueprintType/StructSpecifier_BlueprintType.INT.udn
--- a/Engine/Documentation/Source/Programming/UnrealArchitecture/Reference/Structs/Specifiers/BlueprintType/StructSpecifier_BlueprintType.INT.udn
+++ b/Engine/Documentation/Source/Programming/UnrealArchitecture/Reference/Structs/Specifiers/BlueprintType/StructSpecifier_BlueprintType.INT.udn
@@ -2,6 +2,6 @@ Title: BlueprintType
 Crumbs:%ROOT%, Programming,  Programming/UnrealArchitecture/Reference,  Programming/UnrealArchitecture/Reference/Structs, Programming/UnrealArchitecture/Reference/Structs/Specifiers
 Description: 
 Availability:Public
-Version: 4.5
+Version: 4.9
 
 Exposes this struct as a type that can be used for variables in Blueprints.
diff --git a/Engine/Documentation/Source/Programming/UnrealArchitecture/Reference/Structs/Specifiers/Immutable/StructSpecifier_Immutable.INT.udn b/Engine/Documentation/Source/Programming/UnrealArchitecture/Reference/Structs/Specifiers/Immutable/StructSpecifier_Immutable.INT.udn
--- a/Engine/Documentation/Source/Programming/UnrealArchitecture/Reference/Structs/Specifiers/Immutable/StructSpecifier_Immutable.INT.udn
+++ b/Engine/Documentation/Source/Programming/UnrealArchitecture/Reference/Structs/Specifiers/Immutable/StructSpecifier_Immutable.INT.udn
@@ -2,6 +2,6 @@ Title: Immutable
 Crumbs:%ROOT%, Programming,  Programming/UnrealArchitecture/Reference,  Programming/UnrealArchitecture/Reference/Structs, Programming/UnrealArchitecture/Reference/Structs/Specifiers
 Description: 
 Availability:Public
-Version: 4.5
+Version: 4.9
 
-Immutable is only legal in `Object.h` and is being phased out, do not use on new structs!
+Immutable is only legal in Object.h and is being phased out, do not use on new structs!
diff --git a/Engine/Documentation/Source/Programming/UnrealArchitecture/Reference/Structs/Specifiers/NoExport/StructSpecifier_NoExport.INT.udn b/Engine/Documentation/Source/Programming/UnrealArchitecture/Reference/Structs/Specifiers/NoExport/StructSpecifier_NoExport.INT.udn
--- a/Engine/Documentation/Source/Programming/UnrealArchitecture/Reference/Structs/Specifiers/NoExport/StructSpecifier_NoExport.INT.udn
+++ b/Engine/Documentation/Source/Programming/UnrealArchitecture/Reference/Structs/Specifiers/NoExport/StructSpecifier_NoExport.INT.udn
@@ -2,6 +2,6 @@ Title: NoExport
 Crumbs:%ROOT%, Programming,  Programming/UnrealArchitecture/Reference,  Programming/UnrealArchitecture/Reference/Structs, Programming/UnrealArchitecture/Reference/Structs/Specifiers
 Description: 
 Availability:Public
-Version: 4.5
+Version: 4.9
 
  No autogenerated code will be created for this class; the header is only provided to parse metadata from.
diff --git a/Engine/Documentation/Source/Programming/UnrealArchitecture/Reference/Structs/Structs.INT.udn b/Engine/Documentation/Source/Programming/UnrealArchitecture/Reference/Structs/Structs.INT.udn
--- a/Engine/Documentation/Source/Programming/UnrealArchitecture/Reference/Structs/Structs.INT.udn
+++ b/Engine/Documentation/Source/Programming/UnrealArchitecture/Reference/Structs/Structs.INT.udn
@@ -2,9 +2,9 @@ Availability:Public
 Title:Structs
 Crumbs:%ROOT%, Programming, Programming/UnrealArchitecture/Reference
 Description:Reference to creating and implementing structs for gameplay classes.
-Version: 4.5
+Version: 4.9
 
-**Structs** in gameplay classes can contain variables, including `UProperty` variables, functions, and operators. Structs are declared prior to the class declaration.
+**Structs** in gameplay classes can contain variables, including UProperty variables, functions, and operators. Structs are declared prior to the class declaration.
 
 	USTRUCT([Specifier, Specifier, ...])
 	struct StructName
