#UE4doc #docs:
Updated engine version
Changed font settings for multiple terms

diff --git a/Engine/Documentation/Source/Programming/UnrealArchitecture/SmartPointerLibrary/SmartPointerLibrary.INT.udn b/Engine/Documentation/Source/Programming/UnrealArchitecture/SmartPointerLibrary/SmartPointerLibrary.INT.udn
--- a/Engine/Documentation/Source/Programming/UnrealArchitecture/SmartPointerLibrary/SmartPointerLibrary.INT.udn
+++ b/Engine/Documentation/Source/Programming/UnrealArchitecture/SmartPointerLibrary/SmartPointerLibrary.INT.udn
@@ -2,20 +2,20 @@ Availability:Public
 Title:Unreal Smart Pointer Library
 Crumbs:%ROOT%, Programming, Programming/UnrealArchitecture
 Description:Custom implementation of shared pointers, including weak pointers and non-nullable shared references.
-Version: 4.5
+Version: 4.9
 
 [TOC(start:2)]
 
 ## Smart Pointers
 
 [EXCERPT:Overview]
-The Unreal Smart Pointer Library is a custom implementation of shared references (`TSharedRef`), shared pointers (`TSharedPtr`), weak pointers (`TWeakPtr`) as well as related helper functions and classes.  This implementation is modeled after the C++0x standard library's shared_ptr as well as Boost smart pointers.
+The **Unreal Smart Pointer Library** is a custom implementation of shared references (**TSharedRef**), shared pointers (**TSharedPtr**), weak pointers (**TWeakPtr**) as well as related helper functions and classes.  This implementation is modeled after the C++0x standard library's shared_ptr as well as Boost smart pointers.
 
 | Type | Description |
 | --- | --- |
-| [](Programming/UnrealArchitecture/SmartPointerLibrary/SharedPointer) (`TSharedPtr`) | Reference counted non-intrusive authoritative smart pointer. |
-| [](Programming/UnrealArchitecture/SmartPointerLibrary/SharedReference) (`TSharedRef`) | Non-nullable, reference counted non-intrusive authoritative smart pointer. |
-| [](Programming/UnrealArchitecture/SmartPointerLibrary/WeakPointer) (`TWeakPtr`) | Reference counted non-intrusive weak pointer reference. |
+| [](Programming/UnrealArchitecture/SmartPointerLibrary/SharedPointer) (TSharedPtr) | Reference counted non-intrusive authoritative smart pointer. |
+| [](Programming/UnrealArchitecture/SmartPointerLibrary/SharedReference) (TSharedRef) | Non-nullable, reference counted non-intrusive authoritative smart pointer. |
+| [](Programming/UnrealArchitecture/SmartPointerLibrary/WeakPointer) (TWeakPtr) | Reference counted non-intrusive weak pointer reference. |
 [/EXCERPT:Overview]
 
 ### Benefits of Shared References and Pointers
@@ -25,13 +25,13 @@ The Unreal Smart Pointer Library is a custom implementation of shared references
 | Clean syntax | You can copy, dereference, and compare shared pointers just like regular C++ pointers. |
 | Prevents memory leaks | Resources are destroyed automatically when there are no more shared references. |
 | Weak referencing | Weak pointers allow you to safely check when an object has been destroyed. |
-| Thread safety | Includes "thread safe" version that can be safely accessed from multiple threads. |
-| Ubiquitous | You can create shared pointers to virtually **any** type of object. |
+| Thread safety | Includes thread safe version that can be safely accessed from multiple threads. |
+| Ubiquitous | You can create shared pointers to virtually any type of object. |
 | Runtime safety | Shared references are never null and can always be de-referenced. |
 | No reference cycles | Use weak pointers to break reference cycles. |
-| Confers intent | You can easily tell an object **owner** from an **observer**. |
+| Confers intent | You can easily tell an object owner from an observer. |
 | Performance | Shared pointers have minimal overhead.  All operations are constant-time. |
-| Robust features | Supports 'const', forward declarations to incomplete types, type-casting, etc. |
+| Robust features | Supports const, forward declarations to incomplete types, type-casting, etc. |
 | Memory | Only twice the size of a C++ pointer in 64-bit (plus a shared 16-byte reference controller.) |
 
 
@@ -42,7 +42,7 @@ The Unreal Smart Pointer Library is a custom implementation of shared references
 * Can work seamlessly with other Unreal containers and types.
 * Better control over platform specifics, including threading and optimizations.
 * We want thread-safety features to be optional (for performance).
-* We have added our own improvements (MakeShareable, assign to NULL, etc.).
+* We have added our own improvements (**MakeShareable**, assign to NULL, etc.).
 * Exceptions were not needed nor desired in our implementation.
 * We wanted more control over performance (inlining, memory, use of virtuals, etc.).
 * Potentially easier to debug (liberal code comments, etc.).
@@ -95,14 +95,14 @@ Performance drawbacks of shared pointers include:
 
 ### Memory Usage
 
-All shared pointers (`TSharedPtr`, `TSharedRef`, `TWeakPtr`) are 8 bytes (when compiling for 32-bit) consisting of:
+All shared pointers (TSharedPtr, TSharedRef, TWeakPtr) are 8 bytes (when compiling for 32-bit) consisting of:
 
 * C++ pointer (uint32)
 * Reference controller pointer (uint32)
 
 
 [REGION:note]
-`TSharedFromThis` is also 8 bytes as it embeds weak pointer.
+TSharedFromThis is also 8 bytes as it embeds weak pointer.
 [/REGION]
 
 The Reference controller object is 12 bytes (when compiling for 32-bit) consisting of:
@@ -120,7 +120,7 @@ Only one reference controller is created per object, no matter how many shared/w
 
 Shared pointers are non-intrusive which means the class itself does not know whether it is owned by a shared pointer (or reference). In general, this is okay, but sometimes you really want to access the current instance as a shared reference. The solution to this is to derive the class from `TSharedFromThis<>`.
 
-By deriving from `TSharedFromThis<>`, you can use the `AsSharedRef()` method to convert 'this' to a shared reference. This can be very useful with class factories that always return shared references.
+By deriving from TSharedFromThis<>, you can use the **AsSharedRef()** method to convert 'this' to a shared reference. This can be very useful with class factories that always return shared references.
 
 	class FAnimation : public TSharedFromThis<FMyClass>
 	{
@@ -172,23 +172,23 @@ If you know your pointer will never be accessed by more than one thread, do not
 
 ## Usage Details
 
-The `SharedPointerTesting.h` file (located in `[UE4RootLocation]/Engine/Source/Runtime/Core/Public/Templates/`) contains various examples of using shared pointers and references.
+The SharedPointerTesting.h file (located in [UE4RootLocation]/Engine/Source/Runtime/Core/Public/Templates/) contains various examples of using shared pointers and references.
 
 ### Tips
 
 * Usually you should "operator new" when passing a C++ pointer to a new shared pointer.
-* Use `TSharedRef` or `TSharedPtr` when passing smart pointers as function parameters, not TWeakPtr.
+* Use TSharedRef or TSharedPtr when passing smart pointers as function parameters, not TWeakPtr.
 * The "thread-safe" versions of smart pointers are a bit slower -- only use them when needed.
 * You can forward declare shared pointers to incomplete types, just how you would expect to!
 * Shared pointers of compatible types will be converted implicitly (e.g. upcasting).
-* You can create a typedef to `TSharedRef< MyClass >` to make it easier to type.
-* For best performance, minimize calls to `TWeakPtr::Pin` (or conversions to `TSharedRef`/`TSharedPtr`).
-* Your class can return itself as a shared reference if you derive from `TSharedFromThis`.
-* To downcast a pointer to a derived object class, use the `StaticCastSharedPtr()` function.
-* `const` objects are fully supported with shared pointers!
-* You can make a `const` shared pointer mutable using the `ConstCastSharedPtr()` function.
+* You can create a typedef to TSharedRef< MyClass > to make it easier to type.
+* For best performance, minimize calls to TWeakPtr::Pin (or conversions to TSharedRef/TSharedPtr).
+* Your class can return itself as a shared reference if you derive from TSharedFromThis.
+* To downcast a pointer to a derived object class, use the **StaticCastSharedPtr()** function.
+* const objects are fully supported with shared pointers!
+* You can make a const shared pointer mutable using the **ConstCastSharedPtr()** function.
 * Always convert to C++ references in deep stack frames. Shared pointers are best for member references, not stack temporaries.
-* Unlike C++ pointers, shared pointers cannot be `memcpy`'d, so be sure to consider this when using arrays of shared pointers.
+* Unlike C++ pointers, shared pointers cannot be memcpy'd, so be sure to consider this when using arrays of shared pointers.
 
 
 ### Limitations
@@ -196,17 +196,17 @@ The `SharedPointerTesting.h` file (located in `[UE4RootLocation]/Engine/Source/R
 * Shared pointers are not compatible with Unreal objects (UObject classes)!
 * Currently only types with that have regular destructors (no custom deleters).
 * Dynamically-allocated arrays are not supported yet (e.g. `MakeSharable( new int32[20] )`).
-* Implicit conversion of `TSharedPtr`/`TSharedRef` to bool is not supported yet.
+* Implicit conversion of TSharedPtr/TSharedRef to bool is not supported yet.
 
 
 ### Differences from Other Implementations
 
 * Type names and method names are more consistent with Unreal's codebase.
 * Thread-safety features are optional instead of forced.
-* `TSharedFromThis` returns a shared **reference**, not a shared **pointer**.
-* Some features were omitted (e.g. `use_count()`, `unique()`, etc.).
+* TSharedFromThis returns a shared reference, not a shared pointer.
+* Some features were omitted (e.g. use_count(), unique(), etc.).
 * No exceptions are allowed (all related features have been omitted).
 * Custom allocators and custom delete functions are not supported yet.
-* Our implementation supports non-nullable smart pointers (`TSharedRef`).
-* Several other new features added, such as `MakeShareable` and `NULL` assignment.
+* Our implementation supports non-nullable smart pointers (TSharedRef).
+* Several other new features added, such as MakeShareable and NULL assignment.
 
