#UE4doc #docs:
Updated engine version
Changed font settings for multiple terms
Changed all mentions of NULL to null to reflect changes made to the engine.

diff --git a/Engine/Documentation/Source/Programming/UnrealArchitecture/SmartPointerLibrary/SharedPointer/SharedPointer.INT.udn b/Engine/Documentation/Source/Programming/UnrealArchitecture/SmartPointerLibrary/SharedPointer/SharedPointer.INT.udn
--- a/Engine/Documentation/Source/Programming/UnrealArchitecture/SmartPointerLibrary/SharedPointer/SharedPointer.INT.udn
+++ b/Engine/Documentation/Source/Programming/UnrealArchitecture/SmartPointerLibrary/SharedPointer/SharedPointer.INT.udn
@@ -2,7 +2,7 @@ Availability:Public
 Title:Shared Pointers
 Crumbs:%ROOT%, Programming, Programming/UnrealArchitecture/SmartPointerLibrary
 Description:Smart pointers that support shared ownership, automatic invalidation, weak references, and more.
-Version: 4.5
+Version: 4.9
 
 [TOC(start:2)]
 
@@ -42,10 +42,10 @@ When you copy a shared pointer, a reference to the object being pointed to is ad
 
 The object will persist until no more shared pointers are referencing it.
 
-Shared pointers can be reset using the `Reset()` method or by assigning NULL.
+Shared pointers can be reset using the **Reset()** method or by assigning null.
 
 	Node.Reset();
-	Node = NULL;
+	Node = null;
 
 ## Dereferencing and Accessing
 
@@ -66,14 +66,14 @@ Shared pointers can be tested for equality, i.e. whether they point to the same
 		...
 	}
 
-Shared pointers can be tested for NULL by using the `IsValid()` method or accessing the raw pointer and performing a comparison.
+Shared pointers can be tested for null by using the **IsValid()** method or accessing the raw pointer and performing a comparison.
 
 	if( Node.IsValid() )
 	{
 		...
 	}
 
-	if( Node.Get() != NULL )
+	if( Node.Get() != null )
 	{
 		...
 	}
@@ -84,7 +84,7 @@ Implicit conversion is not allowed as it is extremely dangerous. For example, th
 
 	TSharedPtr<float> Height = 45.0f;
 
-However, the `MakeShareable()` method can be used to be explicit.
+However, the MakeShareable() method can be used to be explicit.
 
 	TSharedPtr<float> Height = MakeShareable( new float( 45.0f ) );
 
