#UE4doc #docs:
Updated engine version
Changed font settings for multiple terms
Changed all mentions of NULL to null to reflect changes made to the engine.

diff --git a/Engine/Documentation/Source/Programming/UnrealArchitecture/SmartPointerLibrary/SharedReference/SharedReference.INT.udn b/Engine/Documentation/Source/Programming/UnrealArchitecture/SmartPointerLibrary/SharedReference/SharedReference.INT.udn
--- a/Engine/Documentation/Source/Programming/UnrealArchitecture/SmartPointerLibrary/SharedReference/SharedReference.INT.udn
+++ b/Engine/Documentation/Source/Programming/UnrealArchitecture/SmartPointerLibrary/SharedReference/SharedReference.INT.udn
@@ -2,11 +2,11 @@ Availability:Public
 Title:Shared References
 Crumbs:%ROOT%, Programming, Programming/UnrealArchitecture/SmartPointerLibrary
 Description:Smart pointer type that cannot be uninitialized or assigned null.
-Version: 4.5
+Version: 4.9
 
 [TOC(start:2)]
 
-A **shared reference** is a non-nullable shared pointer. This means you cannot create an empty shared reference and you cannot assign null to a shared reference, which helps you write safer, cleaner code since you know it is safe to access the object. In fact, shared references do not even have an `IsValid()` method - like shared pointers do - since they are always valid. It is recommended to always use shared references instead of shared pointers when possible.
+A **shared reference** is a non-nullable shared pointer. This means you cannot create an empty shared reference and you cannot assign null to a shared reference, which helps you write safer, cleaner code since you know it is safe to access the object. In fact, shared references do not even have an **IsValid()** method - like shared pointers do - since they are always valid. It is recommended to always use shared references instead of shared pointers when possible.
 
 If you need empty/nullable references, you should use [](Programming/UnrealArchitecture/SmartPointerLibrary/SharedPointer) instead.
 
@@ -19,7 +19,7 @@ An example of creating a new shared reference:
 	// Create a shared reference to a new node
 	TSharedRef<FTreeNode> NodeRef( new FTreeNode() );
 
-Attempting to declare a shared reference without initializing it or assigning NULL is never allowed:
+Attempting to declare a shared reference without initializing it or assigning null is never allowed:
 
 	TSharedRef<UBool> EmptyBool;			// Will not Compile!!
 
@@ -34,9 +34,9 @@ For the most part, using a shared reference is just like using a shared pointer.
 	( *NodeRef ).ListChildren();
 	const FTreeNode& NodeVar = NodeRef.Get();
 
-Again, the main difference is that you cannot set a shared pointer to NULL:
+Again, the main difference is that you cannot set a shared pointer to null:
 
-	TSharedRef<FTreeNode> Node = NULL;		// Will not Compile!!
+	TSharedRef<FTreeNode> Node = null;		// Will not Compile!!
 
 ### Pointer and Reference Conversion
 
@@ -44,11 +44,11 @@ Shared references convert to shared pointers implicitly so it is always safe to
 
 	TSharedPtr<FTreeNode> SomeNodePtr = NodeRef;
 
-However, conversion from a pointer is potentially unsafe, and you are forced to explicitly perform the conversion using the `ToSharedRef()` function.
+However, conversion from a pointer is potentially unsafe, and you are forced to explicitly perform the conversion using the **ToSharedRef()** function.
 
 	// Explicitly reference a pointer. Will assert if the pointer is not valid!
 	NodeRef = SomeNodePtr.ToSharedRef();
 
 [REGION:note]
-In a similar manner to the `*` operator with C++ pointers, using `ToSharedRef()` will assert if the object is not valid.
+In a similar manner to the * operator with C++ pointers, using ToSharedRef() will assert if the object is not valid.
 [/REGION]
