#UE4doc #docs:
Updated engine version
Changed font settings for multiple terms
Changed all mentions of NULL to null to reflect changes made to the engine.

diff --git a/Engine/Documentation/Source/Programming/UnrealArchitecture/SmartPointerLibrary/WeakPointer/WeakPointer.INT.udn b/Engine/Documentation/Source/Programming/UnrealArchitecture/SmartPointerLibrary/WeakPointer/WeakPointer.INT.udn
--- a/Engine/Documentation/Source/Programming/UnrealArchitecture/SmartPointerLibrary/WeakPointer/WeakPointer.INT.udn
+++ b/Engine/Documentation/Source/Programming/UnrealArchitecture/SmartPointerLibrary/WeakPointer/WeakPointer.INT.udn
@@ -2,11 +2,11 @@ Availability:Public
 Title:Weak Pointers
 Crumbs:%ROOT%, Programming, Programming/UnrealArchitecture/SmartPointerLibrary
 Description:Smart pointers that store weak references and do not prevent their objects from being destroyed.
-Version: 4.5
+Version: 4.9
 
 [TOC(start:2)]
 
-Weak pointers store a _weak reference_ to an object. Unlike a shared pointer, a weak pointer will not prevent an object from being destroyed. Weak pointers are automatically emptied when their object is destroyed no matter who destroyed the object, which allows you to safely cache pointers to volatile objects. This also means, weak pointers might become empty when you are not expecting it and weak pointers can be used to break reference cycles.
+**Weak pointers** store a weak reference to an object. Unlike a shared pointer, a weak pointer will not prevent an object from being destroyed. Weak pointers are automatically emptied when their object is destroyed no matter who destroyed the object, which allows you to safely cache pointers to volatile objects. This also means, weak pointers might become empty when you are not expecting it and weak pointers can be used to break reference cycles.
 
 A weak pointer's object will be destroyed when there are no more shared references to it.
 
@@ -33,11 +33,11 @@ Weak pointers can be copied around safely, just like shared pointers.
 
 You can reset (or reassign) a weak pointer when you are done with it.
 
-	NodeObserver = NULL;
+	NodeObserver = null;
 
 ## Dereferencing and Accessing
 
-To access a weak pointer's object, first promote it to a shared pointer by calling the `Lock()` method.
+To access a weak pointer's object, first promote it to a shared pointer by calling the **Lock()** method.
 
 	// Get access to the node through the weak pointer
 	TSharedPtr<TFreeNode> LockedObserver( NodeObserver.Lock() );
@@ -48,9 +48,9 @@ To access a weak pointer's object, first promote it to a shared pointer by calli
 		LockedObserver->ListChildren();
 	}
 
-`Lock()` is fast. It is called Lock because it prevents the object from being destroyed while the weak pointer is temporarily accessed.
+Lock() is fast. It is called Lock because it prevents the object from being destroyed while the weak pointer is temporarily accessed.
 
-Weak pointers automatically know when their object has been deleted, so you can handle these cases in a safe way by using the `IsValid()` method.
+Weak pointers automatically know when their object has been deleted, so you can handle these cases in a safe way by using the IsValid() method.
 
 	//Make sure the weak pointer's object still exists before accessing it
 	if( NodeObserver.IsValid() )
@@ -71,10 +71,10 @@ You can use weak pointers to break cycles:
 		TWeakPtr<FTreeNode> Parent;
 	}
 
-`FTreeNode` has shared pointers to its child nodes that it _owns_. This means the child nodes will be destroyed with it. It also uses a weak pointer to cache a reference to its parent node. Weak pointers never prevent an object from being destroyed. Using shared/weak pointers this way makes the chain of authority more obvious.
+FTreeNode has shared pointers to its child nodes that it _owns_. This means the child nodes will be destroyed with it. It also uses a weak pointer to cache a reference to its parent node. Weak pointers never prevent an object from being destroyed. Using shared/weak pointers this way makes the chain of authority more obvious.
 
-Weak pointers "just know" when their object is deleted, so you could safely cache pointers to any node in the tree this way.
+Weak pointers just know when their object is deleted, so you could safely cache pointers to any node in the tree this way.
 
 ## Conversion
 
-You must use the `Pin()` method to convert weak pointers to shared pointers. There is no explicit constructor.
+You must use the **Pin()** method to convert weak pointers to shared pointers. There is no explicit constructor.
