#UE4doc #docs - adjusted engine version number metadata

diff --git a/Engine/Documentation/Source/Programming/UnrealArchitecture/CoreArchitecture.INT.udn b/Engine/Documentation/Source/Programming/UnrealArchitecture/CoreArchitecture.INT.udn
--- a/Engine/Documentation/Source/Programming/UnrealArchitecture/CoreArchitecture.INT.udn
+++ b/Engine/Documentation/Source/Programming/UnrealArchitecture/CoreArchitecture.INT.udn
@@ -10,7 +10,7 @@ The goal of this document is to provide a brief tour of the engine's core featur
 
 **[](Programming\UnrealArchitecture\Actors)** are instances of classes that derive from the `AActor` class; the base class of all gameplay objects that can be placed into the world. 
 **[](Programming\UnrealArchitecture\Objects)** are instances of classes that inherit from the `UObject` class; the base class of all objects in Unreal Engine, including Actors. So, in reality, all 
-instances in Unreal Engine are **Objects**; however, the term **Actors** is commonly used to refer to instances of classes that derive from `AActor` in their hierarchy, while the term Objects 
+instances in Unreal Engine are Objects; however, the term **Actors** is commonly used to refer to instances of classes that derive from `AActor` in their hierarchy, while the term Objects 
 is used to refer to instances of classes that do not inherit from the `AActor` class. The majority of the classes you create will inherit from `AActor` at some point in their hierarchy.
 
 In general, Actors can be thought of as whole items or entities, while Objects are more specialized parts. Actors often make use of **[](Programming\UnrealArchitecture\Actors\Components)**, which are specialized Objects, 
diff --git a/Engine/Documentation/Source/Programming/UnrealArchitecture/Reference/Classes/GameplayClasses.INT.udn b/Engine/Documentation/Source/Programming/UnrealArchitecture/Reference/Classes/GameplayClasses.INT.udn
--- a/Engine/Documentation/Source/Programming/UnrealArchitecture/Reference/Classes/GameplayClasses.INT.udn
+++ b/Engine/Documentation/Source/Programming/UnrealArchitecture/Reference/Classes/GameplayClasses.INT.udn
@@ -2,7 +2,7 @@ Availability:Public
 Title: Gameplay Classes
 Crumbs:%ROOT%, Programming, Programming/UnrealArchitecture/Reference
 Description:Reference to creating and implementing gameplay classes.
-Version: 4.5
+Version: 4.9
 
 [TOC(start:2)]
 
diff --git a/Engine/Documentation/Source/Programming/UnrealArchitecture/Reference/Classes/Specifiers/Abstract/ClassSpecifier_Abstract.INT.udn b/Engine/Documentation/Source/Programming/UnrealArchitecture/Reference/Classes/Specifiers/Abstract/ClassSpecifier_Abstract.INT.udn
--- a/Engine/Documentation/Source/Programming/UnrealArchitecture/Reference/Classes/Specifiers/Abstract/ClassSpecifier_Abstract.INT.udn
+++ b/Engine/Documentation/Source/Programming/UnrealArchitecture/Reference/Classes/Specifiers/Abstract/ClassSpecifier_Abstract.INT.udn
@@ -2,7 +2,7 @@ Title: Abstract
 Crumbs:%ROOT%, Programming, Programming/UnrealArchitecture/Reference, Programming/UnrealArchitecture/Reference/Classes, Programming/UnrealArchitecture/Reference/Classes/Specifiers
 Description: Class specifier declaring the class as unable to be instanced.
 Availability:Public
-Version: 4.5
+Version: 4.9
 
 The **Abstract** class specifier declares the class as an "abstract base class", preventing the user 
 from adding Actors of this class to the world in Unreal Editor or creating instances of this class during 
diff --git a/Engine/Documentation/Source/Programming/UnrealArchitecture/Reference/Classes/Specifiers/AdvancedClassDisplay/ClassSpecifier_AdvancedClassDisplay.INT.udn b/Engine/Documentation/Source/Programming/UnrealArchitecture/Reference/Classes/Specifiers/AdvancedClassDisplay/ClassSpecifier_AdvancedClassDisplay.INT.udn
--- a/Engine/Documentation/Source/Programming/UnrealArchitecture/Reference/Classes/Specifiers/AdvancedClassDisplay/ClassSpecifier_AdvancedClassDisplay.INT.udn
+++ b/Engine/Documentation/Source/Programming/UnrealArchitecture/Reference/Classes/Specifiers/AdvancedClassDisplay/ClassSpecifier_AdvancedClassDisplay.INT.udn
@@ -2,7 +2,7 @@ Title: AdvancedClassDisplay
 Crumbs:%ROOT%, Programming, Programming/UnrealArchitecture/Reference, Programming/UnrealArchitecture/Reference/Classes, Programming/UnrealArchitecture/Reference/Classes/Specifiers
 Description: Class specifier declaring the class as unable to be instanced.
 Availability:Public
-Version: 4.5
+Version: 4.9
 
 The **AdvancedClassDisplay** class specifier forces all the properties of the class to only show in 
 the **[advanced sections](Engine/UI/LevelEditor/Details#AdvancedProperties)** in the **Details** panel, 
diff --git a/Engine/Documentation/Source/Programming/UnrealArchitecture/Reference/Classes/Specifiers/AutoCollapseCategories/ClassSpecifier_AutoCollapseCategories.INT.udn b/Engine/Documentation/Source/Programming/UnrealArchitecture/Reference/Classes/Specifiers/AutoCollapseCategories/ClassSpecifier_AutoCollapseCategories.INT.udn
--- a/Engine/Documentation/Source/Programming/UnrealArchitecture/Reference/Classes/Specifiers/AutoCollapseCategories/ClassSpecifier_AutoCollapseCategories.INT.udn
+++ b/Engine/Documentation/Source/Programming/UnrealArchitecture/Reference/Classes/Specifiers/AutoCollapseCategories/ClassSpecifier_AutoCollapseCategories.INT.udn
@@ -2,7 +2,7 @@ Title: AutoCollapseCategories
 Crumbs:%ROOT%, Programming, Programming/UnrealArchitecture/Reference, Programming/UnrealArchitecture/Reference/Classes, Programming/UnrealArchitecture/Reference/Classes/Specifiers
 Description:
 Availability:Public
-Version: 4.5
+Version: 4.9
 
 The **AutoCollapseCategories** class specifier negates the effects for the listed categories from using 
 the  **AutoExpandCategories** specifier on a parent class.
diff --git a/Engine/Documentation/Source/Programming/UnrealArchitecture/Reference/Classes/Specifiers/AutoExpandCategories/ClassSpecifier_AutoExpandCategories.INT.udn b/Engine/Documentation/Source/Programming/UnrealArchitecture/Reference/Classes/Specifiers/AutoExpandCategories/ClassSpecifier_AutoExpandCategories.INT.udn
--- a/Engine/Documentation/Source/Programming/UnrealArchitecture/Reference/Classes/Specifiers/AutoExpandCategories/ClassSpecifier_AutoExpandCategories.INT.udn
+++ b/Engine/Documentation/Source/Programming/UnrealArchitecture/Reference/Classes/Specifiers/AutoExpandCategories/ClassSpecifier_AutoExpandCategories.INT.udn
@@ -2,7 +2,7 @@ Title: AutoExpandCategories
 Crumbs:%ROOT%, Programming, Programming/UnrealArchitecture/Reference, Programming/UnrealArchitecture/Reference/Classes, Programming/UnrealArchitecture/Reference/Classes/Specifiers
 Description:
 Availability:Public
-Version: 4.5
+Version: 4.9
 
 Specifies one or more categories that should be automatically expanded in the Unreal Editor property window 
 for objects of this class. To auto-expand variables declared with no category, use the name of the class which 
diff --git a/Engine/Documentation/Source/Programming/UnrealArchitecture/Reference/ProgrammingReference.INT.udn b/Engine/Documentation/Source/Programming/UnrealArchitecture/Reference/ProgrammingReference.INT.udn
--- a/Engine/Documentation/Source/Programming/UnrealArchitecture/Reference/ProgrammingReference.INT.udn
+++ b/Engine/Documentation/Source/Programming/UnrealArchitecture/Reference/ProgrammingReference.INT.udn
@@ -9,15 +9,15 @@ When programming gameplay elements using C++ code, each module can contain many
 ![](Programming/Basics/ProjectModuleClassOrg.png)
 [/REGION]
 
-Each class defines a template for a new Actor or Object. Within the class header file, the class and any class [functions](Programming\UnrealArchitecture\Reference\Functions) and [properties](Programming\UnrealArchitecture\Reference\Properties)
-are declared. Classes can also contain [structs](Programming\UnrealArchitecture\Reference\Structs), data structures that help with organization and manipulation of related properties. Structures can also be defined on their own.
-[](Programming\UnrealArchitecture\Reference\Interfaces) allow additional gameplay behavior to be implemented by different classes. 
+Each class defines a template for a new Actor or Object. Within the class header file, the class and any class [**functions**](Programming\UnrealArchitecture\Reference\Functions) and [**properties**](Programming\UnrealArchitecture\Reference\Properties)
+are declared. Classes can also contain [**structs**](Programming\UnrealArchitecture\Reference\Structs), data structures that help with organization and manipulation of related properties. Structures can also be defined on their own.
+[**interface**](Programming\UnrealArchitecture\Reference\Interfaces) allow additional gameplay behavior to be implemented by different classes. 
 
 When programming with Unreal Engine, it is possible to have standard C++ classes, functions, and variables. These can be defined using standard C++ syntax. However, `UCLASS()`, `UFUNCTION()`, and
 `UPROPERTY()` macros can be used to make Unreal Engine aware of the new classes, functions, and variables. For instance, a variable with a declaration prefaced by a `UPROPERTY()` macro can 
 be garbage collected by the engine, and can be displayed and edited within Unreal Editor. There are also `UINTERFACE()` and `USTRUCT()` macros, and keywords for each macro that can be used
-to specify the [class](Programming\UnrealArchitecture\Reference\Classes\Specifiers)'s, [function](Programming\UnrealArchitecture\Reference\Functions\Specifiers)'s, [property](Programming\UnrealArchitecture\Reference\Properties\Specifiers)'s,
- [interface](Programming\UnrealArchitecture\Reference\Interfaces\Specifiers)'s, or [struct](Programming\UnrealArchitecture\Reference\Structs\Specifiers)'s behavior within Unreal Engine and Unreal Editor.
+to specify the [**class**](Programming\UnrealArchitecture\Reference\Classes\Specifiers)'s, [**function**](Programming\UnrealArchitecture\Reference\Functions\Specifiers)'s, [**property**](Programming\UnrealArchitecture\Reference\Properties\Specifiers)'s,
+ [**interface**](Programming\UnrealArchitecture\Reference\Interfaces\Specifiers)'s, or [**struct**](Programming\UnrealArchitecture\Reference\Structs\Specifiers)'s behavior within Unreal Engine and Unreal Editor.
 
 [OBJECT:TopicButtonList]
 	[PARAM:icon]
