#UE4doc #docs - Adjusted engine version number metadata

diff --git a/Engine/Documentation/Source/Programming/UnrealArchitecture/Reference/Classes/Specifiers/Blueprintable/ClassSpecifier_Blueprintable.INT.udn b/Engine/Documentation/Source/Programming/UnrealArchitecture/Reference/Classes/Specifiers/Blueprintable/ClassSpecifier_Blueprintable.INT.udn
--- a/Engine/Documentation/Source/Programming/UnrealArchitecture/Reference/Classes/Specifiers/Blueprintable/ClassSpecifier_Blueprintable.INT.udn
+++ b/Engine/Documentation/Source/Programming/UnrealArchitecture/Reference/Classes/Specifiers/Blueprintable/ClassSpecifier_Blueprintable.INT.udn
@@ -2,6 +2,6 @@ Title: Blueprintable
 Crumbs:%ROOT%, Programming, Programming/UnrealArchitecture/Reference, Programming/UnrealArchitecture/Reference/Classes, Programming/UnrealArchitecture/Reference/Classes/Specifiers
 Description:
 Availability:Public
-Version: 4.5
+Version: 4.9
 
 Exposes this class as an acceptable base class for creating blueprints. The default is NotBlueprintable, unless inherited otherwise. This is inherited by subclasses.
