#UE4doc #docs - Adjusted engine version number metadata

diff --git a/Engine/Documentation/Source/Programming/UnrealArchitecture/Reference/Classes/Specifiers/BlueprintType/ClassSpecifier_BlueprintType.INT.udn b/Engine/Documentation/Source/Programming/UnrealArchitecture/Reference/Classes/Specifiers/BlueprintType/ClassSpecifier_BlueprintType.INT.udn
--- a/Engine/Documentation/Source/Programming/UnrealArchitecture/Reference/Classes/Specifiers/BlueprintType/ClassSpecifier_BlueprintType.INT.udn
+++ b/Engine/Documentation/Source/Programming/UnrealArchitecture/Reference/Classes/Specifiers/BlueprintType/ClassSpecifier_BlueprintType.INT.udn
@@ -2,6 +2,6 @@ Title: BlueprintType
 Crumbs:%ROOT%, Programming, Programming/UnrealArchitecture/Reference, Programming/UnrealArchitecture/Reference/Classes, Programming/UnrealArchitecture/Reference/Classes/Specifiers
 Description:
 Availability:Public
-Version: 4.5
+Version: 4.9
 
 Exposes this class as a type that can be used for variables in blueprints.
