#UE4doc #docs - Adjusted engine version number metadata

diff --git a/Engine/Documentation/Source/Programming/UnrealArchitecture/Reference/Classes/Specifiers/ClassGroup/ClassSpecifier_ClassGroup.INT.udn b/Engine/Documentation/Source/Programming/UnrealArchitecture/Reference/Classes/Specifiers/ClassGroup/ClassSpecifier_ClassGroup.INT.udn
--- a/Engine/Documentation/Source/Programming/UnrealArchitecture/Reference/Classes/Specifiers/ClassGroup/ClassSpecifier_ClassGroup.INT.udn
+++ b/Engine/Documentation/Source/Programming/UnrealArchitecture/Reference/Classes/Specifiers/ClassGroup/ClassSpecifier_ClassGroup.INT.udn
@@ -2,7 +2,7 @@ Title: ClassGroup
 Crumbs:%ROOT%, Programming, Programming/UnrealArchitecture/Reference, Programming/UnrealArchitecture/Reference/Classes, Programming/UnrealArchitecture/Reference/Classes/Specifiers
 Description:
 Availability:Public
-Version: 4.5
+Version: 4.9
 
 Indicates that Unreal Editor's **Actor Browser** should include this class and any subclass of this 
 class within the specified *GroupName* when **Group View** is enabled in the Actor Browser.
diff --git a/Engine/Documentation/Source/Programming/UnrealArchitecture/Reference/Classes/Specifiers/CollapseCategories/ClassSpecifier_CollapseCategories.INT.udn b/Engine/Documentation/Source/Programming/UnrealArchitecture/Reference/Classes/Specifiers/CollapseCategories/ClassSpecifier_CollapseCategories.INT.udn
--- a/Engine/Documentation/Source/Programming/UnrealArchitecture/Reference/Classes/Specifiers/CollapseCategories/ClassSpecifier_CollapseCategories.INT.udn
+++ b/Engine/Documentation/Source/Programming/UnrealArchitecture/Reference/Classes/Specifiers/CollapseCategories/ClassSpecifier_CollapseCategories.INT.udn
@@ -2,7 +2,7 @@ Title: CollapseCategories
 Crumbs:%ROOT%, Programming, Programming/UnrealArchitecture/Reference, Programming/UnrealArchitecture/Reference/Classes, Programming/UnrealArchitecture/Reference/Classes/Specifiers
 Description:
 Availability:Public
-Version: 4.5
+Version: 4.9
 
 Indicates that properties of this class should not be grouped in categories in Unreal Editor property 
 windows. This keyword is propagated to child classes; however, child classes can override this flag using 
diff --git a/Engine/Documentation/Source/Programming/UnrealArchitecture/Reference/Classes/Specifiers/Config/ClassSpecifier_Config.INT.udn b/Engine/Documentation/Source/Programming/UnrealArchitecture/Reference/Classes/Specifiers/Config/ClassSpecifier_Config.INT.udn
--- a/Engine/Documentation/Source/Programming/UnrealArchitecture/Reference/Classes/Specifiers/Config/ClassSpecifier_Config.INT.udn
+++ b/Engine/Documentation/Source/Programming/UnrealArchitecture/Reference/Classes/Specifiers/Config/ClassSpecifier_Config.INT.udn
@@ -2,7 +2,7 @@ Title: Config
 Crumbs:%ROOT%, Programming, Programming/UnrealArchitecture/Reference, Programming/UnrealArchitecture/Reference/Classes, Programming/UnrealArchitecture/Reference/Classes/Specifiers
 Description:
 Availability:Public
-Version: 4.5
+Version: 4.9
 
 Indicates that this class is allowed to store data in a configuration file (`.ini`). If there are any configurable variables in the class (declared with the `config` or `globalconfig` variable specifiers), this specifier causes those variables to be stored in the specified configuration file inside the *(* and *)*. This flag is propagated to all child classes and cannot be negated, but child classes can change the `.ini` file by re-declaring the `config` keyword and specifying a different file name. The value of *IniName* is appended to the game name - minus the "Game" part - to specify the name of the `.ini` file to store data in (for instance, in UDKGame specifying `config(Camera)` will cause the class to use the `UDKCamera.ini` file). The keyword inherit can also be specified as the *IniName*, which will cause the class to use the same config file as its parent. Some `.ini` files exist by default, such as:
 
diff --git a/Engine/Documentation/Source/Programming/UnrealArchitecture/Reference/Classes/Specifiers/Const/ClassSpecifier_Const.INT.udn b/Engine/Documentation/Source/Programming/UnrealArchitecture/Reference/Classes/Specifiers/Const/ClassSpecifier_Const.INT.udn
--- a/Engine/Documentation/Source/Programming/UnrealArchitecture/Reference/Classes/Specifiers/Const/ClassSpecifier_Const.INT.udn
+++ b/Engine/Documentation/Source/Programming/UnrealArchitecture/Reference/Classes/Specifiers/Const/ClassSpecifier_Const.INT.udn
@@ -2,6 +2,6 @@ Title: Const
 Crumbs:%ROOT%, Programming, Programming/UnrealArchitecture/Reference, Programming/UnrealArchitecture/Reference/Classes, Programming/UnrealArchitecture/Reference/Classes/Specifiers
 Description:
 Availability:Public
-Version: 4.5
+Version: 4.9
 
 All properties and functions in this class are const and should be exported as const. This flag is inherited by subclasses.
diff --git a/Engine/Documentation/Source/Programming/UnrealArchitecture/Reference/Classes/Specifiers/ConversionRoot/ClassSpecifier_ConversionRoot.INT.udn b/Engine/Documentation/Source/Programming/UnrealArchitecture/Reference/Classes/Specifiers/ConversionRoot/ClassSpecifier_ConversionRoot.INT.udn
--- a/Engine/Documentation/Source/Programming/UnrealArchitecture/Reference/Classes/Specifiers/ConversionRoot/ClassSpecifier_ConversionRoot.INT.udn
+++ b/Engine/Documentation/Source/Programming/UnrealArchitecture/Reference/Classes/Specifiers/ConversionRoot/ClassSpecifier_ConversionRoot.INT.udn
@@ -2,6 +2,6 @@ Title: ConversionRoot
 Crumbs:%ROOT%, Programming, Programming/UnrealArchitecture/Reference, Programming/UnrealArchitecture/Reference/Classes, Programming/UnrealArchitecture/Reference/Classes/Specifiers
 Description:
 Availability:Public
-Version: 4.5
+Version: 4.9
 
 A root convert limits a sub-class to only be able to convert to child classes of the first root class going up the hierarchy.
diff --git a/Engine/Documentation/Source/Programming/UnrealArchitecture/Reference/Classes/Specifiers/CustomConstructor/ClassSpecifier_CustomConstructor.INT.udn b/Engine/Documentation/Source/Programming/UnrealArchitecture/Reference/Classes/Specifiers/CustomConstructor/ClassSpecifier_CustomConstructor.INT.udn
--- a/Engine/Documentation/Source/Programming/UnrealArchitecture/Reference/Classes/Specifiers/CustomConstructor/ClassSpecifier_CustomConstructor.INT.udn
+++ b/Engine/Documentation/Source/Programming/UnrealArchitecture/Reference/Classes/Specifiers/CustomConstructor/ClassSpecifier_CustomConstructor.INT.udn
@@ -2,6 +2,6 @@ Title: CustomConstructor
 Crumbs:%ROOT%, Programming, Programming/UnrealArchitecture/Reference, Programming/UnrealArchitecture/Reference/Classes, Programming/UnrealArchitecture/Reference/Classes/Specifiers
 Description:
 Availability:Public
-Version: 4.5
+Version: 4.9
 
 Prevents automatic generation of the constructor declaration.
diff --git a/Engine/Documentation/Source/Programming/UnrealArchitecture/Reference/Classes/Specifiers/DefaultToInstanced/ClassSpecifier_DefaultToInstanced.INT.udn b/Engine/Documentation/Source/Programming/UnrealArchitecture/Reference/Classes/Specifiers/DefaultToInstanced/ClassSpecifier_DefaultToInstanced.INT.udn
--- a/Engine/Documentation/Source/Programming/UnrealArchitecture/Reference/Classes/Specifiers/DefaultToInstanced/ClassSpecifier_DefaultToInstanced.INT.udn
+++ b/Engine/Documentation/Source/Programming/UnrealArchitecture/Reference/Classes/Specifiers/DefaultToInstanced/ClassSpecifier_DefaultToInstanced.INT.udn
@@ -2,6 +2,6 @@ Title: DefaultToInstanced
 Crumbs:%ROOT%, Programming, Programming/UnrealArchitecture/Reference, Programming/UnrealArchitecture/Reference/Classes, Programming/UnrealArchitecture/Reference/Classes/Specifiers
 Description:
 Availability:Public
-Version: 4.5
+Version: 4.9
 
 All instances of this class are considered "instanced". Instanced classes (components) are duplicated upon construction. This flag is inherited by subclasses.
diff --git a/Engine/Documentation/Source/Programming/UnrealArchitecture/Reference/Classes/Specifiers/DependsOn/ClassSpecifier_DependsOn.INT.udn b/Engine/Documentation/Source/Programming/UnrealArchitecture/Reference/Classes/Specifiers/DependsOn/ClassSpecifier_DependsOn.INT.udn
--- a/Engine/Documentation/Source/Programming/UnrealArchitecture/Reference/Classes/Specifiers/DependsOn/ClassSpecifier_DependsOn.INT.udn
+++ b/Engine/Documentation/Source/Programming/UnrealArchitecture/Reference/Classes/Specifiers/DependsOn/ClassSpecifier_DependsOn.INT.udn
@@ -2,7 +2,7 @@ Title: DependsOn
 Crumbs:%ROOT%, Programming, Programming/UnrealArchitecture/Reference, Programming/UnrealArchitecture/Reference/Classes, Programming/UnrealArchitecture/Reference/Classes/Specifiers
 Description:
 Availability:Public
-Version: 4.5
+Version: 4.9
 
 Indicates that *ClassName* is compiled before this class. *ClassName* must specify a class in the same (or a previous) package. Multiple dependency classes can be specified using a single `DependsOn` line delimited by commas, or can be specified using a separate `DependsOn` line for each class. This is important when a class uses a struct or enum declared in another class as the compiler only knows what is in the classes it has already compiled.
 
diff --git a/Engine/Documentation/Source/Programming/UnrealArchitecture/Reference/Classes/Specifiers/Deprecated/ClassSpecifier_Deprecated.INT.udn b/Engine/Documentation/Source/Programming/UnrealArchitecture/Reference/Classes/Specifiers/Deprecated/ClassSpecifier_Deprecated.INT.udn
--- a/Engine/Documentation/Source/Programming/UnrealArchitecture/Reference/Classes/Specifiers/Deprecated/ClassSpecifier_Deprecated.INT.udn
+++ b/Engine/Documentation/Source/Programming/UnrealArchitecture/Reference/Classes/Specifiers/Deprecated/ClassSpecifier_Deprecated.INT.udn
@@ -2,6 +2,6 @@ Title: Deprecated
 Crumbs:%ROOT%, Programming, Programming/UnrealArchitecture/Reference, Programming/UnrealArchitecture/Reference/Classes, Programming/UnrealArchitecture/Reference/Classes/Specifiers
 Description:
 Availability:Public
-Version: 4.5
+Version: 4.9
 
 This class is deprecated and objects of this class will not be saved when serializing. This flag is inherited by subclasses.
