#UE4doc #docs Checking-in images.

diff --git a/Engine/Documentation/Source/GettingStarted/DownloadingUnrealEngine/DownloadZip/DownloadZip.INT.udn b/Engine/Documentation/Source/GettingStarted/DownloadingUnrealEngine/DownloadZip/DownloadZip.INT.udn
new file mode 100644
--- /dev/null
+++ b/Engine/Documentation/Source/GettingStarted/DownloadingUnrealEngine/DownloadZip/DownloadZip.INT.udn
@@ -0,0 +1,88 @@
+Availability:NoPublish
+Title:Downloading Unreal Engine As a Zip File
+Crumbs:%ROOT%, GettingStarted, GettingStarted/DownloadingUnrealEngine
+Description:Downloading the Unreal Engine source code as a Zip file.
+
+[OBJECT:ToggleButtonContent]
+	[PARAMLITERAL:category]
+	OS
+	[/PARAMLITERAL]
+	[PARAMLITERAL:id]
+	windows
+	[/PARAMLITERAL]
+	[PARAMLITERAL:active]
+	active_button_content
+	[/PARAMLITERAL]
+	[PARAM:content]
+	1.	**Navigate** to our GitHub page.
+	1.	**Click** the **Download ZIP** button (pictured below).
+		![](DownloadZipButton.png)
+		[REGION:warning]
+		The built-in Windows zip utility marks downloaded zip file contents as unsafe to execute.  
+		Before decompressing the zip file:
+		1.	**Right-click** the zip file.
+		1.	**Click 'Properties...'**.
+		1.	**Select 'Unblock'**.
+
+		Third party zip utilities may not require the above mentioned steps prior to decompression.
+		[/REGION]
+	1.	**Decompress** 'UnrealEngine-(branch_name).zip' onto your hard disk.
+	1.	If you already haven't done so, **install** Visual Studio 2013. Although MFC libraries are included in Visual Studio's default installation settings, make sure that MFC libraries are included for ATL support.
+		[REGION:note]
+		Developers build UE4 with all desktop editions of Visual Studio, including Microsoft's free version ([Visual Studio Community 2013](https://www.visualstudio.com/products/visual-studio-community-vs)).
+		[/REGION]
+	1.	**Open** your source folder in Explorer and **run Setup.bat**.
+		[REGION:note]
+		Our Setup batch file downloads binary content for the engine, installs prerequesites, and sets up Unreal file associations. Downloading engine binaries may take some time because they are currently take up three to four gigabytes of memory. Subsequent checkouts take less time to download.
+		[/REGION]
+		[REGION:warning]
+		A warning from SmartScreen may appear in Windows 8. **Click "More Info"**, then **"Run Anyway"** to continue running our Setup batch file.
+		[/REGION]
+	[/PARAM]
+[/OBJECT]
+
+[OBJECT:ToggleButtonContent]
+	[PARAMLITERAL:category]
+	OS
+	[/PARAMLITERAL]
+	[PARAMLITERAL:id]
+	mac
+	[/PARAMLITERAL]
+	[PARAMLITERAL:active]
+	active_button_content
+	[/PARAMLITERAL]
+	[PARAM:content]
+	1.	**Navigate** to our GitHub page.
+	1.	**Click** the **Download ZIP** button (pictured below).
+		![](DownloadZipButton.png)
+	1.	**Decompress** 'UnrealEngine-(branch_name).zip' onto your hard disk.
+	1.	If you haven't already done so, **install** the latest version of [Xcode](https://itunes.apple.com/us/app/xcode/id497799835).
+	1.	**Open** your source folder in Finder and **run Setup.command**.
+		[REGION:note]
+		**Close** the Terminal window after Setup.command downloads binary content for the engine, installs prequisites, and sets up Unreal file associations.
+		[/REGION]
+	[/PARAM]
+[/OBJECT]
+
+[OBJECT:ToggleButtonContent]
+	[PARAMLITERAL:category]
+	OS
+	[/PARAMLITERAL]
+	[PARAMLITERAL:id]
+	linux
+	[/PARAMLITERAL]
+	[PARAMLITERAL:active]
+	active_button_content
+	[/PARAMLITERAL]
+	[PARAM:content]
+	1.	**Navigate** to our GitHub page.
+	1.	**Click** the **Download ZIP** button (pictured below).
+		![](DownloadZipButton.png)
+	1.	**Decompress** 'UnrealEngine-(branch_name).zip' onto your hard disk.
+	1.	**Open** your source folder and **run Setup.sh** to **download** the engine's binary content.
+	[/PARAM]
+[/OBJECT]
+
+## Next Steps
+
+Now that you've downloaded and set-up Unreal Engine, you're ready to [build the engine from source](Programming/Development/BuildingUnrealEngine).
\ No newline at end of file
diff --git a/Engine/Documentation/Source/GettingStarted/DownloadingUnrealEngine/DownloadingUnrealEngine4.INT.udn b/Engine/Documentation/Source/GettingStarted/DownloadingUnrealEngine/DownloadingUnrealEngine4.INT.udn
--- a/Engine/Documentation/Source/GettingStarted/DownloadingUnrealEngine/DownloadingUnrealEngine4.INT.udn
+++ b/Engine/Documentation/Source/GettingStarted/DownloadingUnrealEngine/DownloadingUnrealEngine4.INT.udn
@@ -72,6 +72,11 @@ Subscribers can download Unreal Engine's source code by:
 ### %GettingStarted/DownloadingUnrealEngine/GitHub:title%
 
 [INCLUDE:GettingStarted/DownloadingUnrealEngine/GitHub]
+
+### %GettingStarted/DownloadingUnrealEngine/DownloadZip:title%
+
+[INCLUDE:GettingStarted/DownloadingUnrealEngine/DownloadZip]
+
 [/PUBLISH:Rocket]
 
 [/EXCERPT:Main]
diff --git a/Engine/Documentation/Source/GettingStarted/DownloadingUnrealEngine/GitHub/GitHub.INT.udn b/Engine/Documentation/Source/GettingStarted/DownloadingUnrealEngine/GitHub/GitHub.INT.udn
--- a/Engine/Documentation/Source/GettingStarted/DownloadingUnrealEngine/GitHub/GitHub.INT.udn
+++ b/Engine/Documentation/Source/GettingStarted/DownloadingUnrealEngine/GitHub/GitHub.INT.udn
@@ -81,10 +81,6 @@ Description:Downloading the Unreal Engine source code from GitHub.
 	[/PARAM]
 [/OBJECT]
 
-## Next Steps
-
-Now that you've downloaded and set-up Unreal Engine, you're ready to [build the engine from source](Programming/Development/BuildingUnrealEngine).
-
 <!--
 
 There are two ways to get the source code for Unreal Engine from GitHub:
