#UE4doc #docs - Adjusted engine version number metadata

diff --git a/Engine/Documentation/Source/Programming/UnrealArchitecture/Reference/Classes/Specifiers/DontAutoCollapseCategories/ClassSpecifier_DontAutoCollapseCategories.INT.udn b/Engine/Documentation/Source/Programming/UnrealArchitecture/Reference/Classes/Specifiers/DontAutoCollapseCategories/ClassSpecifier_DontAutoCollapseCategories.INT.udn
--- a/Engine/Documentation/Source/Programming/UnrealArchitecture/Reference/Classes/Specifiers/DontAutoCollapseCategories/ClassSpecifier_DontAutoCollapseCategories.INT.udn
+++ b/Engine/Documentation/Source/Programming/UnrealArchitecture/Reference/Classes/Specifiers/DontAutoCollapseCategories/ClassSpecifier_DontAutoCollapseCategories.INT.udn
@@ -2,7 +2,7 @@ Title: DontAutoCollapseCategories
 Crumbs:%ROOT%, Programming, Programming/UnrealArchitecture/Reference, Programming/UnrealArchitecture/Reference/Classes, Programming/UnrealArchitecture/Reference/Classes/Specifiers
 Description:
 Availability:Public
-Version: 4.5
+Version: 4.9
 
 Negates an `AutoCollapseCategories` keyword for the specified categories inherited from a parent class.
 
diff --git a/Engine/Documentation/Source/Programming/UnrealArchitecture/Reference/Classes/Specifiers/DontCollapseCategories/ClassSpecifier_DontCollapseCategories.INT.udn b/Engine/Documentation/Source/Programming/UnrealArchitecture/Reference/Classes/Specifiers/DontCollapseCategories/ClassSpecifier_DontCollapseCategories.INT.udn
--- a/Engine/Documentation/Source/Programming/UnrealArchitecture/Reference/Classes/Specifiers/DontCollapseCategories/ClassSpecifier_DontCollapseCategories.INT.udn
+++ b/Engine/Documentation/Source/Programming/UnrealArchitecture/Reference/Classes/Specifiers/DontCollapseCategories/ClassSpecifier_DontCollapseCategories.INT.udn
@@ -2,6 +2,6 @@ Title: DontCollapseCategories
 Crumbs:%ROOT%, Programming, Programming/UnrealArchitecture/Reference, Programming/UnrealArchitecture/Reference/Classes, Programming/UnrealArchitecture/Reference/Classes/Specifiers
 Description:
 Availability:Public
-Version: 4.5
+Version: 4.9
 
 Negates a `CollapseCatogories` keyword inherited from a base class.
diff --git a/Engine/Documentation/Source/Programming/UnrealArchitecture/Reference/Classes/Specifiers/EditInlineNew/ClassSpecifier_EditInlineNew.INT.udn b/Engine/Documentation/Source/Programming/UnrealArchitecture/Reference/Classes/Specifiers/EditInlineNew/ClassSpecifier_EditInlineNew.INT.udn
--- a/Engine/Documentation/Source/Programming/UnrealArchitecture/Reference/Classes/Specifiers/EditInlineNew/ClassSpecifier_EditInlineNew.INT.udn
+++ b/Engine/Documentation/Source/Programming/UnrealArchitecture/Reference/Classes/Specifiers/EditInlineNew/ClassSpecifier_EditInlineNew.INT.udn
@@ -2,6 +2,6 @@ Title: EditInlineNew
 Crumbs:%ROOT%, Programming, Programming/UnrealArchitecture/Reference, Programming/UnrealArchitecture/Reference/Classes, Programming/UnrealArchitecture/Reference/Classes/Specifiers
 Description:
 Availability:Public
-Version: 4.5
+Version: 4.9
 
 Indicates that objects of this class can be created from the Unreal Editor property window (default behavior is that only references to existing objects may be assigned through the property window). This flag is propagated to all child classes; child classes can override this flag using the `NotEditInlineNew` keyword.
diff --git a/Engine/Documentation/Source/Programming/UnrealArchitecture/Reference/Classes/Specifiers/HideCategories/ClassSpecifier_HideCategories.INT.udn b/Engine/Documentation/Source/Programming/UnrealArchitecture/Reference/Classes/Specifiers/HideCategories/ClassSpecifier_HideCategories.INT.udn
--- a/Engine/Documentation/Source/Programming/UnrealArchitecture/Reference/Classes/Specifiers/HideCategories/ClassSpecifier_HideCategories.INT.udn
+++ b/Engine/Documentation/Source/Programming/UnrealArchitecture/Reference/Classes/Specifiers/HideCategories/ClassSpecifier_HideCategories.INT.udn
@@ -2,7 +2,7 @@ Title: HideCategories
 Crumbs:%ROOT%, Programming, Programming/UnrealArchitecture/Reference, Programming/UnrealArchitecture/Reference/Classes, Programming/UnrealArchitecture/Reference/Classes/Specifiers
 Description:
 Availability:Public
-Version: 4.5
+Version: 4.9
 
 Specifies one or more categories that should be hidden in the Unreal Editor property window for objects of this class. To hide variables declared with no category, use the name of the class which declares the variable. This keyword is propagated to child classes.
 
diff --git a/Engine/Documentation/Source/Programming/UnrealArchitecture/Reference/Classes/Specifiers/HideDropdown/ClassSpecifier_HideDropdown.INT.udn b/Engine/Documentation/Source/Programming/UnrealArchitecture/Reference/Classes/Specifiers/HideDropdown/ClassSpecifier_HideDropdown.INT.udn
--- a/Engine/Documentation/Source/Programming/UnrealArchitecture/Reference/Classes/Specifiers/HideDropdown/ClassSpecifier_HideDropdown.INT.udn
+++ b/Engine/Documentation/Source/Programming/UnrealArchitecture/Reference/Classes/Specifiers/HideDropdown/ClassSpecifier_HideDropdown.INT.udn
@@ -2,6 +2,6 @@ Title: HideDropdown
 Crumbs:%ROOT%, Programming, Programming/UnrealArchitecture/Reference, Programming/UnrealArchitecture/Reference/Classes, Programming/UnrealArchitecture/Reference/Classes/Specifiers
 Description:
 Availability:Public
-Version: 4.5
+Version: 4.9
 
 Prevents this class from showing up in Unreal Editor property window combo boxes.
diff --git a/Engine/Documentation/Source/Programming/UnrealArchitecture/Reference/Classes/Specifiers/HideFunctions/ClassSpecifier_HideFunctions.INT.udn b/Engine/Documentation/Source/Programming/UnrealArchitecture/Reference/Classes/Specifiers/HideFunctions/ClassSpecifier_HideFunctions.INT.udn
--- a/Engine/Documentation/Source/Programming/UnrealArchitecture/Reference/Classes/Specifiers/HideFunctions/ClassSpecifier_HideFunctions.INT.udn
+++ b/Engine/Documentation/Source/Programming/UnrealArchitecture/Reference/Classes/Specifiers/HideFunctions/ClassSpecifier_HideFunctions.INT.udn
@@ -2,6 +2,6 @@ Title: HideFunctions
 Crumbs:%ROOT%, Programming, Programming/UnrealArchitecture/Reference, Programming/UnrealArchitecture/Reference/Classes, Programming/UnrealArchitecture/Reference/Classes/Specifiers
 Description:
 Availability:Public
-Version: 4.5
+Version: 4.9
 
 Hides the specified function in a property viewer. Usage: hideFunctions=FunctionName or hideFunctions=(category0, category1, ...)
diff --git a/Engine/Documentation/Source/Programming/UnrealArchitecture/Reference/Classes/Specifiers/Intrinsic/ClassSpecifier_Intrinsic.INT.udn b/Engine/Documentation/Source/Programming/UnrealArchitecture/Reference/Classes/Specifiers/Intrinsic/ClassSpecifier_Intrinsic.INT.udn
--- a/Engine/Documentation/Source/Programming/UnrealArchitecture/Reference/Classes/Specifiers/Intrinsic/ClassSpecifier_Intrinsic.INT.udn
+++ b/Engine/Documentation/Source/Programming/UnrealArchitecture/Reference/Classes/Specifiers/Intrinsic/ClassSpecifier_Intrinsic.INT.udn
@@ -2,7 +2,7 @@ Title: Intrinsic
 Crumbs:%ROOT%, Programming, Programming/UnrealArchitecture/Reference, Programming/UnrealArchitecture/Reference/Classes, Programming/UnrealArchitecture/Reference/Classes/Specifiers
 Description:
 Availability:Public
-Version: 4.5
+Version: 4.9
 
 Class was declared directly in C++ and has no boilerplate generated by UnrealHeaderTool.
 
diff --git a/Engine/Documentation/Source/Programming/UnrealArchitecture/Reference/Classes/Specifiers/MinimalAPI/ClassSpecifier_MinimalAPI.INT.udn b/Engine/Documentation/Source/Programming/UnrealArchitecture/Reference/Classes/Specifiers/MinimalAPI/ClassSpecifier_MinimalAPI.INT.udn
--- a/Engine/Documentation/Source/Programming/UnrealArchitecture/Reference/Classes/Specifiers/MinimalAPI/ClassSpecifier_MinimalAPI.INT.udn
+++ b/Engine/Documentation/Source/Programming/UnrealArchitecture/Reference/Classes/Specifiers/MinimalAPI/ClassSpecifier_MinimalAPI.INT.udn
@@ -2,6 +2,6 @@ Title: MinimalAPI
 Crumbs:%ROOT%, Programming, Programming/UnrealArchitecture/Reference, Programming/UnrealArchitecture/Reference/Classes, Programming/UnrealArchitecture/Reference/Classes/Specifiers
 Description:
 Availability:Public
-Version: 4.7
+Version: 4.9
 
 Causes only the class's type information to be exported for use by other modules. The class can be cast to, but the functions of the class cannot be called (with the exception of inline methods). This improves compile times of exporting everything for classes that do not need all of their functions accessible in other modules.
diff --git a/Engine/Documentation/Source/Programming/UnrealArchitecture/Reference/Classes/Specifiers/NoExport/ClassSpecifier_NoExport.INT.udn b/Engine/Documentation/Source/Programming/UnrealArchitecture/Reference/Classes/Specifiers/NoExport/ClassSpecifier_NoExport.INT.udn
--- a/Engine/Documentation/Source/Programming/UnrealArchitecture/Reference/Classes/Specifiers/NoExport/ClassSpecifier_NoExport.INT.udn
+++ b/Engine/Documentation/Source/Programming/UnrealArchitecture/Reference/Classes/Specifiers/NoExport/ClassSpecifier_NoExport.INT.udn
@@ -2,7 +2,7 @@ Title: NoExport
 Crumbs:%ROOT%, Programming, Programming/UnrealArchitecture/Reference, Programming/UnrealArchitecture/Reference/Classes, Programming/UnrealArchitecture/Reference/Classes/Specifiers
 Description:
 Availability:Public
-Version: 4.5
+Version: 4.9
 
 Indicates that this class's declaration should not be included in the automatically-generated C++ header file by the header generator. The C++ class declaration must be defined manually in a separate header file. Only valid for native classes.
 
diff --git a/Engine/Documentation/Source/Programming/UnrealArchitecture/Reference/Classes/Specifiers/NonTransient/ClassSpecifier_NonTransient.INT.udn b/Engine/Documentation/Source/Programming/UnrealArchitecture/Reference/Classes/Specifiers/NonTransient/ClassSpecifier_NonTransient.INT.udn
--- a/Engine/Documentation/Source/Programming/UnrealArchitecture/Reference/Classes/Specifiers/NonTransient/ClassSpecifier_NonTransient.INT.udn
+++ b/Engine/Documentation/Source/Programming/UnrealArchitecture/Reference/Classes/Specifiers/NonTransient/ClassSpecifier_NonTransient.INT.udn
@@ -2,6 +2,6 @@ Title: NonTransient
 Crumbs:%ROOT%, Programming, Programming/UnrealArchitecture/Reference, Programming/UnrealArchitecture/Reference/Classes, Programming/UnrealArchitecture/Reference/Classes/Specifiers
 Description:
 Availability:Public
-Version: 4.5
+Version: 4.9
 
 Negates a `Transient` keyword inherited from a base class.
diff --git a/Engine/Documentation/Source/Programming/UnrealArchitecture/Reference/Classes/Specifiers/NotBlueprintable/ClassSpecifier_NotBlueprintable.INT.udn b/Engine/Documentation/Source/Programming/UnrealArchitecture/Reference/Classes/Specifiers/NotBlueprintable/ClassSpecifier_NotBlueprintable.INT.udn
--- a/Engine/Documentation/Source/Programming/UnrealArchitecture/Reference/Classes/Specifiers/NotBlueprintable/ClassSpecifier_NotBlueprintable.INT.udn
+++ b/Engine/Documentation/Source/Programming/UnrealArchitecture/Reference/Classes/Specifiers/NotBlueprintable/ClassSpecifier_NotBlueprintable.INT.udn
@@ -2,6 +2,6 @@ Title: NotBlueprintable
 Crumbs:%ROOT%, Programming, Programming/UnrealArchitecture/Reference, Programming/UnrealArchitecture/Reference/Classes, Programming/UnrealArchitecture/Reference/Classes/Specifiers
 Description:
 Availability:Public
-Version: 4.5
+Version: 4.9
 
 Specifies that this class is NOT an acceptable base class for creating blueprints. The default is NotBlueprintable, unless inherited otherwise. This is inherited by subclasses.
diff --git a/Engine/Documentation/Source/Programming/UnrealArchitecture/Reference/Classes/Specifiers/NotPlaceable/ClassSpecifier_NotPlaceable.INT.udn b/Engine/Documentation/Source/Programming/UnrealArchitecture/Reference/Classes/Specifiers/NotPlaceable/ClassSpecifier_NotPlaceable.INT.udn
--- a/Engine/Documentation/Source/Programming/UnrealArchitecture/Reference/Classes/Specifiers/NotPlaceable/ClassSpecifier_NotPlaceable.INT.udn
+++ b/Engine/Documentation/Source/Programming/UnrealArchitecture/Reference/Classes/Specifiers/NotPlaceable/ClassSpecifier_NotPlaceable.INT.udn
@@ -2,6 +2,7 @@ Title: NotPlaceable
 Crumbs:%ROOT%, Programming, Programming/UnrealArchitecture/Reference, Programming/UnrealArchitecture/Reference/Classes, Programming/UnrealArchitecture/Reference/Classes/Specifiers
 Description:
 Availability:Public
-Version: 4.5
+Version: 4.9
+
 
 Negates a `Placeable` keyword inherited from a base class. Indicates that this class may not be placed into a level, etc. in Unreal Editor.
diff --git a/Engine/Documentation/Source/Programming/UnrealArchitecture/Reference/Classes/Specifiers/PerObjectConfig/ClassSpecifier_PerObjectConfig.INT.udn b/Engine/Documentation/Source/Programming/UnrealArchitecture/Reference/Classes/Specifiers/PerObjectConfig/ClassSpecifier_PerObjectConfig.INT.udn
--- a/Engine/Documentation/Source/Programming/UnrealArchitecture/Reference/Classes/Specifiers/PerObjectConfig/ClassSpecifier_PerObjectConfig.INT.udn
+++ b/Engine/Documentation/Source/Programming/UnrealArchitecture/Reference/Classes/Specifiers/PerObjectConfig/ClassSpecifier_PerObjectConfig.INT.udn
@@ -2,6 +2,6 @@ Title: PerObjectConfig
 Crumbs:%ROOT%, Programming, Programming/UnrealArchitecture/Reference, Programming/UnrealArchitecture/Reference/Classes, Programming/UnrealArchitecture/Reference/Classes/Specifiers
 Description:
 Availability:Public
-Version: 4.5
+Version: 4.9
 
 Configuration information for this class will be stored per object, where each object has a section in the `.ini` file named after the object in the format `[ObjectName ClassName]`. This keyword is propagated to child classes.
diff --git a/Engine/Documentation/Source/Programming/UnrealArchitecture/Reference/Classes/Specifiers/Placeable/ClassSpecifier_Placeable.INT.udn b/Engine/Documentation/Source/Programming/UnrealArchitecture/Reference/Classes/Specifiers/Placeable/ClassSpecifier_Placeable.INT.udn
--- a/Engine/Documentation/Source/Programming/UnrealArchitecture/Reference/Classes/Specifiers/Placeable/ClassSpecifier_Placeable.INT.udn
+++ b/Engine/Documentation/Source/Programming/UnrealArchitecture/Reference/Classes/Specifiers/Placeable/ClassSpecifier_Placeable.INT.udn
@@ -2,6 +2,6 @@ Title: Placeable
 Crumbs:%ROOT%, Programming, Programming/UnrealArchitecture/Reference, Programming/UnrealArchitecture/Reference/Classes, Programming/UnrealArchitecture/Reference/Classes/Specifiers
 Description:
 Availability:Public
-Version: 4.5
+Version: 4.9
 
 Indicates that this class can be created in Unreal Editor and placed into a level, UI scene, or Blueprint (depending on the class type). This flag is propagated to all child classes; child classes can override this flag using the `NotPlaceable` keyword.
diff --git a/Engine/Documentation/Source/Programming/UnrealArchitecture/Reference/Classes/Specifiers/ShowCategories/ClassSpecifier_ShowCategories.INT.udn b/Engine/Documentation/Source/Programming/UnrealArchitecture/Reference/Classes/Specifiers/ShowCategories/ClassSpecifier_ShowCategories.INT.udn
--- a/Engine/Documentation/Source/Programming/UnrealArchitecture/Reference/Classes/Specifiers/ShowCategories/ClassSpecifier_ShowCategories.INT.udn
+++ b/Engine/Documentation/Source/Programming/UnrealArchitecture/Reference/Classes/Specifiers/ShowCategories/ClassSpecifier_ShowCategories.INT.udn
@@ -2,7 +2,7 @@ Title: ShowCategories
 Crumbs:%ROOT%, Programming, Programming/UnrealArchitecture/Reference, Programming/UnrealArchitecture/Reference/Classes, Programming/UnrealArchitecture/Reference/Classes/Specifiers
 Description:
 Availability:Public
-Version: 4.5
+Version: 4.9
 
 Negates a `HideCategories` keyword for the specified categories inherited from a base class.
 
diff --git a/Engine/Documentation/Source/Programming/UnrealArchitecture/Reference/Classes/Specifiers/ShowFunctions/ClassSpecifier_ShowFunctions.INT.udn b/Engine/Documentation/Source/Programming/UnrealArchitecture/Reference/Classes/Specifiers/ShowFunctions/ClassSpecifier_ShowFunctions.INT.udn
--- a/Engine/Documentation/Source/Programming/UnrealArchitecture/Reference/Classes/Specifiers/ShowFunctions/ClassSpecifier_ShowFunctions.INT.udn
+++ b/Engine/Documentation/Source/Programming/UnrealArchitecture/Reference/Classes/Specifiers/ShowFunctions/ClassSpecifier_ShowFunctions.INT.udn
@@ -2,6 +2,6 @@ Title: ShowFunctions
 Crumbs:%ROOT%, Programming, Programming/UnrealArchitecture/Reference, Programming/UnrealArchitecture/Reference/Classes, Programming/UnrealArchitecture/Reference/Classes/Specifiers
 Description:
 Availability:Public
-Version: 4.5
+Version: 4.9
 
 Shows the specified function in a property viewer. Usage: showFunctions=FunctionName or showFunctions=(category0, category1, ...)
diff --git a/Engine/Documentation/Source/Programming/UnrealArchitecture/Reference/Classes/Specifiers/Transient/ClassSpecifier_Transient.INT.udn b/Engine/Documentation/Source/Programming/UnrealArchitecture/Reference/Classes/Specifiers/Transient/ClassSpecifier_Transient.INT.udn
--- a/Engine/Documentation/Source/Programming/UnrealArchitecture/Reference/Classes/Specifiers/Transient/ClassSpecifier_Transient.INT.udn
+++ b/Engine/Documentation/Source/Programming/UnrealArchitecture/Reference/Classes/Specifiers/Transient/ClassSpecifier_Transient.INT.udn
@@ -2,6 +2,6 @@ Title: Transient
 Crumbs:%ROOT%, Programming, Programming/UnrealArchitecture/Reference, Programming/UnrealArchitecture/Reference/Classes, Programming/UnrealArchitecture/Reference/Classes/Specifiers
 Description:
 Availability:Public
-Version: 4.5
+Version: 4.9
 
 Says "objects belonging to this class should never be saved on disk". Only useful in conjunction with certain kinds of native classes which are non-persistent by nature, such as players or windows. This keyword is propagated to child classes; child classes can override this flag using the `NonTransient` keyword.
diff --git a/Engine/Documentation/Source/Programming/UnrealArchitecture/Reference/Classes/Specifiers/Within/ClassSpecifier_Within.INT.udn b/Engine/Documentation/Source/Programming/UnrealArchitecture/Reference/Classes/Specifiers/Within/ClassSpecifier_Within.INT.udn
--- a/Engine/Documentation/Source/Programming/UnrealArchitecture/Reference/Classes/Specifiers/Within/ClassSpecifier_Within.INT.udn
+++ b/Engine/Documentation/Source/Programming/UnrealArchitecture/Reference/Classes/Specifiers/Within/ClassSpecifier_Within.INT.udn
@@ -2,7 +2,7 @@ Title: Within
 Crumbs:%ROOT%, Programming, Programming/UnrealArchitecture/Reference, Programming/UnrealArchitecture/Reference/Classes, Programming/UnrealArchitecture/Reference/Classes/Specifiers
 Description:
 Availability:Public
-Version: 4.5
+Version: 4.9
 
 The **Within** class specifier indicates that objects of this class cannot exist _outside_ of an instance of 
 *ClassName*. In order to create an object of this class, you must specify an instance of *ClassName* as the 
diff --git a/Engine/Documentation/Source/Programming/UnrealArchitecture/Reference/Functions/Functions.INT.udn b/Engine/Documentation/Source/Programming/UnrealArchitecture/Reference/Functions/Functions.INT.udn
--- a/Engine/Documentation/Source/Programming/UnrealArchitecture/Reference/Functions/Functions.INT.udn
+++ b/Engine/Documentation/Source/Programming/UnrealArchitecture/Reference/Functions/Functions.INT.udn
@@ -2,7 +2,7 @@ Availability:Public
 Title:Functions
 Crumbs:%ROOT%, Programming, Programming/UnrealArchitecture/Reference
 Description:Reference to creating and implementing functions for gameplay classes.
-Version: 4.5
+Version: 4.9
 
 [TOC(start:2)]
 
