#UE4doc #docs - Adjusted engine version number metadata

diff --git a/Engine/Documentation/Source/Programming/UnrealArchitecture/Reference/Functions/Specifiers/BlueprintCosmetic/FunctionSpecifier_BlueprintCosmetic.INT.udn b/Engine/Documentation/Source/Programming/UnrealArchitecture/Reference/Functions/Specifiers/BlueprintCosmetic/FunctionSpecifier_BlueprintCosmetic.INT.udn
--- a/Engine/Documentation/Source/Programming/UnrealArchitecture/Reference/Functions/Specifiers/BlueprintCosmetic/FunctionSpecifier_BlueprintCosmetic.INT.udn
+++ b/Engine/Documentation/Source/Programming/UnrealArchitecture/Reference/Functions/Specifiers/BlueprintCosmetic/FunctionSpecifier_BlueprintCosmetic.INT.udn
@@ -2,6 +2,6 @@ Title: BlueprintCosmetic
 Crumbs:%ROOT%, Programming,  Programming/UnrealArchitecture/Reference,  Programming/UnrealArchitecture/Reference/Functions, Programming/UnrealArchitecture/Reference/Functions/Specifiers
 Description: 
 Availability:Public
-Version: 4.5
+Version: 4.9
 
 This function is cosmetic and will not run on dedicated servers.
diff --git a/Engine/Documentation/Source/Programming/UnrealArchitecture/Reference/Functions/Specifiers/BlueprintImplementableEvent/FunctionSpecifier_BlueprintImplementableEvent.INT.udn b/Engine/Documentation/Source/Programming/UnrealArchitecture/Reference/Functions/Specifiers/BlueprintImplementableEvent/FunctionSpecifier_BlueprintImplementableEvent.INT.udn
--- a/Engine/Documentation/Source/Programming/UnrealArchitecture/Reference/Functions/Specifiers/BlueprintImplementableEvent/FunctionSpecifier_BlueprintImplementableEvent.INT.udn
+++ b/Engine/Documentation/Source/Programming/UnrealArchitecture/Reference/Functions/Specifiers/BlueprintImplementableEvent/FunctionSpecifier_BlueprintImplementableEvent.INT.udn
@@ -2,6 +2,6 @@ Title: BlueprintImplementableEvent
 Crumbs:%ROOT%, Programming,  Programming/UnrealArchitecture/Reference,  Programming/UnrealArchitecture/Reference/Functions, Programming/UnrealArchitecture/Reference/Functions/Specifiers
 Description: 
 Availability:Public
-Version: 4.5
+Version: 4.9
 
 The function can be overridden in a Blueprint or Level Blueprint graph.
diff --git a/Engine/Documentation/Source/Programming/UnrealArchitecture/Reference/Functions/Specifiers/BlueprintNativeEvent/FunctionSpecifier_BlueprintNativeEvent.INT.udn b/Engine/Documentation/Source/Programming/UnrealArchitecture/Reference/Functions/Specifiers/BlueprintNativeEvent/FunctionSpecifier_BlueprintNativeEvent.INT.udn
--- a/Engine/Documentation/Source/Programming/UnrealArchitecture/Reference/Functions/Specifiers/BlueprintNativeEvent/FunctionSpecifier_BlueprintNativeEvent.INT.udn
+++ b/Engine/Documentation/Source/Programming/UnrealArchitecture/Reference/Functions/Specifiers/BlueprintNativeEvent/FunctionSpecifier_BlueprintNativeEvent.INT.udn
@@ -2,6 +2,6 @@ Title: BlueprintNativeEvent
 Crumbs:%ROOT%, Programming,  Programming/UnrealArchitecture/Reference,  Programming/UnrealArchitecture/Reference/Functions, Programming/UnrealArchitecture/Reference/Functions/Specifiers
 Description: 
 Availability:Public
-Version: 4.5
+Version: 4.9
 
 This function is designed to be overridden by a Blueprint, but also has a native implementation. Provide a body named [FunctionName]_Implementation instead of [FunctionName]; the autogenerated code will include a thunk that calls the implementation method when necessary.
diff --git a/Engine/Documentation/Source/Programming/UnrealArchitecture/Reference/Functions/Specifiers/BlueprintPure/FunctionSpecifier_BlueprintPure.INT.udn b/Engine/Documentation/Source/Programming/UnrealArchitecture/Reference/Functions/Specifiers/BlueprintPure/FunctionSpecifier_BlueprintPure.INT.udn
--- a/Engine/Documentation/Source/Programming/UnrealArchitecture/Reference/Functions/Specifiers/BlueprintPure/FunctionSpecifier_BlueprintPure.INT.udn
+++ b/Engine/Documentation/Source/Programming/UnrealArchitecture/Reference/Functions/Specifiers/BlueprintPure/FunctionSpecifier_BlueprintPure.INT.udn
@@ -2,6 +2,6 @@ Title: BlueprintPure
 Crumbs:%ROOT%, Programming,  Programming/UnrealArchitecture/Reference,  Programming/UnrealArchitecture/Reference/Functions, Programming/UnrealArchitecture/Reference/Functions/Specifiers
 Description: 
 Availability:Public
-Version: 4.5
+Version: 4.9
 
 The function does not affect the owning object in any way and can be executed in a Blueprint or Level Blueprint graph. 
diff --git a/Engine/Documentation/Source/Programming/UnrealArchitecture/Reference/Functions/Specifiers/Category/FunctionSpecifier_Category.INT.udn b/Engine/Documentation/Source/Programming/UnrealArchitecture/Reference/Functions/Specifiers/Category/FunctionSpecifier_Category.INT.udn
--- a/Engine/Documentation/Source/Programming/UnrealArchitecture/Reference/Functions/Specifiers/Category/FunctionSpecifier_Category.INT.udn
+++ b/Engine/Documentation/Source/Programming/UnrealArchitecture/Reference/Functions/Specifiers/Category/FunctionSpecifier_Category.INT.udn
@@ -2,7 +2,7 @@ Title: Category
 Crumbs:%ROOT%, Programming,  Programming/UnrealArchitecture/Reference,  Programming/UnrealArchitecture/Reference/Functions, Programming/UnrealArchitecture/Reference/Functions/Specifiers
 Description: 
 Availability:Public
-Version: 4.5
+Version: 4.9
 
 Specifies the category of the function when displayed in Blueprint editing tools.
 	* Usage: Category=CategoryName or Category="MajorCategory,SubCategory"
diff --git a/Engine/Documentation/Source/Programming/UnrealArchitecture/Reference/Functions/Specifiers/Client/FunctionSpecifier_Client.INT.udn b/Engine/Documentation/Source/Programming/UnrealArchitecture/Reference/Functions/Specifiers/Client/FunctionSpecifier_Client.INT.udn
--- a/Engine/Documentation/Source/Programming/UnrealArchitecture/Reference/Functions/Specifiers/Client/FunctionSpecifier_Client.INT.udn
+++ b/Engine/Documentation/Source/Programming/UnrealArchitecture/Reference/Functions/Specifiers/Client/FunctionSpecifier_Client.INT.udn
@@ -2,6 +2,6 @@ Title: Client
 Crumbs:%ROOT%, Programming,  Programming/UnrealArchitecture/Reference,  Programming/UnrealArchitecture/Reference/Functions, Programming/UnrealArchitecture/Reference/Functions/Specifiers
 Description: 
 Availability:Public
-Version: 4.5
+Version: 4.9
 
 The function is only executed on the client that owns the Object the function belongs to. Provide a body named [FunctionName]_Implementation instead of [FunctionName]; the autogenerated code will include a thunk that calls the implementation method when necessary.
diff --git a/Engine/Documentation/Source/Programming/UnrealArchitecture/Reference/Functions/Specifiers/CustomThunk/FunctionSpecifier_CustomThunk.INT.udn b/Engine/Documentation/Source/Programming/UnrealArchitecture/Reference/Functions/Specifiers/CustomThunk/FunctionSpecifier_CustomThunk.INT.udn
--- a/Engine/Documentation/Source/Programming/UnrealArchitecture/Reference/Functions/Specifiers/CustomThunk/FunctionSpecifier_CustomThunk.INT.udn
+++ b/Engine/Documentation/Source/Programming/UnrealArchitecture/Reference/Functions/Specifiers/CustomThunk/FunctionSpecifier_CustomThunk.INT.udn
@@ -2,6 +2,6 @@ Title: CustomThunk
 Crumbs:%ROOT%, Programming,  Programming/UnrealArchitecture/Reference,  Programming/UnrealArchitecture/Reference/Functions, Programming/UnrealArchitecture/Reference/Functions/Specifiers
 Description: 
 Availability:Public
-Version: 4.5
+Version: 4.9
 
 The UnrealHeaderTool code generator will not produce a execFoo thunk for this function; it is up to the user to provide one.
diff --git a/Engine/Documentation/Source/Programming/UnrealArchitecture/Reference/Functions/Specifiers/Exec/FunctionSpecifier_Exec.INT.udn b/Engine/Documentation/Source/Programming/UnrealArchitecture/Reference/Functions/Specifiers/Exec/FunctionSpecifier_Exec.INT.udn
--- a/Engine/Documentation/Source/Programming/UnrealArchitecture/Reference/Functions/Specifiers/Exec/FunctionSpecifier_Exec.INT.udn
+++ b/Engine/Documentation/Source/Programming/UnrealArchitecture/Reference/Functions/Specifiers/Exec/FunctionSpecifier_Exec.INT.udn
@@ -2,6 +2,6 @@ Title: Exec
 Crumbs:%ROOT%, Programming,  Programming/UnrealArchitecture/Reference,  Programming/UnrealArchitecture/Reference/Functions, Programming/UnrealArchitecture/Reference/Functions/Specifiers
 Description: 
 Availability:Public
-Version: 4.5
+Version: 4.9
 
 The function can be executed from the in-game console. Exec commands only function when declared within certain class.
diff --git a/Engine/Documentation/Source/Programming/UnrealArchitecture/Reference/Functions/Specifiers/NetMulticast/FunctionSpecifier_NetMulticast.INT.udn b/Engine/Documentation/Source/Programming/UnrealArchitecture/Reference/Functions/Specifiers/NetMulticast/FunctionSpecifier_NetMulticast.INT.udn
--- a/Engine/Documentation/Source/Programming/UnrealArchitecture/Reference/Functions/Specifiers/NetMulticast/FunctionSpecifier_NetMulticast.INT.udn
+++ b/Engine/Documentation/Source/Programming/UnrealArchitecture/Reference/Functions/Specifiers/NetMulticast/FunctionSpecifier_NetMulticast.INT.udn
@@ -2,6 +2,6 @@ Title: NetMulticast
 Crumbs:%ROOT%, Programming,  Programming/UnrealArchitecture/Reference,  Programming/UnrealArchitecture/Reference/Functions, Programming/UnrealArchitecture/Reference/Functions/Specifiers
 Description: 
 Availability:Public
-Version: 4.5
+Version: 4.9
 
 This function is both executed locally on the server and replicated to all clients, regardless of the Actor's NetOwner.
diff --git a/Engine/Documentation/Source/Programming/UnrealArchitecture/Reference/Functions/Specifiers/Reliable/FunctionSpecifier_Reliable.INT.udn b/Engine/Documentation/Source/Programming/UnrealArchitecture/Reference/Functions/Specifiers/Reliable/FunctionSpecifier_Reliable.INT.udn
--- a/Engine/Documentation/Source/Programming/UnrealArchitecture/Reference/Functions/Specifiers/Reliable/FunctionSpecifier_Reliable.INT.udn
+++ b/Engine/Documentation/Source/Programming/UnrealArchitecture/Reference/Functions/Specifiers/Reliable/FunctionSpecifier_Reliable.INT.udn
@@ -2,6 +2,6 @@ Title: Reliable
 Crumbs:%ROOT%, Programming,  Programming/UnrealArchitecture/Reference,  Programming/UnrealArchitecture/Reference/Functions, Programming/UnrealArchitecture/Reference/Functions/Specifiers
 Description: 
 Availability:Public
-Version: 4.5
+Version: 4.9
 
 The function is replicated over the network, and is guaranteed to arrive regardless of bandwidth or network errors. Only valid when used in conjunction with Client or Server.
diff --git a/Engine/Documentation/Source/Programming/UnrealArchitecture/Reference/Functions/Specifiers/Server/FunctionSpecifier_Server.INT.udn b/Engine/Documentation/Source/Programming/UnrealArchitecture/Reference/Functions/Specifiers/Server/FunctionSpecifier_Server.INT.udn
--- a/Engine/Documentation/Source/Programming/UnrealArchitecture/Reference/Functions/Specifiers/Server/FunctionSpecifier_Server.INT.udn
+++ b/Engine/Documentation/Source/Programming/UnrealArchitecture/Reference/Functions/Specifiers/Server/FunctionSpecifier_Server.INT.udn
@@ -2,6 +2,6 @@ Title: Server
 Crumbs:%ROOT%, Programming,  Programming/UnrealArchitecture/Reference,  Programming/UnrealArchitecture/Reference/Functions, Programming/UnrealArchitecture/Reference/Functions/Specifiers
 Description: 
 Availability:Public
-Version: 4.5
+Version: 4.9
 
 The function is only executed on the server. Provide a body named [FunctionName]_Implementation instead of [FunctionName]; the autogenerated code will include a thunk that calls the implementation method when necessary.
diff --git a/Engine/Documentation/Source/Programming/UnrealArchitecture/Reference/Functions/Specifiers/Unreliable/FunctionSpecifier_Unreliable.INT.udn b/Engine/Documentation/Source/Programming/UnrealArchitecture/Reference/Functions/Specifiers/Unreliable/FunctionSpecifier_Unreliable.INT.udn
--- a/Engine/Documentation/Source/Programming/UnrealArchitecture/Reference/Functions/Specifiers/Unreliable/FunctionSpecifier_Unreliable.INT.udn
+++ b/Engine/Documentation/Source/Programming/UnrealArchitecture/Reference/Functions/Specifiers/Unreliable/FunctionSpecifier_Unreliable.INT.udn
@@ -2,6 +2,6 @@ Title: Unreliable
 Crumbs:%ROOT%, Programming,  Programming/UnrealArchitecture/Reference,  Programming/UnrealArchitecture/Reference/Functions, Programming/UnrealArchitecture/Reference/Functions/Specifiers
 Description: 
 Availability:Public
-Version: 4.5
+Version: 4.9
 
 The function is replicated over the network but can fail due to bandwidth limitations or network errors. Only valid when used in conjunction with Client or Server.
