#UE4doc #docs - Adjusted engine version number metadata

diff --git a/Engine/Documentation/Source/Programming/UnrealArchitecture/Reference/Functions/Specifiers/BlueprintCallable/FunctionSpecifier_BlueprintCallable.INT.udn b/Engine/Documentation/Source/Programming/UnrealArchitecture/Reference/Functions/Specifiers/BlueprintCallable/FunctionSpecifier_BlueprintCallable.INT.udn
--- a/Engine/Documentation/Source/Programming/UnrealArchitecture/Reference/Functions/Specifiers/BlueprintCallable/FunctionSpecifier_BlueprintCallable.INT.udn
+++ b/Engine/Documentation/Source/Programming/UnrealArchitecture/Reference/Functions/Specifiers/BlueprintCallable/FunctionSpecifier_BlueprintCallable.INT.udn
@@ -2,6 +2,6 @@ Title: BlueprintCallable
 Crumbs:%ROOT%, Programming,  Programming/UnrealArchitecture/Reference,  Programming/UnrealArchitecture/Reference/Functions, Programming/UnrealArchitecture/Reference/Functions/Specifiers
 Description: 
 Availability:Public
-Version: 4.5
+Version: 4.9
 
 The function can be executed in a Blueprint or Level Blueprint graph.
diff --git a/Engine/Documentation/Source/Programming/UnrealArchitecture/Reference/Interfaces/Interfaces.INT.udn b/Engine/Documentation/Source/Programming/UnrealArchitecture/Reference/Interfaces/Interfaces.INT.udn
--- a/Engine/Documentation/Source/Programming/UnrealArchitecture/Reference/Interfaces/Interfaces.INT.udn
+++ b/Engine/Documentation/Source/Programming/UnrealArchitecture/Reference/Interfaces/Interfaces.INT.udn
@@ -2,7 +2,7 @@ Availability:Public
 Title: Interfaces
 Crumbs:%ROOT%, Programming, Programming/UnrealArchitecture/Reference
 Description:Reference to creating and implementing interfaces.
-Version: 4.5
+Version: 4.9
 
 [TOC(start:2)]
 
diff --git a/Engine/Documentation/Source/Programming/UnrealArchitecture/Reference/Metadata/BlueprintInternalUseOnly/FunctionMetadataSpecifier_BlueprintInternalUseOnly.INT.udn b/Engine/Documentation/Source/Programming/UnrealArchitecture/Reference/Metadata/BlueprintInternalUseOnly/FunctionMetadataSpecifier_BlueprintInternalUseOnly.INT.udn
--- a/Engine/Documentation/Source/Programming/UnrealArchitecture/Reference/Metadata/BlueprintInternalUseOnly/FunctionMetadataSpecifier_BlueprintInternalUseOnly.INT.udn
+++ b/Engine/Documentation/Source/Programming/UnrealArchitecture/Reference/Metadata/BlueprintInternalUseOnly/FunctionMetadataSpecifier_BlueprintInternalUseOnly.INT.udn
@@ -2,6 +2,6 @@ Title: BlueprintInternalUseOnly
 Crumbs: %ROOT%, Programming, Programming/UnrealArchitecture/Reference, Programming/UnrealArchitecture/Reference/Functions, Programming/UnrealArchitecture/Reference/Functions/Specifiers
 Description: 
 Availability:Public
-Version: 4.5
+Version: 4.9
 
 This function is an internal implementation detail, used to implement another function or node.  It is never directly exposed in a graph.
diff --git a/Engine/Documentation/Source/Programming/UnrealArchitecture/Reference/Metadata/BlueprintProtected/FunctionMetadataSpecifier_BlueprintProtected.INT.udn b/Engine/Documentation/Source/Programming/UnrealArchitecture/Reference/Metadata/BlueprintProtected/FunctionMetadataSpecifier_BlueprintProtected.INT.udn
--- a/Engine/Documentation/Source/Programming/UnrealArchitecture/Reference/Metadata/BlueprintProtected/FunctionMetadataSpecifier_BlueprintProtected.INT.udn
+++ b/Engine/Documentation/Source/Programming/UnrealArchitecture/Reference/Metadata/BlueprintProtected/FunctionMetadataSpecifier_BlueprintProtected.INT.udn
@@ -2,6 +2,6 @@ Title: BlueprintProtected
 Crumbs: %ROOT%, Programming, Programming/UnrealArchitecture/Reference, Programming/UnrealArchitecture/Reference/Functions, Programming/UnrealArchitecture/Reference/Functions/Specifiers
 Description: 
 Availability:Public
-Version: 4.5
+Version: 4.9
 
  This function can only be called on 'this' in a blueprint. It cannot be called on another instance.
diff --git a/Engine/Documentation/Source/Programming/UnrealArchitecture/Reference/Metadata/BlueprintSpawnableComponent/ClassMetadataSpecifier_BlueprintSpawnableComponent.INT.udn b/Engine/Documentation/Source/Programming/UnrealArchitecture/Reference/Metadata/BlueprintSpawnableComponent/ClassMetadataSpecifier_BlueprintSpawnableComponent.INT.udn
--- a/Engine/Documentation/Source/Programming/UnrealArchitecture/Reference/Metadata/BlueprintSpawnableComponent/ClassMetadataSpecifier_BlueprintSpawnableComponent.INT.udn
+++ b/Engine/Documentation/Source/Programming/UnrealArchitecture/Reference/Metadata/BlueprintSpawnableComponent/ClassMetadataSpecifier_BlueprintSpawnableComponent.INT.udn
@@ -2,6 +2,6 @@ Title: BlueprintSpawnableComponent
 Crumbs: %ROOT%, Programming, Programming/UnrealArchitecture/Reference, Programming/UnrealArchitecture/Reference/Classes, Programming/UnrealArchitecture/Reference/Classes/Specifiers
 Description: 
 Availability:Public
-Version: 4.5
+Version: 4.9
 
  If present, the component class can be spawned by a Blueprint.
