#UE4 Docs: Adding changes based on dev feedback as well as adding / updating the images to reflect thoes changes.

diff --git a/Engine/Documentation/Source/Platforms/Android/ReducingAPKSize/ReducingAPKSizes.INT.udn b/Engine/Documentation/Source/Platforms/Android/ReducingAPKSize/ReducingAPKSizes.INT.udn
--- a/Engine/Documentation/Source/Platforms/Android/ReducingAPKSize/ReducingAPKSizes.INT.udn
+++ b/Engine/Documentation/Source/Platforms/Android/ReducingAPKSize/ReducingAPKSizes.INT.udn
@@ -7,12 +7,19 @@ Version: 4.8
 
 [TOC(start:2 end:3)]
 
-No matter what type of Android device your UE4 project is targeting reducing the size of the Android application package or **APK** to be under the 50MB limit for the Google play store can be challenging. 
+No matter what type of Android device your UE4 project is targeting reducing the size of the Android application package or **APK** to be under the 50MB limit for the Google Play store can be challenging. 
 In the following tutorial we will cover steps you can take to help reduce your projects final APK package size to be as small as possible using nothing but the tools provided to you in the UE4 Editor. 
 
+## ETC1 Texture Considerations
+When packing your project for Android using the ETC1 Texture format you need to be aware that ETC1 will not compress Textures that have embedded Alpha information in them and instead treat them as fully uncompressed Textures. 
+While this will not greatly affect your projects run time performance, it will increase the size of your projects APK.
+If your project does require Textures that use Alpha's it is better to use two separate Textures, one for the Diffuse and one for the Alpha, 
+instead of a single Texture with the Alpha embedded inside of it as this will ensure that the Textures will be compressed when your project is cooked and packaged.
+
+
 ## Creating a new empty project
 When you first start working on your Android based mobile project you might be tempted to use an existing project as a base to work from or create a new project that has the Starter Content enabled. 
-**DO NOT** do this and instead create a new, completely empty C++ or Blueprint based project and then use the [Migration Tool](Engine/Content/Browser/UserGuide/Migrate) to bring over the content that you want to use. 
+**DO NOT** do this and instead create a completely new, empty C++ or Blueprint based project and then use the [Migration Tool](Engine/Content/Browser/UserGuide/Migrate) to bring over the content that you want to use. 
 This way you ensure that the only content in your project is content that you need to be there.
 
 ## Compressing Cooked Content
@@ -20,7 +27,7 @@ The easiest and quickest way to reduce your APK package size is to tell UE4 to c
 To turn on package compression you need to do the following in the UE4 Editor.
 
 [region:note]
-Please note that the Create compressed cooked packages option is only available in UE versions 4.8 and later.
+Please note that the Create compressed cooked packages option is **enabled** by default and only available in Unreal Engine versions **4.8** and later.
 [/region]
 
 1. First open up the projects settings by going to the Main Toolbar and selecting the **Edit** option then selecting **Project Settings**.
@@ -39,8 +46,7 @@ Please note that the Create compressed cooked packages option is only available
 
 	![](T_Enable_Comp_Packages.png)
 
-Now make note of the current size of your projects APK package and then re-package your project. 
-When finished compare the sizes and you should now see, depending on the content used in the project, quite a big difference between the sizes of the two APK files. 
+If you have not packaged your game with the Create compressed cooked packages checkmark box enabled you should see a huge difference in size when your project is re-packaged with it enabled.
 In fact it is not uncommon for some projects APK packages to be reduced in size by up to 50% when create compressed cooked packages is enabled.
 
 ##  Setting up a Project's Levels
@@ -59,6 +65,10 @@ To set which level or levels should be used for this type of interaction you can
 To ensure that only the content that is related to your project is added to your APK file, you can check which content is being added to your projects APK by looking in the **Cooked** folder in your project.
 You can find the Cooked folder by going to **(ProjectName)\Saved\Cooked** and then choosing the folder with the format you cooked your project for.
 
+[region:note]
+Please note that you will only be able to see the content in the Cooked folder after you complete your first cook. 
+[/region]
+
 ![](T_Cooked_Content_Folder.png)
 
 The image below shows the Texture folder from the Match 3 sample game. 
@@ -67,6 +77,10 @@ Once the biggest assets have been established you can then examine the assets in
 
 ![](T_Cooked_Content.png)
 
+## Per-Device Texture LOD's
+With the release of Unreal Engine 4.8 you can now specify which size Texture should be used on which devices. 
+You can read more about this in the [Texture Guidelines for Mobile Platforms](Platforms/Mobile/Textures) section.
+
 ## Removing Unused Content
 When your project is finally ready to be packaged up for the store make sure to first remove any unused or testing content by selecting it in the Content Browser and deleting it.
 You can delete content by first selecting the content you want to delete and then either press the **Delete** key on the keyboard or use the Delete option in the Right - Click menu. 
@@ -112,3 +126,26 @@ To disable a plug-in inside of UE4 you need to do the following.
 	[/PARAMLITERAL]
 	[/OBJECT]
 
+<!--
+## Package Black List
+
+With the release of Unreal Engine 4.9 and later you can now place a text file in your projects **Build/Platform/(Target Platform Folder)** directory that will tell the cooker to exclude partial or complete file paths from being packaged into your project.
+Project's can have multiple Black List files for Debug, Development, Test, and Shipping builds that can be setup to include or exclude whatever project data you want.
+You can even have different Black List files for each platform your project supports like one for Android, one for iOS, ect. 
+Here is an example of what the Black List files would look like for Tappy Chicken on Android. 
+
+![](T_Blacklist_Packages.png)
+
+Here is an example from Tappy Chicken where the cooker is told to not include the following files when the game is cooked and packaged.
+
+* **Blacklist Text File Location & Name:** TappyChicken/Build/Android/**PakBlacklist-Shipping.txt**:
+
+	![](T_Black_List.png)
+
+The first three entries should be safe for any project, but there may be cases were the **AssetRegistry.bin** is needed at runtime depending on your project.
+It is highly recommend that you thoroughly test to ensure that the content you are removing will not negatively impact your game.
+An easy and simple way to go about this is to open up all your projects levels to ensure that they correctly load.
+If you run into further problems and are not sure where they are coming from check the Logcat for errors.
+-->
+
+
