#UE4doc #docs - Adjusted engine version number metadata

diff --git a/Engine/Documentation/Source/Programming/UnrealArchitecture/Reference/Functions/Specifiers/BlueprintAuthorityOnly/FunctionSpecifier_BlueprintAuthorityOnly.INT.udn b/Engine/Documentation/Source/Programming/UnrealArchitecture/Reference/Functions/Specifiers/BlueprintAuthorityOnly/FunctionSpecifier_BlueprintAuthorityOnly.INT.udn
--- a/Engine/Documentation/Source/Programming/UnrealArchitecture/Reference/Functions/Specifiers/BlueprintAuthorityOnly/FunctionSpecifier_BlueprintAuthorityOnly.INT.udn
+++ b/Engine/Documentation/Source/Programming/UnrealArchitecture/Reference/Functions/Specifiers/BlueprintAuthorityOnly/FunctionSpecifier_BlueprintAuthorityOnly.INT.udn
@@ -2,6 +2,6 @@ Title: BlueprintAuthorityOnly
 Crumbs:%ROOT%, Programming,  Programming/UnrealArchitecture/Reference,  Programming/UnrealArchitecture/Reference/Functions, Programming/UnrealArchitecture/Reference/Functions/Specifiers
 Description: 
 Availability:Public
-Version: 4.5
+Version: 4.9
 
 This function will not execute from blueprint code if running on something without network authority.
diff --git a/Engine/Documentation/Source/Programming/UnrealArchitecture/Reference/Metadata/CannotImplementInterfaceBlueprint/InterfaceMetadataSpecifier_CannotImplementInterfaceInBlueprint.INT.udn b/Engine/Documentation/Source/Programming/UnrealArchitecture/Reference/Metadata/CannotImplementInterfaceBlueprint/InterfaceMetadataSpecifier_CannotImplementInterfaceInBlueprint.INT.udn
--- a/Engine/Documentation/Source/Programming/UnrealArchitecture/Reference/Metadata/CannotImplementInterfaceBlueprint/InterfaceMetadataSpecifier_CannotImplementInterfaceInBlueprint.INT.udn
+++ b/Engine/Documentation/Source/Programming/UnrealArchitecture/Reference/Metadata/CannotImplementInterfaceBlueprint/InterfaceMetadataSpecifier_CannotImplementInterfaceInBlueprint.INT.udn
@@ -2,6 +2,6 @@ Title: CannotImplementInterfaceInBlueprint
 Crumbs: %ROOT%, Programming, Programming/UnrealArchitecture/Reference, Programming/UnrealArchitecture/Reference/Interfaces, Programming/UnrealArchitecture/Reference/Interfaces/Specifiers
 Description: 
 Availability:Public
-Version: 4.5
+Version: 4.9
 
 This interface cannot be implemented by a Blueprint (e.g., it has only non-exposed C++ member methods).
diff --git a/Engine/Documentation/Source/Programming/UnrealArchitecture/Reference/Metadata/DeprecatedFunction/FunctionMetadataSpecifier_DeprecatedFunction.INT.udn b/Engine/Documentation/Source/Programming/UnrealArchitecture/Reference/Metadata/DeprecatedFunction/FunctionMetadataSpecifier_DeprecatedFunction.INT.udn
--- a/Engine/Documentation/Source/Programming/UnrealArchitecture/Reference/Metadata/DeprecatedFunction/FunctionMetadataSpecifier_DeprecatedFunction.INT.udn
+++ b/Engine/Documentation/Source/Programming/UnrealArchitecture/Reference/Metadata/DeprecatedFunction/FunctionMetadataSpecifier_DeprecatedFunction.INT.udn
@@ -2,6 +2,6 @@ Title: DeprecatedFunction
 Crumbs: %ROOT%, Programming, Programming/UnrealArchitecture/Reference, Programming/UnrealArchitecture/Reference/Functions, Programming/UnrealArchitecture/Reference/Functions/Specifiers
 Description: 
 Availability:Public
-Version: 4.5
+Version: 4.9
 
 This function is deprecated, any blueprint references to it cause a compilation warning.
diff --git a/Engine/Documentation/Source/Programming/UnrealArchitecture/Reference/Metadata/DeprecationMessage/FunctionMetadataSpecifier_DeprecationMessage.INT.udn b/Engine/Documentation/Source/Programming/UnrealArchitecture/Reference/Metadata/DeprecationMessage/FunctionMetadataSpecifier_DeprecationMessage.INT.udn
--- a/Engine/Documentation/Source/Programming/UnrealArchitecture/Reference/Metadata/DeprecationMessage/FunctionMetadataSpecifier_DeprecationMessage.INT.udn
+++ b/Engine/Documentation/Source/Programming/UnrealArchitecture/Reference/Metadata/DeprecationMessage/FunctionMetadataSpecifier_DeprecationMessage.INT.udn
@@ -2,6 +2,6 @@ Title: DeprecationMessage
 Crumbs: %ROOT%, Programming, Programming/UnrealArchitecture/Reference, Programming/UnrealArchitecture/Reference/Functions, Programming/UnrealArchitecture/Reference/Functions/Specifiers
 Description: 
 Availability:Public
-Version: 4.5
+Version: 4.9
 
 Supplies the custom message to use for deprecation.
diff --git a/Engine/Documentation/Source/Programming/UnrealArchitecture/Reference/Metadata/UnsafeDuringActorConstruction/FunctionMetadataSpecifier_UnsafeDuringActorConstruction.INT.udn b/Engine/Documentation/Source/Programming/UnrealArchitecture/Reference/Metadata/UnsafeDuringActorConstruction/FunctionMetadataSpecifier_UnsafeDuringActorConstruction.INT.udn
--- a/Engine/Documentation/Source/Programming/UnrealArchitecture/Reference/Metadata/UnsafeDuringActorConstruction/FunctionMetadataSpecifier_UnsafeDuringActorConstruction.INT.udn
+++ b/Engine/Documentation/Source/Programming/UnrealArchitecture/Reference/Metadata/UnsafeDuringActorConstruction/FunctionMetadataSpecifier_UnsafeDuringActorConstruction.INT.udn
@@ -2,6 +2,6 @@ Title: UnsafeDuringActorConstruction
 Crumbs: %ROOT%, Programming, Programming/UnrealArchitecture/Reference, Programming/UnrealArchitecture/Reference/Functions, Programming/UnrealArchitecture/Reference/Functions/Specifiers
 Description: 
 Availability:Public
-Version: 4.5
+Version: 4.9
 
 This function is not safe to call during Actor construction.
diff --git a/Engine/Documentation/Source/Programming/UnrealArchitecture/Reference/Properties/Properties.INT.udn b/Engine/Documentation/Source/Programming/UnrealArchitecture/Reference/Properties/Properties.INT.udn
--- a/Engine/Documentation/Source/Programming/UnrealArchitecture/Reference/Properties/Properties.INT.udn
+++ b/Engine/Documentation/Source/Programming/UnrealArchitecture/Reference/Properties/Properties.INT.udn
@@ -2,7 +2,7 @@ Availability:Public
 Title: Properties
 Crumbs:%ROOT%, Programming, Programming/UnrealArchitecture/Reference
 Description:Reference for creating and implementing properties for gameplay classes.
-Version: 4.5
+Version: 4.9
 
 ## Property Declaration
 
diff --git a/Engine/Documentation/Source/Programming/UnrealArchitecture/Reference/Properties/Specifiers/AdvancedDisplay/PropertySpecifier_AdvancedDisplay.INT.udn b/Engine/Documentation/Source/Programming/UnrealArchitecture/Reference/Properties/Specifiers/AdvancedDisplay/PropertySpecifier_AdvancedDisplay.INT.udn
--- a/Engine/Documentation/Source/Programming/UnrealArchitecture/Reference/Properties/Specifiers/AdvancedDisplay/PropertySpecifier_AdvancedDisplay.INT.udn
+++ b/Engine/Documentation/Source/Programming/UnrealArchitecture/Reference/Properties/Specifiers/AdvancedDisplay/PropertySpecifier_AdvancedDisplay.INT.udn
@@ -2,6 +2,6 @@ Title: AdvancedDisplay
 Crumbs: %ROOT%, Programming, Programming/UnrealArchitecture/Reference, Programming/UnrealArchitecture/Reference/Properties, Programming/UnrealArchitecture/Reference/Properties/Specifiers
 Description: 
 Availability:Public
-Version: 4.5
+Version: 4.9
 
 Properties are in the advanced dropdown in a **Details** panel.
diff --git a/Engine/Documentation/Source/Programming/UnrealArchitecture/Reference/Properties/Specifiers/AssetRegistrySearchable/PropertySpecifier_AssetRegistrySearchable.INT.udn b/Engine/Documentation/Source/Programming/UnrealArchitecture/Reference/Properties/Specifiers/AssetRegistrySearchable/PropertySpecifier_AssetRegistrySearchable.INT.udn
--- a/Engine/Documentation/Source/Programming/UnrealArchitecture/Reference/Properties/Specifiers/AssetRegistrySearchable/PropertySpecifier_AssetRegistrySearchable.INT.udn
+++ b/Engine/Documentation/Source/Programming/UnrealArchitecture/Reference/Properties/Specifiers/AssetRegistrySearchable/PropertySpecifier_AssetRegistrySearchable.INT.udn
@@ -2,7 +2,7 @@ Title: AssetRegistrySearchable
 Crumbs: %ROOT%, Programming, Programming/UnrealArchitecture/Reference, Programming/UnrealArchitecture/Reference/Properties, Programming/UnrealArchitecture/Reference/Properties/Specifiers
 Description: 
 Availability:Public
-Version: 4.5
+Version: 4.9
 
 The AssetRegistrySearchable keyword indicates that this property and its value will be automatically added to the asset registry for any asset class instances containing this as a member variable.  It is not legal
 to use on struct properties or parameters.
diff --git a/Engine/Documentation/Source/Programming/UnrealArchitecture/Reference/Properties/Specifiers/BlueprintAssignable/PropertySpecifier_BlueprintAssignable.INT.udn b/Engine/Documentation/Source/Programming/UnrealArchitecture/Reference/Properties/Specifiers/BlueprintAssignable/PropertySpecifier_BlueprintAssignable.INT.udn
--- a/Engine/Documentation/Source/Programming/UnrealArchitecture/Reference/Properties/Specifiers/BlueprintAssignable/PropertySpecifier_BlueprintAssignable.INT.udn
+++ b/Engine/Documentation/Source/Programming/UnrealArchitecture/Reference/Properties/Specifiers/BlueprintAssignable/PropertySpecifier_BlueprintAssignable.INT.udn
@@ -2,6 +2,6 @@ Title: BlueprintAssignable
 Crumbs: %ROOT%, Programming, Programming/UnrealArchitecture/Reference, Programming/UnrealArchitecture/Reference/Properties, Programming/UnrealArchitecture/Reference/Properties/Specifiers
 Description: 
 Availability:Public
-Version: 4.5
+Version: 4.9
 
 Multicast Delegates only.  Property should be exposed for assigning in Blueprints.
diff --git a/Engine/Documentation/Source/Programming/UnrealArchitecture/Reference/Properties/Specifiers/BlueprintCallable/PropertySpecifier_BlueprintCallable.INT.udn b/Engine/Documentation/Source/Programming/UnrealArchitecture/Reference/Properties/Specifiers/BlueprintCallable/PropertySpecifier_BlueprintCallable.INT.udn
--- a/Engine/Documentation/Source/Programming/UnrealArchitecture/Reference/Properties/Specifiers/BlueprintCallable/PropertySpecifier_BlueprintCallable.INT.udn
+++ b/Engine/Documentation/Source/Programming/UnrealArchitecture/Reference/Properties/Specifiers/BlueprintCallable/PropertySpecifier_BlueprintCallable.INT.udn
@@ -2,6 +2,6 @@ Title: BlueprintCallable
 Crumbs: %ROOT%, Programming, Programming/UnrealArchitecture/Reference, Programming/UnrealArchitecture/Reference/Properties, Programming/UnrealArchitecture/Reference/Properties/Specifiers
 Description: 
 Availability:Public
-Version: 4.5
+Version: 4.9
 
  Multicast Delegates only.  Property should be exposed for calling in Blueprint code.
diff --git a/Engine/Documentation/Source/Programming/UnrealArchitecture/Reference/Properties/Specifiers/BlueprintReadOnly/PropertySpecifier_BlueprintReadOnly.INT.udn b/Engine/Documentation/Source/Programming/UnrealArchitecture/Reference/Properties/Specifiers/BlueprintReadOnly/PropertySpecifier_BlueprintReadOnly.INT.udn
--- a/Engine/Documentation/Source/Programming/UnrealArchitecture/Reference/Properties/Specifiers/BlueprintReadOnly/PropertySpecifier_BlueprintReadOnly.INT.udn
+++ b/Engine/Documentation/Source/Programming/UnrealArchitecture/Reference/Properties/Specifiers/BlueprintReadOnly/PropertySpecifier_BlueprintReadOnly.INT.udn
@@ -2,6 +2,6 @@ Title: BlueprintReadOnly
 Crumbs: %ROOT%, Programming, Programming/UnrealArchitecture/Reference, Programming/UnrealArchitecture/Reference/Properties, Programming/UnrealArchitecture/Reference/Properties/Specifiers
 Description: 
 Availability:Public
-Version: 4.5
+Version: 4.9
 
 This property can be read by Blueprints, but not modified.
diff --git a/Engine/Documentation/Source/Programming/UnrealArchitecture/Reference/Properties/Specifiers/BlueprintReadWrite/PropertySpecifier_BlueprintReadWrite.INT.udn b/Engine/Documentation/Source/Programming/UnrealArchitecture/Reference/Properties/Specifiers/BlueprintReadWrite/PropertySpecifier_BlueprintReadWrite.INT.udn
--- a/Engine/Documentation/Source/Programming/UnrealArchitecture/Reference/Properties/Specifiers/BlueprintReadWrite/PropertySpecifier_BlueprintReadWrite.INT.udn
+++ b/Engine/Documentation/Source/Programming/UnrealArchitecture/Reference/Properties/Specifiers/BlueprintReadWrite/PropertySpecifier_BlueprintReadWrite.INT.udn
@@ -2,6 +2,6 @@ Title: BlueprintReadWrite
 Crumbs: %ROOT%, Programming, Programming/UnrealArchitecture/Reference, Programming/UnrealArchitecture/Reference/Properties, Programming/UnrealArchitecture/Reference/Properties/Specifiers
 Description: 
 Availability:Public
-Version: 4.5
+Version: 4.9
 
 This property can be read or written from a Blueprint.
diff --git a/Engine/Documentation/Source/Programming/UnrealArchitecture/Reference/Properties/Specifiers/Category/PropertySpecifier_Category.INT.udn b/Engine/Documentation/Source/Programming/UnrealArchitecture/Reference/Properties/Specifiers/Category/PropertySpecifier_Category.INT.udn
--- a/Engine/Documentation/Source/Programming/UnrealArchitecture/Reference/Properties/Specifiers/Category/PropertySpecifier_Category.INT.udn
+++ b/Engine/Documentation/Source/Programming/UnrealArchitecture/Reference/Properties/Specifiers/Category/PropertySpecifier_Category.INT.udn
@@ -2,7 +2,7 @@ Title: Category
 Crumbs: %ROOT%, Programming, Programming/UnrealArchitecture/Reference, Programming/UnrealArchitecture/Reference/Properties, Programming/UnrealArchitecture/Reference/Properties/Specifiers
 Description: 
 Availability:Public
-Version: 4.5
+Version: 4.9
 
 Specifies the category of the property. 
 
diff --git a/Engine/Documentation/Source/Programming/UnrealArchitecture/Reference/Properties/Specifiers/Config/PropertySpecifier_Config.INT.udn b/Engine/Documentation/Source/Programming/UnrealArchitecture/Reference/Properties/Specifiers/Config/PropertySpecifier_Config.INT.udn
--- a/Engine/Documentation/Source/Programming/UnrealArchitecture/Reference/Properties/Specifiers/Config/PropertySpecifier_Config.INT.udn
+++ b/Engine/Documentation/Source/Programming/UnrealArchitecture/Reference/Properties/Specifiers/Config/PropertySpecifier_Config.INT.udn
@@ -2,6 +2,6 @@ Title: Config
 Crumbs: %ROOT%, Programming, Programming/UnrealArchitecture/Reference, Programming/UnrealArchitecture/Reference/Properties, Programming/UnrealArchitecture/Reference/Properties/Specifiers
 Description: 
 Availability:Public
-Version: 4.5
+Version: 4.9
 
 This variable will be made configurable. The current value can be saved to the ini file and will be loaded when created. Cannot be given a value in default properties. Implies `ReadOnly`. 
diff --git a/Engine/Documentation/Source/Programming/UnrealArchitecture/Reference/Properties/Specifiers/Const/PropertySpecifier_Const.INT.udn b/Engine/Documentation/Source/Programming/UnrealArchitecture/Reference/Properties/Specifiers/Const/PropertySpecifier_Const.INT.udn
--- a/Engine/Documentation/Source/Programming/UnrealArchitecture/Reference/Properties/Specifiers/Const/PropertySpecifier_Const.INT.udn
+++ b/Engine/Documentation/Source/Programming/UnrealArchitecture/Reference/Properties/Specifiers/Const/PropertySpecifier_Const.INT.udn
@@ -2,6 +2,6 @@ Title: Const
 Crumbs: %ROOT%, Programming, Programming/UnrealArchitecture/Reference, Programming/UnrealArchitecture/Reference/Properties, Programming/UnrealArchitecture/Reference/Properties/Specifiers
 Description: 
 Availability:Public
-Version: 4.5
+Version: 4.9
 
 This variable is const and should be exported as const.
diff --git a/Engine/Documentation/Source/Programming/UnrealArchitecture/Reference/Properties/Specifiers/DuplicateTransient/PropertySpecifier_DuplicateTransient.INT.udn b/Engine/Documentation/Source/Programming/UnrealArchitecture/Reference/Properties/Specifiers/DuplicateTransient/PropertySpecifier_DuplicateTransient.INT.udn
--- a/Engine/Documentation/Source/Programming/UnrealArchitecture/Reference/Properties/Specifiers/DuplicateTransient/PropertySpecifier_DuplicateTransient.INT.udn
+++ b/Engine/Documentation/Source/Programming/UnrealArchitecture/Reference/Properties/Specifiers/DuplicateTransient/PropertySpecifier_DuplicateTransient.INT.udn
@@ -2,6 +2,6 @@ Title: DuplicateTransient
 Crumbs: %ROOT%, Programming, Programming/UnrealArchitecture/Reference, Programming/UnrealArchitecture/Reference/Properties, Programming/UnrealArchitecture/Reference/Properties/Specifiers
 Description: 
 Availability:Public
-Version: 4.5
+Version: 4.9
 
 Indicates that the variable's value should be reset to the class default value during any type of duplication (copy/paste, binary duplication, etc.).
diff --git a/Engine/Documentation/Source/Programming/UnrealArchitecture/Reference/Properties/Specifiers/EditAnywhere/PropertySpecifier_EditAnywhere.INT.udn b/Engine/Documentation/Source/Programming/UnrealArchitecture/Reference/Properties/Specifiers/EditAnywhere/PropertySpecifier_EditAnywhere.INT.udn
--- a/Engine/Documentation/Source/Programming/UnrealArchitecture/Reference/Properties/Specifiers/EditAnywhere/PropertySpecifier_EditAnywhere.INT.udn
+++ b/Engine/Documentation/Source/Programming/UnrealArchitecture/Reference/Properties/Specifiers/EditAnywhere/PropertySpecifier_EditAnywhere.INT.udn
@@ -2,6 +2,6 @@ Title: EditAnywhere
 Crumbs: %ROOT%, Programming, Programming/UnrealArchitecture/Reference, Programming/UnrealArchitecture/Reference/Properties, Programming/UnrealArchitecture/Reference/Properties/Specifiers
 Description: 
 Availability:Public
-Version: 4.5
+Version: 4.9
 
 Indicates that this property can be edited by property windows in the editor.
diff --git a/Engine/Documentation/Source/Programming/UnrealArchitecture/Reference/Properties/Specifiers/EditDefaultsOnly/PropertySpecifier_EditDefaultsOnly.INT.udn b/Engine/Documentation/Source/Programming/UnrealArchitecture/Reference/Properties/Specifiers/EditDefaultsOnly/PropertySpecifier_EditDefaultsOnly.INT.udn
--- a/Engine/Documentation/Source/Programming/UnrealArchitecture/Reference/Properties/Specifiers/EditDefaultsOnly/PropertySpecifier_EditDefaultsOnly.INT.udn
+++ b/Engine/Documentation/Source/Programming/UnrealArchitecture/Reference/Properties/Specifiers/EditDefaultsOnly/PropertySpecifier_EditDefaultsOnly.INT.udn
@@ -2,6 +2,6 @@ Title: EditDefaultsOnly
 Crumbs: %ROOT%, Programming, Programming/UnrealArchitecture/Reference, Programming/UnrealArchitecture/Reference/Properties, Programming/UnrealArchitecture/Reference/Properties/Specifiers
 Description: 
 Availability:Public
-Version: 4.5
+Version: 4.9
 
 Indicates that this property can be edited by property windows, but only on archetypes.
diff --git a/Engine/Documentation/Source/Programming/UnrealArchitecture/Reference/Properties/Specifiers/EditFixedSize/PropertySpecifier_EditFixedSize.INT.udn b/Engine/Documentation/Source/Programming/UnrealArchitecture/Reference/Properties/Specifiers/EditFixedSize/PropertySpecifier_EditFixedSize.INT.udn
--- a/Engine/Documentation/Source/Programming/UnrealArchitecture/Reference/Properties/Specifiers/EditFixedSize/PropertySpecifier_EditFixedSize.INT.udn
+++ b/Engine/Documentation/Source/Programming/UnrealArchitecture/Reference/Properties/Specifiers/EditFixedSize/PropertySpecifier_EditFixedSize.INT.udn
@@ -2,6 +2,6 @@ Title: EditFixedSize
 Crumbs: %ROOT%, Programming, Programming/UnrealArchitecture/Reference, Programming/UnrealArchitecture/Reference/Properties, Programming/UnrealArchitecture/Reference/Properties/Specifiers
 Description: 
 Availability:Public
-Version: 4.5
+Version: 4.9
 
 Only useful for dynamic arrays. This will prevent the user from changing the length of an array via the Unreal Editor property window.
diff --git a/Engine/Documentation/Source/Programming/UnrealArchitecture/Reference/Properties/Specifiers/EditInline/PropertySpecifier_EditInline.INT.udn b/Engine/Documentation/Source/Programming/UnrealArchitecture/Reference/Properties/Specifiers/EditInline/PropertySpecifier_EditInline.INT.udn
--- a/Engine/Documentation/Source/Programming/UnrealArchitecture/Reference/Properties/Specifiers/EditInline/PropertySpecifier_EditInline.INT.udn
+++ b/Engine/Documentation/Source/Programming/UnrealArchitecture/Reference/Properties/Specifiers/EditInline/PropertySpecifier_EditInline.INT.udn
@@ -2,6 +2,6 @@ Title: EditInline
 Crumbs: %ROOT%, Programming, Programming/UnrealArchitecture/Reference, Programming/UnrealArchitecture/Reference/Properties, Programming/UnrealArchitecture/Reference/Properties/Specifiers
 Description: 
 Availability:Public
-Version: 4.5
+Version: 4.9
 
 Allows the user to edit the properties of the object referenced by this variable within Unreal Editor's property inspector (only useful for object references, including arrays of object reference). 
