#UE4doc #docs - Adjusted engine version number metadata

diff --git a/Engine/Documentation/Source/Programming/UnrealArchitecture/Reference/Properties/Specifiers/EditInstanceOnly/PropertySpecifier_EditInstanceOnly.INT.udn b/Engine/Documentation/Source/Programming/UnrealArchitecture/Reference/Properties/Specifiers/EditInstanceOnly/PropertySpecifier_EditInstanceOnly.INT.udn
--- a/Engine/Documentation/Source/Programming/UnrealArchitecture/Reference/Properties/Specifiers/EditInstanceOnly/PropertySpecifier_EditInstanceOnly.INT.udn
+++ b/Engine/Documentation/Source/Programming/UnrealArchitecture/Reference/Properties/Specifiers/EditInstanceOnly/PropertySpecifier_EditInstanceOnly.INT.udn
@@ -2,6 +2,6 @@ Title: EditInstanceOnly
 Crumbs: %ROOT%, Programming, Programming/UnrealArchitecture/Reference, Programming/UnrealArchitecture/Reference/Properties, Programming/UnrealArchitecture/Reference/Properties/Specifiers
 Description: 
 Availability:Public
-Version: 4.5
+Version: 4.9
 
 Indicates that this property can be edited by property windows, but only on instances, not on archetypes.
diff --git a/Engine/Documentation/Source/Programming/UnrealArchitecture/Reference/Properties/Specifiers/Export/PropertySpecifier_Export.INT.udn b/Engine/Documentation/Source/Programming/UnrealArchitecture/Reference/Properties/Specifiers/Export/PropertySpecifier_Export.INT.udn
--- a/Engine/Documentation/Source/Programming/UnrealArchitecture/Reference/Properties/Specifiers/Export/PropertySpecifier_Export.INT.udn
+++ b/Engine/Documentation/Source/Programming/UnrealArchitecture/Reference/Properties/Specifiers/Export/PropertySpecifier_Export.INT.udn
@@ -2,6 +2,6 @@ Title: Export
 Crumbs: %ROOT%, Programming, Programming/UnrealArchitecture/Reference, Programming/UnrealArchitecture/Reference/Properties, Programming/UnrealArchitecture/Reference/Properties/Specifiers
 Description: 
 Availability:Public
-Version: 4.5
+Version: 4.9
 
 Only useful for object properties (or arrays of objects). Indicates that the object assigned to this property should be exported in its entirety as a subobject block when the object is copied (for copy/paste) or exported to T3D, as opposed to just outputting the object reference itself. 
diff --git a/Engine/Documentation/Source/Programming/UnrealArchitecture/Reference/Properties/Specifiers/GlobalConfig/PropertySpecifier_GlobalConfig.INT.udn b/Engine/Documentation/Source/Programming/UnrealArchitecture/Reference/Properties/Specifiers/GlobalConfig/PropertySpecifier_GlobalConfig.INT.udn
--- a/Engine/Documentation/Source/Programming/UnrealArchitecture/Reference/Properties/Specifiers/GlobalConfig/PropertySpecifier_GlobalConfig.INT.udn
+++ b/Engine/Documentation/Source/Programming/UnrealArchitecture/Reference/Properties/Specifiers/GlobalConfig/PropertySpecifier_GlobalConfig.INT.udn
@@ -2,6 +2,6 @@ Title: GlobalConfig
 Crumbs: %ROOT%, Programming, Programming/UnrealArchitecture/Reference, Programming/UnrealArchitecture/Reference/Properties, Programming/UnrealArchitecture/Reference/Properties/Specifiers
 Description: 
 Availability:Public
-Version: 4.5
+Version: 4.9
 
  Works just like config except that you cannot override it in a subclass. Cannot be given a value in default properties. Implies `ReadOnly`. 
diff --git a/Engine/Documentation/Source/Programming/UnrealArchitecture/Reference/Properties/Specifiers/Instanced/PropertySpecifier_Instanced.INT.udn b/Engine/Documentation/Source/Programming/UnrealArchitecture/Reference/Properties/Specifiers/Instanced/PropertySpecifier_Instanced.INT.udn
--- a/Engine/Documentation/Source/Programming/UnrealArchitecture/Reference/Properties/Specifiers/Instanced/PropertySpecifier_Instanced.INT.udn
+++ b/Engine/Documentation/Source/Programming/UnrealArchitecture/Reference/Properties/Specifiers/Instanced/PropertySpecifier_Instanced.INT.udn
@@ -2,6 +2,6 @@ Title: Instanced
 Crumbs: %ROOT%, Programming, Programming/UnrealArchitecture/Reference, Programming/UnrealArchitecture/Reference/Properties, Programming/UnrealArchitecture/Reference/Properties/Specifiers
 Description: 
 Availability:Public
-Version: 4.5
+Version: 4.9
 
 Object properties only. When an instance of this class is created, it will be given a unique copy of the object assigned to this variable in defaults. Used for instancing subobjects defined in class default properties. Implies EditInline and Export.
