#UE4doc #docs - Adjusted bold lettering and removed header.

diff --git a/Engine/Documentation/Source/Resources/ContentExamples/Physics/Physics.INT.udn b/Engine/Documentation/Source/Resources/ContentExamples/Physics/Physics.INT.udn
--- a/Engine/Documentation/Source/Resources/ContentExamples/Physics/Physics.INT.udn
+++ b/Engine/Documentation/Source/Resources/ContentExamples/Physics/Physics.INT.udn
@@ -47,7 +47,6 @@ Related: Engine/Physics/PhysicalMaterials/PhysMatReference
 [TOC(start:2 end:2)]
 
 
-## Overview
 Inside the **Physics** map, you will find several examples of how physics can be used in Unreal Engine 4. Creating **Physics Actors** out of **Static Meshes** so that they respond to player interaction, applying a **Physics Asset** to a **Skeletal Mesh** so that it responds to physics, and several examples of how **Force**, **Impulse**, **Motors**, and **Physical Constraints** are on display in this map with details of how they can aid in your level design or gameplay scenarios. 
 
 
@@ -65,14 +64,14 @@ Listed below are the examples provided inside the **Physics** map:
 
 | Example | What is Demonstrated |
 | ------- | ------------ |
-| **[1.1 Object Simulating Physics](Resources\ContentExamples\Physics\1_1)** | **StaticMeshes** which have been set to **Physics Actors** and can be picked up, poked, or thrown by mouse input (e.g. a simple pickup/drop system with gravity). |
-| **[1.2 SkeletalMesh RigidBody Simulating Physics](Resources\ContentExamples\Physics\1_2)** | A **SkeletalMesh** which has been assigned a **PhysicalAsset** and can be treated as a physics object (e.g. a **SkeletalMesh** which simulates physics, can be used for ragdoll effects). |
-| **[1.3 Radial Force/Impulse](Resources\ContentExamples\Physics\1_3)** | **Physics Actors** which are affected by a **Radial Force Actor** that pushes them all away from a point in space (e.g. could be used for explosions to blow debris away from a focal point). |
-| **[1.4 Thruster](Resources\ContentExamples\Physics\1_4)** | Shows a **Thruster** attached to a **Physics Actor** which provides constant force from a specified location (e.g. can be used to launch a rocket upward or hurl an object in a given direction at a set velocity). |
-| **[1.5 PhysicsConstraints](Resources\ContentExamples\Physics\1_5)** | **PhysicsConstraints** which are basically a type of joint that allows you to connect two **Actors** together and apply limts or forces (e.g. physics based doors that can be pushed open). |
-| **[1.6 PhysicsConstraints Angular Motors](Resources\ContentExamples\Physics\1_6)** | Similar to 1.5, however **Angular Motors** provide movement for an **Actor** that is bound by a constraint as constraints by themselves do not move (e.g. a swinging pedulum or rope). |
-| **[1.7 PhysicsConstraints Linear Motors](Resources\ContentExamples\Physics\1_7)** | While 1.6 provides movement to an **Actor** who is bound by a constraint, **Linear Motors** confine that amount of mobility by velocity or position.  |
-| **[1.8 PhysicsConstraints Breakable](Resources\ContentExamples\Physics\1_8)** | Another type of constraint that can be broken when enough force is applied (e.g. **Angular Force** like snapping a branch or **Linear Force** like "pulling" and snapping a rubber band).|
+| **[1.1 Object Simulating Physics](Resources\ContentExamples\Physics\1_1)** | StaticMeshes which have been set to Physics Actors and can be picked up, poked, or thrown by mouse input (e.g. a simple pickup/drop system with gravity). |
+| **[1.2 SkeletalMesh RigidBody Simulating Physics](Resources\ContentExamples\Physics\1_2)** | A SkeletalMesh which has been assigned a PhysicalAsset and can be treated as a physics object (e.g. a SkeletalMesh which simulates physics, can be used for ragdoll effects). |
+| **[1.3 Radial Force/Impulse](Resources\ContentExamples\Physics\1_3)** | Physics Actors which are affected by a Radial Force Actor that pushes them all away from a point in space (e.g. could be used for explosions to blow debris away from a focal point). |
+| **[1.4 Thruster](Resources\ContentExamples\Physics\1_4)** | Shows a Thruster attached to a Physics Actor which provides constant force from a specified location (e.g. can be used to launch a rocket upward or hurl an object in a given direction at a set velocity). |
+| **[1.5 PhysicsConstraints](Resources\ContentExamples\Physics\1_5)** | PhysicsConstraints which are basically a type of joint that allows you to connect two Actors together and apply limts or forces (e.g. physics based doors that can be pushed open). |
+| **[1.6 PhysicsConstraints Angular Motors](Resources\ContentExamples\Physics\1_6)** | Similar to 1.5, however Angular Motors provide movement for an Actor that is bound by a constraint as constraints by themselves do not move (e.g. a swinging pedulum or rope). |
+| **[1.7 PhysicsConstraints Linear Motors](Resources\ContentExamples\Physics\1_7)** | While 1.6 provides movement to an Actor who is bound by a constraint, Linear Motors confine that amount of mobility by velocity or position.  |
+| **[1.8 PhysicsConstraints Breakable](Resources\ContentExamples\Physics\1_8)** | Another type of constraint that can be broken when enough force is applied (e.g. Angular Force like snapping a branch or Linear Force like "pulling" and snapping a rubber band).|
 
 
 
