#UE4doc #docs - Adjusted bold lettering

diff --git a/Engine/Documentation/Source/Resources/ContentExamples/Physics/1_3/1_3.INT.udn b/Engine/Documentation/Source/Resources/ContentExamples/Physics/1_3/1_3.INT.udn
--- a/Engine/Documentation/Source/Resources/ContentExamples/Physics/1_3/1_3.INT.udn
+++ b/Engine/Documentation/Source/Resources/ContentExamples/Physics/1_3/1_3.INT.udn
@@ -53,7 +53,7 @@ Related: Engine/Physics/PhysicalMaterials/PhysMatReference
 [/REGION]
 
 
-Using a **Radial Force Actor**, you can apply a **Constant Force** (listed as _Force Strength_ in the editor) or an **Impulse**. Impulses need to be fired from a **Blueprint**, and apply their entire force in only one frame, as compared to every frame from a **Constant Force**. In the content example, only a **Constant Force** is being applied.
+Using a **Radial Force Actor**, you can apply a **Constant Force** (listed as _Force Strength_ in the editor) or an **Impulse**. Impulses need to be fired from a **Blueprint**, and apply their entire force in only one frame, as compared to every frame from a Constant Force. In the content example, only a Constant Forceis being applied.
 
 ![](physics_1_3_002.png)(w:450) ![](physics_1_3_003.png)(w:450)
 
