#UE4doc #docs - Adjusted bold lettering

diff --git a/Engine/Documentation/Source/Resources/ContentExamples/Physics/1_4/1_4.INT.udn b/Engine/Documentation/Source/Resources/ContentExamples/Physics/1_4/1_4.INT.udn
--- a/Engine/Documentation/Source/Resources/ContentExamples/Physics/1_4/1_4.INT.udn
+++ b/Engine/Documentation/Source/Resources/ContentExamples/Physics/1_4/1_4.INT.udn
@@ -53,10 +53,10 @@ Related: Engine/Physics/PhysicalMaterials/PhysMatReference
 [/REGION]
 
 
-**Physics Thrusters** are a special physics Actor you can place in the editor and then parent to a physically simulating object, making said object receive the force from the thruster. The thrust propels the **Physics Thruster** in its **Negative X** direction, or looking at it another way, the **Positive X** axis of the thruster is the "exhaust of a jet engine."
+**Physics Thrusters** are a special physics Actor you can place in the editor and then parent to a physically simulating object, making said object receive the force from the thruster. The thrust propels the Physics Thruster in its **Negative X** direction, or looking at it another way, the **Positive X** axis of the thruster is the "exhaust of a jet engine."
 
 ![](physics_1_4_002.png)(w:300)
 
-**Physics Thrusters** can be found in the **Modes** panel, in the **Place Mode** section _All Classes_.
+Physics Thrusters can be found in the **Modes** panel, in the **Place Mode** section _All Classes_.
 
 
