#UE4doc #docs - fixed misspelled word.

diff --git a/Engine/Documentation/Source/Resources/ContentExamples/Physics/1_7/1_7.INT.udn b/Engine/Documentation/Source/Resources/ContentExamples/Physics/1_7/1_7.INT.udn
--- a/Engine/Documentation/Source/Resources/ContentExamples/Physics/1_7/1_7.INT.udn
+++ b/Engine/Documentation/Source/Resources/ContentExamples/Physics/1_7/1_7.INT.udn
@@ -57,7 +57,7 @@ Related: Engine/Physics/PhysicalMaterials/PhysMatReference
 
 Enabling a linear motor on a **PhyiscsConstraint** will make the attached physically simulating objects move along an axis (or axes) in one of two methods.
 
-The first method is the **Linear Position Drive** which attempt to obtain a position relative to the constraint position. Think of it like a spring. A spring will attempt to obtain its resting position, whether it is compressed or stretched, it will push or pull to obtain its resting state. Its **Linear Drive Spring**, **Linder Drive Damping**, and **Maximum Spring Force** define how quickly and how "bouncy" the spring attempts to obtain its resting state.
+The first method is the **Linear Position Drive** which attempt to obtain a position relative to the constraint position. Think of it like a spring. A spring will attempt to obtain its resting position, whether it is compressed or stretched, it will push or pull to obtain its resting state. Its **Linear Drive Spring**, **Linear Drive Damping**, and **Maximum Spring Force** define how quickly and how "bouncy" the spring attempts to obtain its resting state.
 
 The Linear Velocity Drive simply applies a force in the given direction from the constraint's location. The simulating objects will move away from the constraint until the length of the constraint is reached. Using an inverse force will pull the constrained objects together.
 
