#UE4doc #docs - Adjusted bold lettering and removed header.

diff --git a/Engine/Documentation/Source/Resources/ContentExamples/PivotPainter/PivotPainter.INT.udn b/Engine/Documentation/Source/Resources/ContentExamples/PivotPainter/PivotPainter.INT.udn
--- a/Engine/Documentation/Source/Resources/ContentExamples/PivotPainter/PivotPainter.INT.udn
+++ b/Engine/Documentation/Source/Resources/ContentExamples/PivotPainter/PivotPainter.INT.udn
@@ -45,8 +45,8 @@ Related: Engine/Rendering/Materials/Functions/Creating
 [TOC(start:2 end:2)]
 
 
-## Overview
-The [Pivot Painter Tool](Engine\Content\Tools\PivotPainter) is a MAXScript that stores model pivot and rotation information in the model's vertex data. That information can then be referenced inside of Unreal Engine 4's shader system to create interactive effects. 
+
+The [**Pivot Painter Tool**](Engine\Content\Tools\PivotPainter) is a MAXScript that stores model pivot and rotation information in the model's vertex data. That information can then be referenced inside of Unreal Engine 4's shader system to create interactive effects. 
 
 The **PivotPainter** map inside the **Content Examples** project provides several examples which use **Pivot Painter Material Functions** to produce models that achieve this effect and demonstrates how the models can be manipulated through several forces such as wind or player interaction. 
 
