TypoCorrection

diff --git a/Engine/Documentation/Source/Engine/Rendering/Materials/Functions/Reference/Particles/Particles.INT.udn b/Engine/Documentation/Source/Engine/Rendering/Materials/Functions/Reference/Particles/Particles.INT.udn
--- a/Engine/Documentation/Source/Engine/Rendering/Materials/Functions/Reference/Particles/Particles.INT.udn
+++ b/Engine/Documentation/Source/Engine/Rendering/Materials/Functions/Reference/Particles/Particles.INT.udn
@@ -1,12 +1,10 @@
-Availability: Public
+﻿Availability: Public
 title:Particles
 Crumbs: %ROOT%, Engine, Engine/Rendering/Materials,  Engine/Rendering/Materials/Functions
 Description:Specialized functions designed to aid setting up the look of complex particle networks.
 
 [TOC (start:2 end:3)]
 
-## Overview
-
 Particle functions are designed to aid in complex material generation for particle systems.
 
 ## Particle Functions
