#UE4doc #docs:
Updated engine version
Removed overview header
Changed font settings on multiple terms.

diff --git a/Engine/Documentation/Source/Programming/UnrealArchitecture/StringHandling/StringHandling.INT.udn b/Engine/Documentation/Source/Programming/UnrealArchitecture/StringHandling/StringHandling.INT.udn
--- a/Engine/Documentation/Source/Programming/UnrealArchitecture/StringHandling/StringHandling.INT.udn
+++ b/Engine/Documentation/Source/Programming/UnrealArchitecture/StringHandling/StringHandling.INT.udn
@@ -2,9 +2,8 @@ Availability: Public
 Title:String Handling
 Crumbs:%ROOT%, Programming, Programming/UnrealArchitecture/Reference, Programming/UnrealArchitecture/Reference/Properties
 Description: An overview of the string classes available in UE4 and reference guides for FName, FText, and FString
-Version: 4.5
+Version: 4.9
 
-## Overview
 
 ### FName
 
@@ -65,7 +64,7 @@ Version: 4.5
  
 ## Encoding
 
-In general, you should be using the `TEXT()` macro when setting string variable literals. If you do not specify the `TEXT()` macro, your literal will be encoded using ANSI, which is highly limited in what 
-characters it supports. Any ANSI literals being passed into FString need to undergo a conversion to TCHAR (native Unicode encoding), so it is more efficient to use `TEXT()`.
+In general, you should be using the **TEXT()** macro when setting string variable literals. If you do not specify the TEXT() macro, your literal will be encoded using ANSI, which is highly limited in what 
+characters it supports. Any ANSI literals being passed into FString need to undergo a conversion to TCHAR (native Unicode encoding), so it is more efficient to use TEXT().
 
 For more about encoding, see the [](Programming\UnrealArchitecture\StringHandling\CharacterEncoding) documentation.
