#UE4doc #docs:
Removed Overview header
Updated engine version
Changed font settings for multiple terms.

diff --git a/Engine/Documentation/Source/Programming/UnrealArchitecture/StringHandling/CharacterEncoding/CharacterEncoding.INT.udn b/Engine/Documentation/Source/Programming/UnrealArchitecture/StringHandling/CharacterEncoding/CharacterEncoding.INT.udn
--- a/Engine/Documentation/Source/Programming/UnrealArchitecture/StringHandling/CharacterEncoding/CharacterEncoding.INT.udn
+++ b/Engine/Documentation/Source/Programming/UnrealArchitecture/StringHandling/CharacterEncoding/CharacterEncoding.INT.udn
@@ -2,11 +2,11 @@ Availability: Public
 Title: Character Encoding
 Crumbs:%ROOT%, Programming, Programming/UnrealArchitecture/Reference, Programming/UnrealArchitecture/Reference/Properties, Programming/UnrealArchitecture/StringHandling
 Description:Overview of character encodings used in Unreal Engine.
-Version: 4.5
+Version: 4.9
 
 [TOC(start:2)]
 
-## Overview
+
 
 This document provides an overview of character encodings used by Unreal.
 
@@ -70,7 +70,7 @@ A string made up of 2 bytes per character with a [BOM](http://en.wikipedia.org/w
 
 ## UE4 Internal String Representation
 
-All strings in Unreal Engine 4 are stored in memory in [UTF-16](http://en.wikipedia.org/wiki/UTF-16/UCS-2) format as FStrings or TCHAR arrays. Most code assumes 2 bytes is one codepoint so only the Basic Multilingual Plane (BMP) is supported so Unreal's internal encoding is more correctly described as UCS-2. Strings are stored in the endian-ness appropriate for the current platform.
+All strings in Unreal Engine 4 (UE4) are stored in memory in [UTF-16](http://en.wikipedia.org/wiki/UTF-16/UCS-2) format as FStrings or TCHAR arrays. Most code assumes 2 bytes is one codepoint so only the Basic Multilingual Plane (BMP) is supported so Unreal's internal encoding is more correctly described as UCS-2. Strings are stored in the endian-ness appropriate for the current platform.
 
 When serializing to packages to/from disk or during networking, strings with all TCHAR characters < 0xff are stored as a series of 8-bit bytes, and otherwise as 2-byte UTF-16 strings. The serialization code can deal with any endian conversion as necessary.
 
@@ -96,7 +96,7 @@ Note that there is no code to detect or decode UTF-8 encoded text files loaded w
 
 Most text files generated by the engine will be saved using appSaveStringToFile().
 
-Strings with all TCHAR characters representable by a single byte will be stored as a series of 8-bit bytes, and otherwise as UTF-16 unless the `bAlwaysSaveAsAnsi` flag is passed in as _true_, in which case it will be converted to the default Windows encoding first. This is currently only done on shader files, to work around an issue a shader compiler had with UTF-16 files.
+Strings with all TCHAR characters representable by a single byte will be stored as a series of 8-bit bytes, and otherwise as UTF-16 unless the bAlwaysSaveAsAnsi flag is passed in as true, in which case it will be converted to the default Windows encoding first. This is currently only done on shader files, to work around an issue a shader compiler had with UTF-16 files.
 
 
 
@@ -121,7 +121,7 @@ UTF-8 or default Windows encoding. MSVC, the Xbox360 compiler and gcc should all
 ## Storing UTF-16 Text Files in Perforce
 
 
-* Do **not** use 'Text'
+* Do not use 'Text'
     * If a UTF-x file is checked in and stored as text, it will be corrupted after syncing.
 * If you use 'Binary', mark the files as exclusive checkout
     * People can check in ASCII, UTF-8, UTF-16 and it will work in engine.
@@ -174,7 +174,7 @@ It is also critical that when using TCHAR_TO_ANSI you do not assume the number o
 
 UE4 currently only handles ANSI (ASCII | code page 1252 | | Western European).
 
-The least worst method to do this for all languages seems to be mentioned here http://en.wikipedia.org/wiki/ISO/IEC_8859
+The least worst method to do this for all languages seems to be mentioned here [](http://en.wikipedia.org/wiki/ISO/IEC_8859)
 
 
 * ISO/IEC 8859-1 for English, French, German, Italian, Portuguese, and both Spanishes
@@ -182,7 +182,7 @@ The least worst method to do this for all languages seems to be mentioned here h
 * ISO/IEC 8859-5 for Russian
 
 
-The mappings from ftp://ftp.unicode.org/Public/MAPPINGS/ISO8859/ contain the conversion rules for the above mentioned languages. 'CAPITAL LETTER' and 'SMALL LETTER' info would be cross referenced with the appropriate unicode character to get the desired result.
+The mappings from [](ftp://ftp.unicode.org/Public/MAPPINGS/ISO8859/) contain the conversion rules for the above mentioned languages. 'CAPITAL LETTER' and 'SMALL LETTER' info would be cross referenced with the appropriate unicode character to get the desired result.
 
 
 
@@ -244,7 +244,7 @@ Causes hard-to-find bugs or errors with missing code, if the comment text contai
     
 
 
-The compiler on East Asian code page Windows recognizes the last byte of the UTF-8 decoded East Asian character comment and asterisk '*' as a single East Asian character, and next characters is treated as still part of the comment. In above case, compiler removes important_function() as it seems to be part of the comment.  
+The compiler on East Asian code page Windows recognizes the last byte of the UTF-8 decoded East Asian character comment and asterisk * as a single East Asian character, and next characters is treated as still part of the comment. In above case, compiler removes important_function() as it seems to be part of the comment.  
 This behavior is very dangerous and it is difficult to find the missing code.
 
 In single-line comment:  
@@ -275,7 +275,7 @@ __Conclusion__
 You can use UTF-8 or default Windows encoding for C++ source code, but please be aware of these problem. Again, we do not recommend string literals inside C++ source. Please make sure to use East Asian as your default code page if you have to use East Asian character encoding in C++ source code.  
 Another good way is to use UTF-8 with BOM (some text editors describes the BOM as a Unicode signature).
 
-__Note__
+[REGION:note]
   
 We tested a few compilers with UTF-8 and UTF-16 on 18 Feb 2010.  
 
@@ -286,3 +286,5 @@ Perforce is able to work with both UTF-16 and UTF-8, but p4 diff displays the BO
 
 External reference: [Code Pages Supported by Windows](http://msdn.microsoft.com/en-us/goglobal/bb964654.aspx)
 
+[/REGION]
+
