#UE4doc #docs:
Removed Overviewed header
Updated engine version
Changed font settings for multiple terms.

diff --git a/Engine/Documentation/Source/Programming/UnrealArchitecture/StringHandling/FName/FName.INT.udn b/Engine/Documentation/Source/Programming/UnrealArchitecture/StringHandling/FName/FName.INT.udn
--- a/Engine/Documentation/Source/Programming/UnrealArchitecture/StringHandling/FName/FName.INT.udn
+++ b/Engine/Documentation/Source/Programming/UnrealArchitecture/StringHandling/FName/FName.INT.udn
@@ -2,15 +2,15 @@ Availability: Public
 Title:FName
 Crumbs:%ROOT%, Programming, Programming/UnrealArchitecture/Reference, Programming/UnrealArchitecture/Reference/Properties, Programming/UnrealArchitecture/StringHandling
 Description: %doxygen:FName%
-Version: 4.5
+Version: 4.9
 
 [TOC (start:2 end:3)]
 
-## Overview
+
 
 
 [EXCERPT:Overview]
-When you name a new asset in the **Content Browser**, change a parameter in a Dynamic Material Instance, or access a bone in a Skeletal Mesh, you are using FNames. FNames provide a very lightweight system 
+When you name a new asset in the **Content Browser**, change a parameter in a Dynamic Material Instance, or access a bone in a Skeletal Mesh, you are using **FNames**. FNames provide a very lightweight system 
 for using strings, where a given string is stored only once in a data table, even if it is reused.  FNames are case-insensitive. They are immutable, and cannot be manipulated. The storage system and static nature of FNames means 
 that it is fast to look up and access FNames with keys.  Another feature of the FName subsystem is the use of a hash table to provide fast string to FName conversions. 
 [/EXCERPT:Overview]
@@ -42,26 +42,26 @@ FNames can only be converted to FStrings and FText, and can only be converted fr
  
 [REGION:note]
 When doing these conversions, be aware that they may contain characters that are not valid for the type of FName you are creating.
-The `INVALID_NAME_CHARACTERS` macro in `NameTypes.h` defines which characters are invalid for FNames, and several functions within
-`FName` (such as `IsValidObjectName()`) check the validity of FNames for particular use cases.
+The INVALID_NAME_CHARACTERS macro in NameTypes.h defines which characters are invalid for FNames, and several functions within
+FName (such as **IsValidObjectName()**) check the validity of FNames for particular use cases.
 [/REGION]
 
 ## Comparisons
 
-The `==` operator can be used to compare two FNames, and returns `true` or `false`.  Rather than doing string comparisons, this compares the values in `Index`, which is a significant CPU savings.
+The == operator can be used to compare two FNames, and returns true or false.  Rather than doing string comparisons, this compares the values in Index, which is a significant CPU savings.
 
-`FName::Compare` can also be used to compare two FNames, and it will return < 0 if this < Other, 0 if this == Other, and > 0 if this > Other.
+**FName::Compare** can also be used to compare two FNames, and it will return < 0 if this < Other, 0 if this == Other, and > 0 if this > Other.
 
 	CompareFloat = TestFName.Compare(OtherFName); 
 
 ## Using FNames
 
-Using an FName is pretty straightforward. Let us say that you wanted to get the bone named "pelvis" from the Skeletal Mesh Component of an Actor. The C++ code below shows the use of an FName constructed on 
-the fly while being passed to GetBoneRotation().
+Using an FName is pretty straightforward. Let us say that you wanted to get the bone named pelvis from the Skeletal Mesh Component of an Actor. The C++ code below shows the use of an FName constructed on 
+the fly while being passed to **GetBoneRotation()**.
 
 	FRotator rotPelvis = Mesh->MeshGetInstance(this))->GetBoneRotation(FName(TEXT("pelvis")));
 
-This creates an FName that is passed to `GetBoneRotation()`, which returns the FRotator for that corresponding bone. The name of the bone gets loaded into the FName table when the package loads, 
+This creates an FName that is passed to GetBoneRotation(), which returns the FRotator for that corresponding bone. The name of the bone gets loaded into the FName table when the package loads, 
 so the constructor for the FName finds the bone's name in the hash table, avoiding an allocation.
 
 
@@ -74,4 +74,4 @@ If you want to determine if an FName is in the name table, but you do not want t
 		// Do something
 	}
 
-If the name does not exist in the name table, the Index of the FName is set to `NAME_None`. Note that we are not checking a pointer for NULL as you would with a normal string. 
+If the name does not exist in the name table, the Index of the FName is set to NAME_None. Note that we are not checking a pointer for null as you would with a normal string. 
