#UE4doc #docs:
Removed overview header
Updated engine version
Changed font settings for multiple terms

diff --git a/Engine/Documentation/Source/Programming/UnrealArchitecture/StringHandling/FString/FString.INT.udn b/Engine/Documentation/Source/Programming/UnrealArchitecture/StringHandling/FString/FString.INT.udn
--- a/Engine/Documentation/Source/Programming/UnrealArchitecture/StringHandling/FString/FString.INT.udn
+++ b/Engine/Documentation/Source/Programming/UnrealArchitecture/StringHandling/FString/FString.INT.udn
@@ -2,14 +2,14 @@ Availability: Public
 Title:FString
 Crumbs:%ROOT%, Programming, Programming/UnrealArchitecture/Reference, Programming/UnrealArchitecture/Reference/Properties, Programming/UnrealArchitecture/StringHandling
 Description:
-Version: 4.5
+Version: 4.9
 
 [TOC (start:2 end:3)]
 
-## Overview
+
 
 [EXCERPT:Overview]
-Unlike FName and FText, FString is the only string class that allows for manipulation. There are many methods available for string manipulation, including case changing, excerpting substrings, and 
+Unlike FName and FText, **FString** is the only string class that allows for manipulation. There are many methods available for string manipulation, including case changing, excerpting substrings, and 
 reversing.  FStrings can be searched, modified, and compared against other strings. However, these manipulations can make FStrings more expensive than the immutable string classes.    
 [/EXCERPT:Overview]
 
@@ -50,11 +50,11 @@ reversing.  FStrings can be searched, modified, and compared against other strin
 | FLinearColor    |   `LinearColorVariable.ToString();`  | '(R=,G=,B=,A=)' |
 | UObject         |  `(InObj != NULL) ? InObj->GetName() : FString(TEXT("None"));` | UObject's FName |
 
-For other numeric conversions, you can use the `FString::Printf()` function with the appropriate arguments. 
+For other numeric conversions, you can use the **FString::Printf()** function with the appropriate arguments. 
 
 #### To FString
 
-Conversions also exist from `FString` to `int` and `float` numeric variables, as well as to boolean variables. 
+Conversions also exist from FString to int and float numeric variables, as well as to boolean variables. 
 
 | Variable Type | Conversion to String | Notes |
 | ------------- | -------------------- | --------------|
@@ -64,17 +64,17 @@ Conversions also exist from `FString` to `int` and `float` numeric variables, as
 
 ## Comparisons
 
-The overloaded `==` operator can be used to compare two FStrings, or to compare an FString to an array of TCHAR*s. There is also the `FString::Equals()` method, which takes
-the FString to test against and the `ESearchCase` enum for whether or not the comparison should ignore case as arguments. If you want the comparison to ignore case, use `ESearchCase::IgnoreCase`,
- and if not, use `ESearchCase::CaseSensitive`.
+The overloaded == operator can be used to compare two FStrings, or to compare an FString to an array of TCHAR*s. There is also the **FString::Equals()** method, which takes
+the FString to test against and the **ESearchCase** enum for whether or not the comparison should ignore case as arguments. If you want the comparison to ignore case, use **ESearchCase::IgnoreCase**,
+ and if not, use **ESearchCase::CaseSensitive**.
 
 	TestHUDString.Equals(TEXT("Test"), ESearchCase::CaseSensitive);
 
 
 ## Searching
 
-When searching within FStrings, there are two search types. The first, `FString::Contains()`, returns _true_ if the substring is found, and _false_ otherwise. `FString::Contains()` can search for 
-either an `FString` or a `TCHAR*`s substring.  The `ESearchCase` enum can be used to specify whether or not the search should ignore case. Also, the `ESearchDir` enum can be used to specify
+When searching within FStrings, there are two search types. The first, **FString::Contains()**, returns true if the substring is found, and _false_ otherwise. FString::Contains() can search for 
+either an FString or a TCHAR*s substring.  The ESearchCase enum can be used to specify whether or not the search should ignore case. Also, the `ESearchDir` enum can be used to specify
 the direction of the search.  The default is to ignore case, and to begin searching at the start.
 
 	TestHUDString.Contains(TEXT("Test"), ESearchCase::CaseSensitive, ESearchDir::FromEnd);
@@ -106,8 +106,8 @@ There are two operators for concatenating strings:
 
 ### Printf
 
-FStrings constructed with `FString::Printf` can be stored into FStrings, as well as displayed to the screen with [UE_LOG debug messaging](#debugmessaging). The format argument has the same specifiers
-as the C++ `printf` function, as seen in the below example.
+FStrings constructed with **FString::Printf** can be stored into FStrings, as well as displayed to the screen with [UE_LOG debug messaging](#debugmessaging). The format argument has the same specifiers
+as the C++ printf function, as seen in the below example.
 	
 	FString AShooterHUD::GetTimeString(float TimeSeconds)
 	{
@@ -121,26 +121,26 @@ as the C++ `printf` function, as seen in the below example.
 	}
 
 [REGION:tip]
-When using `%s` parameters to include FStrings, the * operator must be used to return the TCHAR* required for the %s parameter.
+When using %s parameters to include FStrings, the `*` operator must be used to return the TCHAR* required for the %s parameter.
 [/REGION]
 
 ## Manipulating FStrings
 
-There are many functions for manipulating strings. Some will be covered here, but for the full list of available FString functions, see `UnrealString.h` or the API documentation on FString.
-There are functions for copying subsections of strings: `Left`, `Right`, and `Mid`. You can split a string into two strings at the location of a found substring. 
-This is done using the `Split` method. Another method for splitting strings is the `ParseIntoArray`, which splits a string into an array of strings separated by a specified delimiter. 
-Case conversion is done by using `ToUpper` and `ToLower`, converting the string to upper and lower case respectively.
+There are many functions for manipulating strings. Some will be covered here, but for the full list of available FString functions, see UnrealString.h or the API documentation on FString.
+There are functions for copying subsections of strings: Left, Right, and Mid. You can split a string into two strings at the location of a found substring. 
+This is done using the `Split` method. Another method for splitting strings is the **ParseIntoArray**, which splits a string into an array of strings separated by a specified delimiter. 
+Case conversion is done by using **ToUpper** and **ToLower**, converting the string to upper and lower case respectively.
 
 ## FStrings in HUDs
 
 ### Canvas
 
-To display an FString in a HUD using [Canvas](Gameplay/Framework\UIAndHUD#canvas), call the `FCanvas::DrawText()` function:
+To display an FString in a HUD using [Canvas](Gameplay/Framework\UIAndHUD#canvas), call the **FCanvas::DrawText()** function:
 	
 	Canvas->DrawText(BigFont, TestHUDString, 110.0f,110.0f);
 
 [REGION:note] 
-You must call the `DrawText()` function within your HUD class's `DrawHUD()` function, or call it in a function chain that begins with `DrawHUD()`.
+You must call the **DrawText()** function within your HUD class's **DrawHUD()** function, or call it in a function chain that begins with DrawHUD().
 [/REGION]
 
 ## Debug Messaging
@@ -151,7 +151,7 @@ FStrings can be printed to the **Viewport** as well as to the **Output Log** for
 
 ![](DebugToScreen.png)
 
-To print debug messages to the **Viewport**, use `UEngine::AddOnScreenDebugMessage()`. This function takes the following four parameters (in order):
+To print debug messages to the **Viewport**, use **UEngine::AddOnScreenDebugMessage()**. This function takes the following four parameters (in order):
 
 | Parameter Name | Parameter Type | Description |
 | -------------- | ---------------| ------------|
@@ -168,32 +168,32 @@ To print debug messages to the **Viewport**, use `UEngine::AddOnScreenDebugMessa
 
 ![](DebugToLog.png)
 
-`UE_LOG` uses `printf` markup for parameterization. 
+**UE_LOG** uses printf markup for parameterization. 
 
 	UE_LOG(LogClass, Log, TEXT("This is a testing statement. %s"), *TestHUDString);
 
-* `LogClass` is the log category. You can use an existing category (set with the `DECLARE_LOG_CATEGORY_EXTERN` macro in `OutputDevices.h`) or define your own using DEFINE_LOG_CATEGORY_STATIC. 
-* `Log` is the verbosity level to use. Verbosity is defined in the `ELogVerbosity` enum. Valid values are `Fatal`, `Error`, `Warning`, `Display`, `Log`, `Verbose`, or `VeryVerbose`.
+* LogClass is the log category. You can use an existing category (set with the DECLARE_LOG_CATEGORY_EXTERN macro in OutputDevices.h) or define your own using DEFINE_LOG_CATEGORY_STATIC. 
+* Log is the verbosity level to use. Verbosity is defined in the **ELogVerbosity** enum. Valid values are Fatal, Error, Warning, Display, Log, Verbose, or VeryVerbose.
 * The next argument is the text you wish to output, including the markup for the parameters. 
 
 [REGION:tip]
-This example uses a `%s` parameter, so the `*` operator is used to return the `TCHAR*` required for a `%s` parameter.
+This example uses a %s parameter, so the `*` operator is used to return the TCHAR* required for a %s parameter.
 [/REGION]
 
-Messages printed with `UE_LOG` can be found in the **Output Log** (**Window > Output Log** in Unreal Editor).
+Messages printed with UE_LOG can be found in the **Output Log** (**Window > Output Log** in Unreal Editor).
 
 ## Encoding Conversion Macros
 The FString class is built upon a TArray of TCHARs. There are multiple macros available to convert an application string (TCHAR*) to either ANSI or UNICODE character sets and vice versa. 
-The macro definitions are found in `Engine\Source\Runtime\Core\Public\Containers\StringConv.h`.
+The macro definitions are found in Engine\Source\Runtime\Core\Public\Containers\StringConv.h.
 
 If the string is relatively small, the allocation is made on the stack as part of the converter class; otherwise the heap is used to allocate a temporary buffer. The size before using the heap is a 
 template parameter, so you can tune this to your application. This is safe within loops because the scoping of the class pops off the stack allocation. 
 
 The common conversion macros are:
 
-* `TCHAR_TO_ANSI` - Converts an engine string (TCHAR*) to an ANSI one.
+* TCHAR_TO_ANSI - Converts an engine string (TCHAR*) to an ANSI one.
 
-* `ANSI_TO_TCHAR` - Converts an ANSI string to an engine string (TCHAR*).
+* ANSI_TO_TCHAR - Converts an ANSI string to an engine string (TCHAR*).
 
 [REGION:note]
 The objects these macros declare have very short lifetimes. They are meant to be used as parameters to functions. You cannot assign a variable to the contents of the converted string as the 
@@ -202,7 +202,7 @@ object will go out of scope and the string will be released.
 [REGION:warning]
 The parameter you pass in must be a proper string, as the parameter is typecast to a pointer. If you pass in a TCHAR instead of a TCHAR*, it will compile and then crash at runtime.
 
-Usage: `SomeApi(TCHAR_TO_ANSI(SomeUnicodeString));`
+Usage: **SomeApi(TCHAR_TO_ANSI(SomeUnicodeString));**
 	
 [/REGION]
 
