#UE4doc #docs:
Removed overview header
Updated engine version
Changed font settings for multiple terms

diff --git a/Engine/Documentation/Source/Programming/UnrealArchitecture/StringHandling/FText/FText.INT.udn b/Engine/Documentation/Source/Programming/UnrealArchitecture/StringHandling/FText/FText.INT.udn
--- a/Engine/Documentation/Source/Programming/UnrealArchitecture/StringHandling/FText/FText.INT.udn
+++ b/Engine/Documentation/Source/Programming/UnrealArchitecture/StringHandling/FText/FText.INT.udn
@@ -2,13 +2,13 @@ Availability: Public
 Title:FText
 Crumbs:%ROOT%, Programming, Programming/UnrealArchitecture/Reference, Programming/UnrealArchitecture/Reference/Properties, Programming/UnrealArchitecture/StringHandling
 Description:
-Version: 4.5
+Version: 4.9
 
 [TOC (start:2 end:3)]
 
-## Overview
+
 [EXCERPT:Overview]
-FText represents a "display string".  Any text you want to display to the user should be handled with FText. The FText class has built-in support for localization, and can handle text content that is 
+**FText** represents a display string.  Any text you want to display to the user should be handled with FText. The FText class has built-in support for localization, and can handle text content that is 
 localized and stored in a lookup table, as well as text that is localized at runtime, such as numbers, dates, times, and formatted text.  Even text that does not need to be localized can be handled 
 with FText.  This includes user-entered content such as a player's name, and any text displayed with Slate. FText does not offer any mutation functions, because making changes to display strings is a very unsafe operation.
 [/EXCERPT:Overview]
@@ -24,14 +24,14 @@ will exist.
 
 ## Creating FText
 
-Because of the localization applications of FText, each FText must be set up with a key as well as text. The `LOCTEXT` and `NSLOCTEXT` macros can be used to set up FText correctly.
-`LOCTEXT` takes the key and the text as parameters. All parameters must be string literals. All literals will be passed through the localization system.
+Because of the localization applications of FText, each FText must be set up with a key as well as text. The LOCTEXT and NSLOCTEXT macros can be used to set up FText correctly.
+LOCTEXT takes the key and the text as parameters. All parameters must be string literals. All literals will be passed through the localization system.
 
 [REGION:tip]
 The global LOCTEXT_NAMESPACE macro must be first set to a string literal to specify this localization key's namespace.
 [/REGION]
 
-To create an FText from a string using the `LOCTEXT` macro:
+To create an FText from a string using the LOCTEXT macro:
   
 	/* Top of File */ 
 	#define LOCTEXT_NAMESPACE "Your Namespace" 
@@ -41,7 +41,7 @@ To create an FText from a string using the `LOCTEXT` macro:
 	#undef LOCTEXT_NAMESPACE 
 	/* Bottom of File */ 
   
-`NSLOCTEXT` takes the namespace as a parameter, in addition to the key and the text. All parameters must be string literals. All literals will be passed through the localization system.
+NSLOCTEXT takes the namespace as a parameter, in addition to the key and the text. All parameters must be string literals. All literals will be passed through the localization system.
 
 	FText TestHUDText = NSLOCTEXT( "Your Namespace", "Your Key", "Your Text" ) 
 
@@ -79,7 +79,7 @@ No conversions exist from FName strings to numeric variables.
 
 #### To FText
 
-You can convert `float` and `int` variables to FText using `FText::AsNumber()`. This function generates an FText that represents the passed number in the current culture.
+You can convert float and int variables to FText using **FText::AsNumber()**. This function generates an FText that represents the passed number in the current culture.
 
 Example:
 
@@ -95,21 +95,21 @@ Additional numeric conversions include:
   
 ## Comparisons
 
-You can compare two FText strings by using `FText::EqualTo()` or `FText::EqualToCaseIgnored()`, receiving either `true` or `false` as your result. Using `FText::CompareTo()` or 
-`FText::CompareToCaseIgnored()` to compare strings returns an integer value.   
+You can compare two FText strings by using **FText::EqualTo()** or **FText::EqualToCaseIgnored()**, receiving either true or false as your result. Using **FText::CompareTo()** or 
+**FText::CompareToCaseIgnored()** to compare strings returns an integer value.   
   
 FText does not expose operators for these comparison operations currently because their operations are not simple ordinal comparisons like FStrings. We opted to not use the normal operators to hint 
 at the operation's expense. 
 
 ## Formatting FText
   
-You can use the `FText::Format()` function to concatenate different FText together, and if desired, save it to an FText variable. 
+You can use the **FText::Format()** function to concatenate different FText together, and if desired, save it to an FText variable. 
   
 	FText::Format( LOCTEXT( "HelloMessageFormat",  "Hello {0}," ), PlayerName ) 
   
 Currently, the Format function only takes FText as parameters.
 
-`FText::Format` supports two different kinds of format arguments, ordered and named.
+FText::Format supports two different kinds of format arguments, ordered and named.
 
 ### Ordered Arguments
  
@@ -121,7 +121,7 @@ Currently, the Format function only takes FText as parameters.
 ### Named Arguments 
 
 * Arguments = {Name}, {ID}, {Health}, etc... 
-* Curly braces denote the beginning and end of the format argument and the text between the braces is the name of argument to be found in the `FFormatNamedArgs` collection passed in.
+* Curly braces denote the beginning and end of the format argument and the text between the braces is the name of argument to be found in the FFormatNamedArgs collection passed in.
 
 		FFormatNamedArguments Arguments;
 		Arguments.Add(TEXT("CurrentHealth"), FText::AsNumber( CurrentHealth ));
@@ -136,22 +136,22 @@ TextTest.cpp can show you examples of how to perform valid culture aware formatt
 
 ### Canvas
 
-To display FText in a HUD using [Canvas](Gameplay/Framework\UIAndHUD#canvas), you must create a new FCanvasTextItem.
+To display FText in a HUD using [Canvas](Gameplay/Framework\UIAndHUD#canvas), you must create a new **FCanvasTextItem**.
 
 	FCanvasTextItem TextItem( FVector2D::ZeroVector, TestHUDText, BigFont, FLinearColor::Black);
 	
-You can change the text of an FCanvasTextItem by setting the `Text` variable:
+You can change the text of an FCanvasTextItem by setting the Text variable:
 
 	FFormatNamedArguments Arguments;
 	Arguments.Add(TEXT("CurrentHealth"), FText::AsNumber( CurrentHealth ));
 	TextItem.Text = FText::Format(LOCTEXT("ExampleFText", "You currently have {CurrentHealth} health left."), Arguments ) ;
 
-To actually display the FText to the screen, call the `FCanvas::DrawItem()` function:
+To actually display the FText to the screen, call the **FCanvas::DrawItem()** function:
 	
 	Canvas->DrawItem(TextItem,10.0f,10.0f);
 
 [REGION:note] 
-You must call the `DrawItem()` function within your HUD class's `DrawHUD()` function, or call it in a function chain that begins with `DrawHUD()`.
+You must call the DrawItem() function within your HUD class's **DrawHUD()** function, or call it in a function chain that begins with DrawHUD().
 [/REGION]
 
 
@@ -159,11 +159,11 @@ You must call the `DrawItem()` function within your HUD class's `DrawHUD()` func
 
 Blueprints support Text literals. Any input pins for Text values will support being assigned a value in the node and these strings will be gathered up for localization.
 
-`FText::Format` is also available in Blueprints. The "Format" text can be a literal or connected to a pin. When used as a literal, it will auto-generate arguments for formatting. So by entering the 
+**FText::Format** is also available in Blueprints. The Format text can be a literal or connected to a pin. When used as a literal, it will auto-generate arguments for formatting. So by entering the 
 literal "There are {number} days until {day}.", a **Number** pin and a **Day** pin will be automatically added to the node. These pins can be filled in as a literal or from a variable and the output 
 Text will replace the arguments with whatever values are hooked to the pins. 
 
-When the "Format" pin is connected, the "Format" can not be deciphered because it is variable. In this case, you can add pins yourself and name them after the possible arguments. In this way, you can 
+When the Format pin is connected, the Format can not be deciphered because it is variable. In this case, you can add pins yourself and name them after the possible arguments. In this way, you can 
 set up the Format Text node to be ready for multiple possibilities. Whenever an argument matches an argument in the Format variable, it will utilize the provided information.
 
 Pin names can be changed, added, and removed from the **Details** Panel.
