#UE4doc #docs:
Changed font settings for multiple terms
Replaced 'method' with 'function' at every instance, as 'methods' do not exist in normal C++ and are more widely known as 'functions'

diff --git a/Engine/Documentation/Source/Programming/UnrealArchitecture/TArrays/TArrays.INT.udn b/Engine/Documentation/Source/Programming/UnrealArchitecture/TArrays/TArrays.INT.udn
--- a/Engine/Documentation/Source/Programming/UnrealArchitecture/TArrays/TArrays.INT.udn
+++ b/Engine/Documentation/Source/Programming/UnrealArchitecture/TArrays/TArrays.INT.udn
@@ -4,7 +4,7 @@ Description:
 Availability:Public
 Version:4.9
 
-The primary go-to container class is TArray.  TArray is a class which is responsible for the ownership and organization of a sequence of other objects (called elements) of the same type.  As a TArray is a sequence, its elements have a well-defined order and its methods are used to deterministically manipulate those objects and their order.
+The primary go-to container class is **TArray**.  TArray is a class which is responsible for the ownership and organization of a sequence of other objects (called elements) of the same type.  As a TArray is a sequence, its elements have a well-defined order and its functions are used to deterministically manipulate those objects and their order.
 
 ## TArray
 
@@ -24,12 +24,12 @@ To create an array, define it like this:
 
 This creates an empty array designed to hold a sequence of integers. The element type can be any value type which is copyable and destructible according to normal C++ value rules, like int32, FString, TSharedPtr etc..  No allocator has been specified and so the TArray has regular heap-based allocation.  At this point, no memory has been allocated.
 
-TArrays can be populated in several ways.  One way is with the Init method which will fill an array with a number of copies of an element:
+TArrays can be populated in several ways.  One way is with the **Init** function which will fill an array with a number of copies of an element:
 
 	IntArray.Init(10, 5);
 	// IntArray == [10,10,10,10,10]
 
-Add and Emplace methods can be used to create new objects at the end of the array:
+**Add** and **Emplace** functions can be used to create new objects at the end of the array:
 
 	TArray<FString> StrArr;
 	StrArr.Add    (TEXT("Hello"));
@@ -45,13 +45,13 @@ So in the case of our TArray<FString>, Add will create a temporary FString from
 
 In general, Emplace is better than Add, in that it avoids unnecessary temporaries being created at the call site which are then copied or moved into the container.  As a rule of thumb, use Add for trivial types and Emplace otherwise.  Emplace should never be less efficient than Add, but sometimes Add may read better.
 
-Append allows you to add multiple elements at once from either another TArray, or a pointer+size to a regular C array:
+**Append** allows you to add multiple elements at once from either another TArray, or a pointer+size to a regular C array:
 
 	FString Arr[] = { TEXT("of"), TEXT("Tomorrow") };
 	StrArr.Append(Arr, ARRAY_COUNT(Arr));
 	// StrArr == ["Hello","World","of","Tomorrow"]
 
-AddUnique only adds a new element to the container if an equivalent element doesn't already exist.  Equivalence is checked by using the element type's operator==:
+**AddUnique** only adds a new element to the container if an equivalent element doesn't already exist.  Equivalence is checked by using the element type's operator==:
 
 	StrArr.AddUnique(TEXT("!"));
 	// StrArr == ["Hello","World","of","Tomorrow","!"]
@@ -59,12 +59,12 @@ AddUnique only adds a new element to the container if an equivalent element does
 	StrArr.AddUnique(TEXT("!"));
 	// StrArr is unchanged as "!" is already an element
 
-Insert, like Add, Emplace and Append, allows you to add a single element or a copy of an array of elements at a given index:
+**Insert**, like Add, Emplace and Append, allows you to add a single element or a copy of an array of elements at a given index:
 
 	StrArr.Insert(TEXT("Brave"), 1);
 	// StrArr == ["Hello","Brave","World","of","Tomorrow","!"]
 
-The SetNum method can be directly set the number of array elements, with new elements being created using the element type's default constructor if the new number is greater than the current one:
+The **SetNum** function can be directly set the number of array elements, with new elements being created using the element type's default constructor if the new number is greater than the current one:
 
 	StrArr.SetNum(8);
 	// StrArr == ["Hello","Brave","World","of","Tomorrow","!","",""]
@@ -94,7 +94,7 @@ Regular index-based iteration is also possible of course:
 		JoinedStr += TEXT(" ");
 	}
 
-Finally, arrays also have their own iterator type for more control over your iteration.  There are two methods called CreateIterator and CreateConstIterator which can be used for read-write or read-only access to the elements respectively:
+Finally, arrays also have their own iterator type for more control over your iteration.  There are two functions called **CreateIterator** and **CreateConstIterator** which can be used for read-write or read-only access to the elements respectively:
 
 	for (auto It = StrArr.CreateConstIterator(); It; ++It)
 	{
@@ -104,7 +104,7 @@ Finally, arrays also have their own iterator type for more control over your ite
 
 ##Sorting
 
-Arrays can be sorted simply by calling the Sort method:
+Arrays can be sorted simply by calling the **Sort** function:
 
 	StrArr.Sort();
 	// StrArr == ["!","Brave","Hello","of","Tomorrow","World"]
@@ -118,14 +118,14 @@ Here, the values are sorted by means of the element type's operator<.  In FStrin
 
 Now the strings are sorted by their lengths.  Note how the three strings with the same length - "Hello", "Brave" and "World" - have changed order relative to their positions in the array beforehand.  This is because Sort is unstable and the relative order of equivalent elements (those strings are equivalent here, because they have the same length) is not guaranteed.  Sort is implemented as a quicksort.
 
-The HeapSort method, with or without a binary predicate, can be used to perform a heap sort.  Whether or not you choose to use it depends on your particular data and how efficiently it sorts compared to the Sort method.  Like Sort, HeapSort is not stable.  If we had used HeapSort instead of Sort above, this would be the result (the same, in this case):
+The **HeapSort** function, with or without a binary predicate, can be used to perform a heap sort.  Whether or not you choose to use it depends on your particular data and how efficiently it sorts compared to the Sort function.  Like Sort, HeapSort is not stable.  If we had used HeapSort instead of Sort above, this would be the result (the same, in this case):
 
 	StrArr.HeapSort([](const FString& A, const FString& B) {
 		return A.Len() < B.Len();
 	});
 	// StrArr == ["!","of","Hello","Brave","World","Tomorrow"]
 
-Finally, StableSort can be used to guarantee the relative order of equivalent elements after sorting.  If we had called StableSort instead of Sort or HeapSort above, the result would have been as follows:
+Finally, **StableSort** can be used to guarantee the relative order of equivalent elements after sorting.  If we had called StableSort instead of Sort or HeapSort above, the result would have been as follows:
 
 	StrArr.StableSort([](const FString& A, const FString& B) {
 		return A.Len() < B.Len();
@@ -136,12 +136,12 @@ That is, "Brave", "Hello" and "World" remain in their same relative order after
 
 ## Queries
 
-We can ask the array how many elements it holds by using the Num method:
+We can ask the array how many elements it holds by using the **Num** function:
 
 	int32 Count = StrArr.Num();
 	// Count == 6
 
-If you need direct access to the array memory, perhaps for interoperability with a C-style API, you can use the GetData() method to return a pointer to the elements in the array.  This pointer is only valid as long as the array exists and before any mutating operations are made to the array.  Only the first Num() indices from the StrPtr are dereferenceable:
+If you need direct access to the array memory, perhaps for interoperability with a C-style API, you can use the **GetData()** function to return a pointer to the elements in the array.  This pointer is only valid as long as the array exists and before any mutating operations are made to the array.  Only the first Num() indices from the StrPtr are dereferenceable:
 
 	FString* StrPtr = StrArr.GetData();
 	// StrPtr[0] == "!"
@@ -162,7 +162,7 @@ To retrieve elements, you can use the indexing operator[] and pass it a zero-bas
 	FString Elem1 = StrArr[1];
 	// Elem1 == "of"
 
-Passing an invalid index — less than 0 or greater than or equal to Num() — will cause a runtime error.  You can ask the container if a particular index is valid using the IsValidIndex method:
+Passing an invalid index — less than 0 or greater than or equal to Num() — will cause a runtime error.  You can ask the container if a particular index is valid using the **IsValidIndex** function:
 
 	bool bValidM1 = StrArr.IsValidIndex(-1);
 	bool bValid0  = StrArr.IsValidIndex(0);
@@ -178,7 +178,7 @@ operator[] returns a reference, so it can also be used to mutate the elements in
 	StrArr[3] = StrArr[3].ToUpper();
 	// StrArr == ["!","of","Brave","HELLO","World","Tomorrow"]
 
-Like the GetData method, operator[] will return a const reference if the array is const.  You can also index from the end of the array backwards by using the Last method.  The index is defaulted to zero.  The Top method is a synonym for Last, except it doesn't take an index:
+Like the GetData function, operator[] will return a const reference if the array is const.  You can also index from the end of the array backwards by using the **Last** function.  The index is defaulted to zero.  The **Top** function is a synonym for Last, except it doesn't take an index:
 
 	FString ElemEnd  = StrArr.Last();
 	FString ElemEnd0 = StrArr.Last(0);
@@ -207,7 +207,7 @@ Or ask the array if it contains an element which matches a specific predicate:
 	// bLen5 == true
 	// bLen6 == false
 
-We can find elements by using the Find family of methods.  To check if an element exists and return its index, we use Find:
+We can find elements by using the **Find** family of functions.  To check if an element exists and return its index, we use Find:
 
 	int32 Index;
 	if (StrArr.Find(TEXT("Hello"), Index))
@@ -215,7 +215,7 @@ We can find elements by using the Find family of methods.  To check if an elemen
 		// Index == 3
 	}
 
-This sets Index to be the index of the first element found.  If there are duplicate elements and we instead want to find the index of last element, we use the FindLast method instead:
+This sets Index to be the index of the first element found.  If there are duplicate elements and we instead want to find the index of last element, we use the **FindLast** function instead:
 
 	int32 IndexLast;
 	if (StrArr.FindLast(TEXT("Hello"), IndexLast))
@@ -225,7 +225,7 @@ This sets Index to be the index of the first element found.  If there are duplic
 
 Both of these functions return a bool to indicate whether or not an element was found, while also writing the index of that element into a variable when it was found.
 
-Find and FindLast can also return an element index directly. They will do this if you do not pass the index as an explicit argument.  This can be more succinct than the above method and which method you use depends on what suits your particular need or style.
+Find and FindLast can also return an element index directly. They will do this if you do not pass the index as an explicit argument.  This can be more succinct than the above function and which function you use depends on what suits your particular need or style.
 
 If no element was found, the special INDEX_NONE value is returned:
 
@@ -236,28 +236,28 @@ If no element was found, the special INDEX_NONE value is returned:
 	// IndexLast2 == 3
 	// IndexNone  == INDEX_NONE
 
-IndexOfByKey works similarly, but allows comparison of the elements with an arbitrary object.  With the Find methods, the argument is actually converted to the element type (FString in this case) before the search begins.  With IndexOfByKey, the 'key' is compared directly, allowing searches even when the key type isn't directly convertible to the element type.
+**IndexOfByKey** works similarly, but allows comparison of the elements with an arbitrary object.  With the Find functions, the argument is actually converted to the element type (FString in this case) before the search begins.  With IndexOfByKey, the 'key' is compared directly, allowing searches even when the key type isn't directly convertible to the element type.
 
 IndexOfByKey works for any key type for which operator==(ElementType, KeyType) exists; this will then be used to do the comparisons.  IndexOfByKey will return the index of the first found element, or INDEX_NONE if no element was found:
 
 	int32 Index = StrArr.IndexOfByKey(TEXT("Hello"));
 	// Index == 3
 
-The IndexOfByPredicate method can be used to find the index of the first element which matches the specified predicate, again returning the special INDEX_NONE value if none was found:
+The **IndexOfByPredicate** function can be used to find the index of the first element which matches the specified predicate, again returning the special INDEX_NONE value if none was found:
 
 	int32 Index = StrArr.IndexOfByPredicate([](const FString& Str){
 		return Str.Contains(TEXT("r"));
 	});
 	// Index == 2
 
-Instead of returning indices, we can get pointers back to the found elements instead.  FindByKey works like IndexOfByKey by comparing the elements to an arbitrary object, but returns a pointer to the found element.  nullptr is returned if none was found:
+Instead of returning indices, we can get pointers back to the found elements instead.  **FindByKey** works like IndexOfByKey by comparing the elements to an arbitrary object, but returns a pointer to the found element.  nullptr is returned if none was found:
 
 	auto* OfPtr  = StrArr.FindByKey(TEXT("of")));
 	auto* ThePtr = StrArr.FindByKey(TEXT("the")));
 	// OfPtr  == &StrArr[1]
 	// ThePtr == nullptr
 
-Likewise, FindByPredicate can be used like IndexOfByPredicate, except a pointer is returned instead of an index:
+Likewise, **FindByPredicate** can be used like IndexOfByPredicate, except a pointer is returned instead of an index:
 
 	auto* Len5Ptr = StrArr.FindByPredicate([](const FString& Str){
 		return Str.Len() == 5;
@@ -268,7 +268,7 @@ Likewise, FindByPredicate can be used like IndexOfByPredicate, except a pointer
 	// Len5Ptr == &StrArr[2]
 	// Len6Ptr == nullptr
 
-Finally, an array of elements matching a particular predicate will be returned by the FilterByPredicate method:
+Finally, an array of elements matching a particular predicate will be returned by the **FilterByPredicate** function:
 
 	auto Filter = StrArray.FilterByPredicate([](const FString& Str){
 		return !Str.IsEmpty() && Str[0] < TEXT('M');
@@ -276,7 +276,7 @@ Finally, an array of elements matching a particular predicate will be returned b
 
 ##Removal
 
-You can erase elements from the array using the Remove family of methods.  The Remove method will erase all elements equal to the one you pass in:
+You can erase elements from the array using the **Remove** family of functions.  The Remove function will erase all elements equal to the one you pass in:
 
 	StrArr.Remove(TEXT("hello"));
 	// StrArr == ["!","of","Brave","World","Tomorrow"]
@@ -286,7 +286,7 @@ You can erase elements from the array using the Remove family of methods.  The R
 
 Note that "HELLO" was removed even though we asked it to remove "hello".  Equality is tested via the element type's operator==; remember for FString, this is a case-insensitive comparison.
 
-The final element of an array can be removed by the Pop method:
+The final element of an array can be removed by the **Pop** function:
 
 	StrArr.Pop();
 
@@ -302,12 +302,12 @@ Remove removes all instances that compare equal.  For example:
 	ValArr.Remove(20);
 	// ValArr == [10,30,5,10,15,25,30]
 
-You can also use RemoveSingle to erase the closest element to the front of the array.  This is useful if you know your array may contain duplicates and you only want to erase one, or as an optimization if you know that your array can only ever contain one of that element, because it will stop searching after it has found and erased it:
+You can also use **RemoveSingle** to erase the closest element to the front of the array.  This is useful if you know your array may contain duplicates and you only want to erase one, or as an optimization if you know that your array can only ever contain one of that element, because it will stop searching after it has found and erased it:
 
 	ValArr.RemoveSingle(30);
 	// ValArr == [10,5,10,15,25,30]
 
-We can also remove elements by their index by using the RemoveAt method.  The index must exist or it'll cause a runtime error.  Remember that indices are zero-based:
+We can also remove elements by their index by using the **RemoveAt** function.  The index must exist or it'll cause a runtime error.  Remember that indices are zero-based:
 
 	ValArr.RemoveAt(2); // Removes the element at index 2
 	// ValArr == [10,5,15,25,30]
@@ -315,7 +315,7 @@ We can also remove elements by their index by using the RemoveAt method.  The in
 	ValArr.RemoveAt(99); // This will cause a runtime error as
 	                       // there is no element at index 99
 
-We can also remove elements which match a predicate by using the RemoveAll method. For example, removing all values which are a multiple of 3:
+We can also remove elements which match a predicate by using the **RemoveAll** function. For example, removing all values which are a multiple of 3:
 
 	ValArr.RemoveAll([](int32 Val) {
 		return Val % 3 == 0;
@@ -323,10 +323,10 @@ We can also remove elements which match a predicate by using the RemoveAll metho
 	// ValArr == [10,5,25]
 
 [REGION:note]
-In all of these cases when elements were removed, the elements which followed it were shuffled down into lower indices, as there can never be 'holes' left in the array.
+In all of these cases when elements were removed, the elements which followed it were shuffled down into lower indices, as there can never be holes left in the array.
 [/REGION]
 
-The shuffling process has an overhead.  If you don't really care what order the remaining elements are left in, this overhead can be reduced by using the RemoveSwap, RemoveAtSwap and RemoveAllSwap methods, which work like their non-swapping variants except that they don't guarantee the order of the remaining elements, which allows them to be implemented more efficiently:
+The shuffling process has an overhead.  If you don't really care what order the remaining elements are left in, this overhead can be reduced by using the **RemoveSwap**, **RemoveAtSwap** and **RemoveAllSwap** functions, which work like their non-swapping variants except that they don't guarantee the order of the remaining elements, which allows them to be implemented more efficiently:
 
 	TArray<int32> ValArr2;
 	for (int32 i = 0; i != 10; ++i)
@@ -344,14 +344,14 @@ The shuffling process has an overhead.  If you don't really care what order the
 	});
 	// ValArr2 == [1,4,4]
 
-Finally, all elements can be removed by using the Empty method:
+Finally, all elements can be removed by using the **Empty** function:
 
 	ValArr2.Empty();
 	// ValArr2 == []
 
 ##Operators
 
-Arrays are regular value types and as such can be copied via the standard copy constructor or assignment operator.  As arrays strictly own their elements, copying an array is 'deep' and so the new array will have its own copy of the elements:
+Arrays are regular value types and as such can be copied via the standard copy constructor or assignment operator.  As arrays strictly own their elements, copying an array is deep and so the new array will have its own copy of the elements:
 
 	TArray<int32> ValArr3;
 	ValArr3.Add(1);
@@ -364,12 +364,12 @@ Arrays are regular value types and as such can be copied via the standard copy c
 	// ValArr3 == [1,2,3];
 	// ValArr4 == [5,2,3];
 
-As an alternative to the Append method, arrays can be concatenated with the += operator:
+As an alternative to the Append function, arrays can be concatenated with the += operator:
 
 	ValArr4 += ValArr3;
 	// ValArr4 == [5,2,3,1,2,3]
 
-TArray also supports move semantics which can be invoked using the MoveTemp function.  After a move, the source array is guaranteed to be left empty:
+TArray also supports move semantics which can be invoked using the **MoveTemp** function.  After a move, the source array is guaranteed to be left empty:
 
 	ValArr3 = MoveTemp(ValArr4);
 	// ValArr3 == [5,2,3,1,2,3]
@@ -407,7 +407,7 @@ Arrays can be compared using the operator== and operator!=.  The order of the el
 
 TArray has functions which support a binary heap data structure.  A heap is a type of binary tree in which any parent node is equivalent to or ordered before all of its child nodes.  When implemented as an array, the root node of the tree is at element 0 and the indices of the left and right child nodes of a node at index N are 2N+1 and 2N+2 respectively.  The children are not in any particular order with respect to one another.
 
-Any existing array can be turned into a heap by calling the Heapify method.  This is overloaded to take a predicate or not, where the non-predicated version will use the element type's operator< to determine ordering:
+Any existing array can be turned into a heap by calling the **Heapify** function.  This is overloaded to take a predicate or not, where the non-predicated version will use the element type's operator< to determine ordering:
 
 	TArray<int32> HeapArr;
 	for (int32 Val = 10; Val != 0; --Val)
@@ -422,14 +422,14 @@ This is a visualization of the tree:
 
 The nodes in the tree can be read from left-to-right, top-to-bottom as the order of the elements in the heapified array.  Note that the array isn't necessarily sorted after being transformed into a heap.  While a sorted array would also be a valid heap, the heap structure definition is loose enough to allow multiple valid heaps for the same set of elements.
 
-New elements can be added to the heap via the HeapPush method, reordering other nodes to maintain the heap:
+New elements can be added to the heap via the HeapPush function, reordering other nodes to maintain the heap:
 
 	HeapArr.HeapPush(4);
 	// HeapArr == [1,2,4,3,4,5,8,10,7,9,6]
 
 ![image alt text](image_1.png)
 
-The HeapPop and HeapPopDiscard methods are used to remove the top node from the heap.  The difference between the two are that the former takes a reference to an element type to return a copy of the top element, and the latter simply removes the top node without returning it in any way.  Both methods result in the same change to the array, and the heap is again maintained by reordering other elements appropriately:
+The **HeapPop** and **HeapPopDiscard** functions are used to remove the top node from the heap.  The difference between the two are that the former takes a reference to an element type to return a copy of the top element, and the latter simply removes the top node without returning it in any way.  Both functions result in the same change to the array, and the heap is again maintained by reordering other elements appropriately:
 
 	int32 TopNode;
 	HeapArr.HeapPop(TopNode);
@@ -438,7 +438,7 @@ The HeapPop and HeapPopDiscard methods are used to remove the top node from the
 
 ![image alt text](image_2.png)
 
-HeapRemoveAt will remove an element from the array at a given index, and then reorder elements to maintain the heap:
+**HeapRemoveAt** will remove an element from the array at a given index, and then reorder elements to maintain the heap:
 
 	HeapArr.HeapRemoveAt(1);
 	// HeapArr == [2,4,4,6,9,5,8,10,7]
@@ -447,18 +447,18 @@ HeapRemoveAt will remove an element from the array at a given index, and then re
 
 It should be noted that HeapPush, HeapPop, HeapPopDiscard and HeapRemoveAt should only be called when the structure is already a valid heap, e.g. after a Heapify() call, any other heap operation, or by manually manipulating the array into a heap.
 
-Also, each of these methods, including Heapify, can take an optional binary predicate to determine the order of the node elements in the heap.  When using a predicate, the same predicate should be used across all heap operations to maintain the correct structure.  When no predicate is specified, the heap operations use the element type's operator< to determine order.
+Also, each of these functions, including Heapify, can take an optional binary predicate to determine the order of the node elements in the heap.  When using a predicate, the same predicate should be used across all heap operations to maintain the correct structure.  When no predicate is specified, the heap operations use the element type's operator< to determine order.
 
-Finally, the top node of the heap can be inspected using HeapTop, without changing the array:
+Finally, the top node of the heap can be inspected using **HeapTop**, without changing the array:
 
 	int32 Top = HeapArr.HeapTop();
 	// Top == 2
 
 ## Slack
 
-Because arrays can resize, they use a variable amount of memory.  To avoid reallocation every time elements are added, allocators usually provide more memory than was requested so that future Add calls don't pay a performance penalty for reallocation.  Likewise, deleting elements doesn't usually free memory.  The difference between how many elements are in the container and how many elements can be added before the next allocation is known as 'slack'.
+Because arrays can resize, they use a variable amount of memory.  To avoid reallocation every time elements are added, allocators usually provide more memory than was requested so that future Add calls don't pay a performance penalty for reallocation.  Likewise, deleting elements doesn't usually free memory.  The difference between how many elements are in the container and how many elements can be added before the next allocation is known as **slack**.
 
-As a default-constructed array allocates no memory, the slack will initially be zero.  You can find out how much slack there is in an array by using the GetSlack method.  Alternatively, the maximum number of elements that the array can hold before the container reallocates can be obtained by the Max method.  GetSlack() is equivalent to Max() - Num():
+As a default-constructed array allocates no memory, the slack will initially be zero.  You can find out how much slack there is in an array by using the **GetSlack** function.  Alternatively, the maximum number of elements that the array can hold before the container reallocates can be obtained by the Max function.  GetSlack() is equivalent to Max() - Num():
 
 	TArray<int32> SlackArray;
 	// SlackArray.GetSlack() == 0
@@ -480,7 +480,7 @@ As a default-constructed array allocates no memory, the slack will initially be
 
 The amount of slack in a container after a reallocation is decided by the allocator, so it's not something that can be relied upon by users of the array.
 
-Most of the time, you can get away without caring much about slack.  But if you are aware of it, you can use it to your advantage to give the array optimization hints.  For example, if you know you are about to add 100 new elements to the array, you can ensure you have a slack of at least 100 before adding, so that no further allocations will be done while you are adding.  The Empty method, mentioned above, takes an optional slack argument:
+Most of the time, you can get away without caring much about slack.  But if you are aware of it, you can use it to your advantage to give the array optimization hints.  For example, if you know you are about to add 100 new elements to the array, you can ensure you have a slack of at least 100 before adding, so that no further allocations will be done while you are adding.  The Empty function, mentioned above, takes an optional slack argument:
 
 	SlackArray.Empty();
 	// SlackArray.GetSlack() == 0
@@ -497,7 +497,7 @@ Most of the time, you can get away without caring much about slack.  But if you
 	// SlackArray.Num()      == 3
 	// SlackArray.Max()      == 3
 
-There is a Reset method which works similarly to Empty, except that it doesn't free memory if the requested slack is already provided by the current allocation.  However, it will allocate more memory if the requested slack is larger:
+There is a **Reset** function which works similarly to Empty, except that it doesn't free memory if the requested slack is already provided by the current allocation.  However, it will allocate more memory if the requested slack is larger:
 
 	SlackArray.Reset(0);
 	// SlackArray.GetSlack() == 3
@@ -508,7 +508,7 @@ There is a Reset method which works similarly to Empty, except that it doesn't f
 	// SlackArray.Num()      == 0
 	// SlackArray.Max()      == 10
 
-And finally, any wasted slack can be removed by using the Shrink method, which will resize the allocation to be only as large as it needs to be to hold the current sequence of elements, without actually removing the elements themselves:
+And finally, any wasted slack can be removed by using the **Shrink** function, which will resize the allocation to be only as large as it needs to be to hold the current sequence of elements, without actually removing the elements themselves:
 
 	SlackArray.Add(5);
 	SlackArray.Add(10);
@@ -526,9 +526,9 @@ And finally, any wasted slack can be removed by using the Shrink method, which w
 
 TArray is ultimately just a wrapper around some allocated memory.  It can be useful to treat it as such, by direct modification of the bytes of the allocation and by creating elements yourself.  TArray will always try to do the best it can with the information it has, but sometimes you need to drop to a lower level.
 
-It should be noted that these methods allow you to put the container into invalid states and cause undefined behavior if you make mistakes.  It is up to you to return the container to a valid state after invoking these methods but before any other regular method is called.
+It should be noted that these functions allow you to put the container into invalid states and cause undefined behavior if you make mistakes.  It is up to you to return the container to a valid state after invoking these functions but before any other regular function is called.
 
-The AddUninitialized and InsertUninitialized methods will add some uninitialized space to the array.  They work like the Add and Insert methods respectively, but they will not call the constructor of the element type.  This can be useful for structs which have a constructor for safety or convenience, but which you are about to entirely overwrite the state of anyway (e.g. with a Memcpy call), so you don't want the penalty of the construction:
+The **AddUninitialized** and **InsertUninitialized** functions will add some uninitialized space to the array.  They work like the Add and Insert functions respectively, but they will not call the constructor of the element type.  This can be useful for structs which have a constructor for safety or convenience, but which you are about to entirely overwrite the state of anyway (e.g. with a Memcpy call), so you don't want the penalty of the construction:
 
 	int32 SrcInts[] = { 2, 3, 5, 7 };
 	TArray<int32> UninitInts;
@@ -546,7 +546,7 @@ They can also be used for reserving some memory for objects which you plan to ex
 	new ((void*)(UninitStrs.GetData() + 2)) FString(TEXT("C"));
 	// UninitStrs == ["A","B","C","D"]
 
-AddZeroed and InsertZeroed work similarly except they also zero the bytes of the added/inserted space.  This can be useful if you have a type you want to insert of some valid bitwise zero state:
+**AddZeroed** and **InsertZeroed** work similarly except they also zero the bytes of the added/inserted space.  This can be useful if you have a type you want to insert of some valid bitwise zero state:
 
 	struct S
 	{
@@ -564,7 +564,7 @@ AddZeroed and InsertZeroed work similarly except they also zero the bytes of the
 	SArr.AddZeroed();
 	// SArr == [{ Int: 0, Ptr: nullptr, Flt: 0.0f }]
 
-There are also SetNumUninitialized and SetNumZeroed methods which work like SetNum except that, in the case where the new number is greater than the current one, the space for the new elements will be left uninitialized or bitwise-zeroed respectively.  As with the AddUninitialized and InsertUninitialized methods, you should ensure that, if necessary, new elements are properly constructed into the new space if they need to be:
+There are also **SetNumUninitialized** and **SetNumZeroed** functions which work like SetNum except that, in the case where the new number is greater than the current one, the space for the new elements will be left uninitialized or bitwise-zeroed respectively.  As with the AddUninitialized and InsertUninitialized functions, you should ensure that, if necessary, new elements are properly constructed into the new space if they need to be:
 
 	SArr.SetNumUninitialized(3);
 	new ((void*)(SArr.GetData() + 1)) S(5, (void*)0x12345678, 3.14);
@@ -584,13 +584,13 @@ There are also SetNumUninitialized and SetNumZeroed methods which work like SetN
 	//   { Int: 0, Ptr: nullptr,    Flt: 0.0f  }
 	// ]
 
-Use of any of the *Uninitialized or *Zeroed methods should be treated with caution.  If an element type gets modified to include a member that needs construction, or which doesn't have a valid bitwise-zeroed state, it can result in invalid array elements and undefined behavior.  Use of these methods are most useful on arrays of types which will likely never change, like FMatrix or FVector.
+Use of any of the Uninitialized or Zeroed functions should be treated with caution.  If an element type gets modified to include a member that needs construction, or which doesn't have a valid bitwise-zeroed state, it can result in invalid array elements and undefined behavior.  Use of these functions are most useful on arrays of types which will likely never change, like FMatrix or FVector.
 
 ## Misc
 
-The BulkSerialize method is a serialization function which can be used as an alternative operator<< in order to serialize the array as a block of raw bytes instead of doing per-element serialization.  This can be a performance win when your element type is trivial, like a built-in type or a plain data struct.
+The **BulkSerialize** function is a serialization function which can be used as an alternative operator<< in order to serialize the array as a block of raw bytes instead of doing per-element serialization.  This can be a performance win when your element type is trivial, like a built-in type or a plain data struct.
 
-The CountBytes and GetAllocatedSize methods are used to estimate how much memory is currently being utilized by the array.  CountBytes takes an FArchive and GetAllocatedSize can be called directly.  They are typically used for stats reporting.
+The **CountBytes** and **GetAllocatedSize** functions are used to estimate how much memory is currently being utilized by the array.  CountBytes takes an FArchive and GetAllocatedSize can be called directly.  They are typically used for stats reporting.
 
-The Swap and SwapMemory methods both take two indices and will swap the value of the elements at those indices.  They are equivalent except that Swap does some extra error checking on the indices and will assert if the indices are out of range.
+The **Swap** and **SwapMemory** functions both take two indices and will swap the value of the elements at those indices.  They are equivalent except that Swap does some extra error checking on the indices and will assert if the indices are out of range.
 
