#UE4doc #docs:
Updated engine version
Changed font settings for multiple terms.

diff --git a/Engine/Documentation/Source/Programming/UnrealArchitecture/Timers/GameplayTimers.INT.udn b/Engine/Documentation/Source/Programming/UnrealArchitecture/Timers/GameplayTimers.INT.udn
--- a/Engine/Documentation/Source/Programming/UnrealArchitecture/Timers/GameplayTimers.INT.udn
+++ b/Engine/Documentation/Source/Programming/UnrealArchitecture/Timers/GameplayTimers.INT.udn
@@ -2,10 +2,10 @@ Availability:Public
 Title: Gameplay Timers
 Crumbs:%ROOT%, Programming, Programming/UnrealArchitecture/Reference, Programming/UnrealArchitecture/Reference/Functions
 Description:Timer construct for performing actions at set intervals.
-Version: 4.5
+Version: 4.9
 
 [EXCERPT:Overview]
-Timers allow actions to be performed - once or repeatedly - at designated intervals. They can be 
+**Timers** allow actions to be performed - once or repeatedly - at designated intervals. They can be 
 used to set actions to happen on a delay or over a period of time. For example, you may want to 
 make the player invincible after obtaining a certain item and then revert after 10 seconds. Or, 
 you may want to apply damage every second the player is in a room filled with toxic gas. Each of 
@@ -14,53 +14,53 @@ these can be achieved through the use of timers.
 
 ## Timer Management
 
-Timers are managed in a global `FTimerManager`, outside the confines of `AActor`, and can have any 
-of the full range of delegate types assigned. There are several functions in `FTimerManager` available 
+Timers are managed in a global **FTimerManager**, outside the confines of AActor, and can have any 
+of the full range of delegate types assigned. There are several functions in FTimerManager available 
 for managing timers. It is safe to use these functions inside of a timer delegate as the system is ok 
 with manipulating timers while handling a timer. This means, for example, it is ok to set or clear 
 timers inside a timer delegate.
 
-The `AActor::GetWorldTimerManager()` function is used to access the timer manager instance for the 
+The **AActor::GetWorldTimerManager()** function is used to access the timer manager instance for the 
 current world.
 
 ### Creating and Clearing Timers
 
-`FTimerManager::SetTimer()` is used to create a timer to call the given native function at a set interval. 
+**FTimerManager::SetTimer()** is used to create a timer to call the given native function at a set interval. 
 
 Example:
 
 	GetWorldTimerManager().SetTimer(this, &AMatineeActor::CheckPriorityRefresh, 1.0f, true);
 
 [REGION:tip]
-`SetTimer()` can also reset an existing timer. So, if a timer is already set for this delegate, it will 
+SetTimer() can also reset an existing timer. So, if a timer is already set for this delegate, it will 
 update the current timer to the new parameters and reset its elapsed time to 0.
 [/REGION]
 
-`FTimerManager::ClearTimer()` clears a previously set timer so it will not execute any more.
+**FTimerManager::ClearTimer()** clears a previously set timer so it will not execute any more.
 
 Example:
 
 	GetWorldTimerManager().ClearTimer(this, &AMatineeActor::CheckPriorityRefresh);
 
 [REGION:tip]
-Clearing a timer with `ClearTimer()` is identical to calling `SetTimer()` with a <= 0.f rate.
+Clearing a timer with ClearTimer() is identical to calling SetTimer() with a <= 0.f rate.
 [/REGION]
 
 ### Pausing Timers
 
-`FTimerManager::PauseTimer()` pauses a previously set timer in its current state. This means it will 
+**FTimerManager::PauseTimer()** pauses a previously set timer in its current state. This means it will 
 not execute, but the elapsed and remaining time stay the same until the timer is unpaused.
 
-`FTimerManager::UnPauseTimer()` sets a paused timer active once again.
+**FTimerManager::UnPauseTimer()** sets a paused timer active once again.
 
 ## Timer Information
 
-In addition to managing timers, `FTimermanager` also provides funtions for obtaining information - such 
+In addition to managing timers, FTimermanager also provides funtions for obtaining information - such 
 as timer rate, elapsed time, remaining time, etc. - about a specific timer.
 
 ### Is Timer Active
 
-`FTimerManager::IsTimerActive()` is used to determine if the specified timer is curently active and not paused.
+**FTimerManager::IsTimerActive()** is used to determine if the specified timer is curently active and not paused.
 
 Example:
 
@@ -68,7 +68,7 @@ Example:
 
 ### Timer Rate
 
-`FTimerManager::GetTimerRate()` gets the current rate (time between activations) for the specified timer.
+**FTimerManager::GetTimerRate()** gets the current rate (time between activations) for the specified timer.
 
 Example:
 
@@ -76,7 +76,7 @@ Example:
 
 ### Elapsed and Remaining Time
 
-`FTimermanager::GetTimerElapsed()` and `FTimerManager::GetTimerRemaining()` return the elapsed and remaining 
+**FTimermanager::GetTimerElapsed()** and **FTimerManager::GetTimerRemaining()** return the elapsed and remaining 
 time, respectively, for the specified timer. 
 
 Example:
@@ -116,9 +116,9 @@ Each of the timer-manipulation functions has 3 flavors of overloads for interact
 		D.BindDynamic( this, &AMyClass::MyUFunction );
 		GetWorldTimerManager().SetTimer( D, 5.0f, TRUE );
 
-	These delegates are called via the UFunction event mechanism, meaning they will properly respect things like non-native overrides such as in Blueprints.
+	These delegates are called via the UFUNCTION event mechanism, meaning they will properly respect things like non-native overrides such as in Blueprints.
 
 
 ## Known Issues
                
-* The code is not threadsafe at the moment. It will asserts if accessed outside the game thread.
+* The code is not threadsafe at the moment. It will cause an assert if accessed outside the game thread.
