#UE4doc #docs:
Updated engine version
Changed font settings for multiple terms

diff --git a/Engine/Documentation/Source/Programming/UnrealArchitecture/TMap/TMap.INT.udn b/Engine/Documentation/Source/Programming/UnrealArchitecture/TMap/TMap.INT.udn
--- a/Engine/Documentation/Source/Programming/UnrealArchitecture/TMap/TMap.INT.udn
+++ b/Engine/Documentation/Source/Programming/UnrealArchitecture/TMap/TMap.INT.udn
@@ -2,25 +2,25 @@ Title:TMap
 Crumbs:
 Description:
 Availability:Public
-Version:4.7
+Version:4.9
 
-After TArray, the most commonly used container in Unreal Engine is TMap.  This container is an associative container, which means that every key has an associated value, and you can efficiently look up a value object by its key.
+After TArray, the most commonly used container in **Unreal Engine 4** (UE4) is **TMap**.  This container is an associative container, which means that every key has an associated value, and you can efficiently look up a value object by its key.
 
-There are two types of map: TMap and TMultiMap.  TMap's keys are unique, and inserting a new key-value pair when the key already exists will cause the existing pair to be replaced.  TMultiMap's keys are not unique and so existing pairs will not be replaced by new additions.
+There are two types of map: TMap and **TMultiMap**.  TMap's keys are unique, and inserting a new key-value pair when the key already exists will cause the existing pair to be replaced.  TMultiMap's keys are not unique and so existing pairs will not be replaced by new additions.
 
 ## TMap
 
-TMaps are primarily defined by two types - a key type and a value type - which are stored as associated pairs in the map.  It is convenient to refer to these pairs as the element type of the map as if it was an individual object.  In this document, 'element' means a key-value pair, while individual components are referred to as the element's key or the element's value.  The element type is actually a TPair<KeyType, ElementType>, though it should be rare to need to refer to the TPair type directly.
+TMaps are primarily defined by two types - a key type and a value type - which are stored as associated pairs in the map.  It is convenient to refer to these pairs as the element type of the map as if it was an individual object.  In this document, element means a key-value pair, while individual components are referred to as the element's key or the element's value.  The element type is actually a **TPair**< KeyType, ElementType >, though it should be rare to need to refer to the TPair type directly.
 
 Like TArray, TMap is a homogeneous container and so all of its elements are strictly the same type.  TMap is also a value type, and supports the usual copying, assignment and destructor operations, as well as strong ownership of its elements, which will be destroyed when the map is destroyed.  The key type and value type must also be value types.
 
-TMap is a hashing container which means that, by default, the key type must support the GetTypeHash function and provide an operator== for comparing keys for equality.  Hashing is covered in more detail later.
+TMap is a hashing container which means that, by default, the key type must support the **GetTypeHash** function and provide an operator== for comparing keys for equality.  Hashing is covered in more detail later.
 
-TMap can also take an optional allocator to control the memory allocation behavior.  Standard Unreal Engine allocators (e.g. FHeapAllocator, TInlineAllocator) cannot be used as allocators for TMap.  It instead uses set allocators, which define how many hash buckets the map should use and which standard UE allocators should be used for hash and element storage.  See TSetAllocator for more information.
+TMap can also take an optional allocator to control the memory allocation behavior.  Standard UE4 allocators (e.g. FHeapAllocator, TInlineAllocator) cannot be used as allocators for TMap.  It instead uses set allocators, which define how many hash buckets the map should use and which standard UE4 allocators should be used for hash and element storage.  See TSetAllocator for more information.
 
-The final TMap template parameter is KeyFuncs, which tells the map how to retrieve the key from the element type, how to compare two keys for equality and how to hash the key.  These have defaults which will just return a reference to the key, use operator== for equality and the non-member GetTypeHash function for hashing.  If your key type supports these functions, it is usable as a map key without the need to provide a custom KeyFuncs.
+The final TMap template parameter is **KeyFuncs**, which tells the map how to retrieve the key from the element type, how to compare two keys for equality and how to hash the key.  These have defaults which will just return a reference to the key, use operator== for equality and the non-member GetTypeHash function for hashing.  If your key type supports these functions, it is usable as a map key without the need to provide a custom KeyFuncs.
 
-Unlike TArray, the relative order of TMap elements in memory cannot be relied upon and iterating over the elements is likely to return them in a different order from that which they were added.  Elements are also unlikely to be laid out contiguously in memory.  The backing data structure of a map is a sparse array, which is an array with 'holes'.  As elements are removed from the map, holes in the sparse array will appear which are then filled by the addition of future elements.  However, even though TMap doesn't shuffle elements to fill holes, pointers to map elements may still be invalidated, as the entire storage can be reallocated when it is full and new elements are added.
+Unlike TArray, the relative order of TMap elements in memory cannot be relied upon and iterating over the elements is likely to return them in a different order from that which they were added.  Elements are also unlikely to be laid out contiguously in memory.  The backing data structure of a map is a sparse array, which is an array with holes.  As elements are removed from the map, holes in the sparse array will appear which are then filled by the addition of future elements.  However, even though TMap doesn't shuffle elements to fill holes, pointers to map elements may still be invalidated, as the entire storage can be reallocated when it is full and new elements are added.
 
 ## Creating and filling a map
 
@@ -30,7 +30,7 @@ Creating a TMap can be done like this:
 
 This creates an empty TMap designed to map integers to strings.  We have specified neither an allocator nor a KeyFuncs, so the map will do standard heap allocation and compare the key (int32) using == and hash it using GetTypeHash.  No memory has been allocated at this point.
 
-The standard way to populate a map is to use the Add method and provide a key and value:
+The standard way to populate a map is to use the **Add** function and provide a key and value:
 
 	FruitMap.Add(5, TEXT("Banana"));
 	FruitMap.Add(2, TEXT("Grapefruit"));
@@ -54,7 +54,7 @@ This is not a TMultiMap, so keys are guaranteed to be unique.  We can see what h
 
 The map still contains 3 elements, but the previous "Grapefruit" value with a key of 2 has been replaced with "Pear".
 
-The Add method is overloaded to take a key without a value.  If only a key is provided, the value will be default constructed:
+The Add function is overloaded to take a key without a value.  If only a key is provided, the value will be default constructed:
 
 	FruitMap.Add(4);
 	// FruitMap == [
@@ -64,7 +64,7 @@ The Add method is overloaded to take a key without a value.  If only a key is pr
 	// 	{ Key: 4, Value: ""          }
 	// ]
 
-Like TArray, we can also use Emplace instead of Add to avoid the creation of temporaries when inserting into the map:
+Like TArray, we can also use **Emplace** instead of Add to avoid the creation of temporaries when inserting into the map:
 
 	FruitMap.Emplace(3, TEXT("Orange"));
 	// FruitMap == [
@@ -77,7 +77,7 @@ Like TArray, we can also use Emplace instead of Add to avoid the creation of tem
 
 Here, the two arguments are passed directly to the constructors of the key type and value type respectively.  This has no real effect for the int32 here, but it does avoid the creation of a temporary FString for the value.  Unlike TArray, it's only possible to emplace elements into a map with single argument constructors.
 
-It's also possible to insert all the elements from another map by using the Append method to merge them:
+It's also possible to insert all the elements from another map by using the Append function to merge them:
 
 	TMap<int32, FString> FruitMap2;
 	FruitMap2.Emplace(4, TEXT("Kiwi"));
@@ -117,7 +117,7 @@ Iteration over TMaps is similar to TArrays.  You can use C++'s ranged-for featur
 	// (3, "Orange")
 	// (9, "Melon")
 
-Maps also provide their own iterator type for more direct control over iteration.  The CreateIterator method provides read-write access to the elements and the CreateConstIterator method provides read-only access. The iterator objects themselves provide Key() and Value() methods for key and value access.  You may see either form used in code:
+Maps also provide their own iterator type for more direct control over iteration.  The **CreateIterator** function provides read-write access to the elements and the **CreateConstIterator** function provides read-only access. The iterator objects themselves provide **Key()** and **Value()** functions for key and value access.  You may see either form used in code:
 
 	for (auto It = FruitMap.CreateConstIterator(); It; ++It)
 	{
@@ -132,7 +132,7 @@ Maps also provide their own iterator type for more direct control over iteration
 
 ## Queries
 
-We can ask the map how many elements it holds by using the Num method:
+We can ask the map how many elements it holds by using the **Num** function:
 
 	int32 Count = FruitMap.Num();
 	// Count == 6
@@ -143,14 +143,14 @@ We can use the indexing operator[] with a key to retrieve a reference to a value
 	// Val7 == "Pineapple"
 	FString Val8 = FruitMap[8]; // assert!
 
-We can check if a particular key exists in the map using the Contains method:
+We can check if a particular key exists in the map using the **Contains** function:
 
 	bool bHas7 = FruitMap.Contains(7);
 	bool bHas8 = FruitMap.Contains(8);
 	// bHas7 == true
 	// bHas8 == false
 
-Most of the time you will want to look up elements without knowing whether or not the key exists.  Using Contains followed by operator[] means a double lookup of the key, which we don't really want to do.  The Find method will allow you to do a single lookup, returning a pointer to the value of the found element rather than a reference, and returning null when the key does not exist:
+Most of the time you will want to look up elements without knowing whether or not the key exists.  Using Contains followed by operator[] means a double lookup of the key, which we don't really want to do.  The **Find** function will allow you to do a single lookup, returning a pointer to the value of the found element rather than a reference, and returning null when the key does not exist:
 
 	FString* Ptr7 = FruitMap.Find(7);
 	FString* Ptr8 = FruitMap.Find(8);
@@ -159,7 +159,7 @@ Most of the time you will want to look up elements without knowing whether or no
 
 If called on a const map, the returned pointer will also be const.
 
-The FindOrAdd method will search for a given key and return a reference to the associated value; if the key doesn't exist, it will add it with a default constructed value before returning a reference to that.  Because it may need to add, this method cannot be called on a non-const map:
+The **FindOrAdd** function will search for a given key and return a reference to the associated value; if the key doesn't exist, it will add it with a default constructed value before returning a reference to that.  Because it may need to add, this function cannot be called on a non-const map:
 
 	FString& Ref7 = FruitMap.FindOrAdd(7);
 	// Ref7     == "Pineapple"
@@ -185,7 +185,7 @@ The FindOrAdd method will search for a given key and return a reference to the a
 
 Note that here the Ref7 reference may have been invalidated by the call to FruitMap.FindOrAdd(8) if a reallocation occurred.
 
-Despite its name, the FindRef method searches for a key and returns a value rather than a reference.  If the key was found, a copy of the associated value is returned, otherwise a default-constructed value type is returned.  This results in similar behavior to FindOrAdd but, because FindRef method returns a value rather than a reference, the map will not be modified and so can be called on a const object:
+Despite its name, the **FindRef** function searches for a key and returns a value rather than a reference.  If the key was found, a copy of the associated value is returned, otherwise a default-constructed value type is returned.  This results in similar behavior to FindOrAdd but, because FindRef function returns a value rather than a reference, the map will not be modified and so can be called on a const object:
 
 	FString Val7 = FruitMap.FindRef(7);
 	FString Val6 = FruitMap.FindRef(6);
@@ -201,14 +201,14 @@ Despite its name, the FindRef method searches for a key and returns a value rath
 	// 	{ Key: 8, Value: ""          }
 	// ]
 
-The FindKey method allows you to perform a reverse lookup (find a key given a value).  Be aware when using this method since, unlike keys, values are not hashed and therefore looking up a key is a linear operation.  Also, values are not guaranteed to be unique, so the key returned for a particular value will be arbitrary if the map contains duplicate values:
+The **FindKey** function allows you to perform a reverse lookup (find a key given a value).  Be aware when using this function since, unlike keys, values are not hashed and therefore looking up a key is a linear operation.  Also, values are not guaranteed to be unique, so the key returned for a particular value will be arbitrary if the map contains duplicate values:
 
 	const int32* KeyMangoPtr   = FruitMap.FindKey(TEXT("Mango"));
 	const int32* KeyKumquatPtr = FruitMap.FindKey(TEXT("Kumquat"));
 	// *KeyMangoPtr   == 5
 	//  KeyKumquatPtr == nullptr
 
-The GenerateKeyArray and GenerateValueArray methods allow an array to be populated with a copy of all the keys and values respectively.  In both cases, the array being passed is emptied before population, so the resulting number of elements will always equal the number of elements in the map:
+The **GenerateKeyArray** and **GenerateValueArray** functions allow an array to be populated with a copy of all the keys and values respectively.  In both cases, the array being passed is emptied before population, so the resulting number of elements will always equal the number of elements in the map:
 
 	TArray<int32>   FruitKeys;
 	TArray<FString> FruitValues;
@@ -222,7 +222,7 @@ The GenerateKeyArray and GenerateValueArray methods allow an array to be populat
 
 ## Removal
 
-Elements can be removed from maps by using the Remove method and providing the key of the element to remove:
+Elements can be removed from maps by using the **Remove** function and providing the key of the element to remove:
 
 	FruitMap.Remove(8);
 	// FruitMap == [
@@ -234,9 +234,9 @@ Elements can be removed from maps by using the Remove method and providing the k
 	// 	{ Key: 9, Value: "Melon"     }
 	// ]
 
-Removing elements actually leaves 'holes' in the data structure, which you can see when visualizing the map in Visual Studio's watch window, but they have been omitted here for clarity.
+Removing elements actually leaves holes in the data structure, which you can see when visualizing the map in Visual Studio's watch window, but they have been omitted here for clarity.
 
-The FindAndRemoveChecked method can be used to remove an element by and return the associated value.  The 'Checked' part of the name means that the key will be checked for existence and assert if it doesn't exist:
+The **FindAndRemoveChecked** function can be used to remove an element by and return the associated value.  The checked part of the name means that the key will be checked for existence and assert if it doesn't exist:
 
 	FString Removed7 = FruitMap.FindAndRemoveChecked(7);
 	// Removed7 == "Pineapple"
@@ -250,7 +250,7 @@ The FindAndRemoveChecked method can be used to remove an element by and return t
 
 	FString Removed8 = FruitMap.FindAndRemoveChecked(8); // assert!
 
-The RemoveAndCopyValue method does a similar job, but takes a reference to a value type to be passed out and returns a bool to say whether or not the key was found.  This allows it to be used with missing keys without a runtime error.  If the key was not found, the call returns false and the passed object and map remain unchanged:
+The **RemoveAndCopyValue** function does a similar job, but takes a reference to a value type to be passed out and returns a bool to say whether or not the key was found.  This allows it to be used with missing keys without a runtime error.  If the key was not found, the call returns false and the passed object and map remain unchanged:
 
 	FString Removed;
 	bool bFound2 = FruitMap.RemoveAndCopyValue(2, Removed);
@@ -272,7 +272,7 @@ The RemoveAndCopyValue method does a similar job, but takes a reference to a val
 	// 	{ Key: 9, Value: "Melon"  }
 	// ]
 
-Finally, all elements can be removed by using the Empty method:
+Finally, all elements can be removed by using the **Empty** function:
 
 	TMap<int32, FString> FruitMapCopy = FruitMap;
 	// FruitMapCopy == [
@@ -291,7 +291,7 @@ Like TArray, Empty takes an optional slack value which can be used to optimize t
 
 TMaps can temporarily be sorted.  The next iteration over the map will present the elements in sorted order, though future modifications to the map will likely result in the map being unsorted again.  Sorting is unstable, so equivalent elements may appear in any order.
 
-You can sort by key or by value using the KeySort and ValueSort methods respectively, and both methods take a binary predicate which specifies the sort order:
+You can sort by key or by value using the **KeySort** and **ValueSort** functions respectively, and both functions take a binary predicate which specifies the sort order:
 
 	FruitMap.KeySort([](int32 A, int32 B) {
 		return A > B; // sort keys in reverse
@@ -315,7 +315,7 @@ You can sort by key or by value using the KeySort and ValueSort methods respecti
 
 ## Operators
 
-Like TArray, TMap is a regular value type and as such can be copied via the standard copy constructor or assignment operator.  As maps strictly own their elements, copying a map is 'deep' and so the new map will have its own copy of the elements:
+Like TArray, TMap is a regular value type and as such can be copied via the standard copy constructor or assignment operator.  As maps strictly own their elements, copying a map is deep and so the new map will have its own copy of the elements:
 
 	TMap<int32, FString> NewMap = FruitMap;
 	NewMap[5] = "Apple";
@@ -332,7 +332,7 @@ Like TArray, TMap is a regular value type and as such can be copied via the stan
 	// 	{ Key: 9, Value: "Melon" }
 	// ]
 
-TMap also supports move semantics which can be invoked using the MoveTemp function.  After a move, the source map is guaranteed to be left empty:
+TMap also supports move semantics which can be invoked using the **MoveTemp** function.  After a move, the source map is guaranteed to be left empty:
 
 	FruitMap = MoveTemp(NewMap);
 	// FruitMap == [
@@ -344,14 +344,14 @@ TMap also supports move semantics which can be invoked using the MoveTemp functi
 
 ## Slack
 
-TMaps also have a notion of slack and can be used to optimize the population of the map.  The Reset method acts like a Empty() call but does not free the memory that was previously used by the elements:
+TMaps also have a notion of **slack** and can be used to optimize the population of the map.  The **Reset** function acts like a Empty() call but does not free the memory that was previously used by the elements:
 
 	FruitMap.Reset();
 	// FruitMap == [<invalid>, <invalid>, <invalid>]
 
 Here the map is emptied in the same way as Empty would, but the memory used for storage is not freed and remains as slack.
 
-TMap does not provide a way of checking how many elements are preallocated, like TArray::Max() does, but still supports preallocating slack.  The Reserve method can be used to preallocate slack for a particular number of elements before adding:
+TMap does not provide a way of checking how many elements are preallocated, like TArray::Max() does, but still supports preallocating slack.  The Reserve function can be used to preallocate slack for a particular number of elements before adding:
 
 	FruitMap.Reserve(10);
 	for (int32 i = 0; i != 10; ++i)
@@ -368,7 +368,7 @@ TMap does not provide a way of checking how many elements are preallocated, like
 
 Note that the slack happens to have caused the new elements to be added in reverse order. This is an example of why element order in maps should not be relied upon in any way.
 
-The Shrink method also works like the TArray equivalent in that it removes any wasted slack from the end of the container.  However, as TMap allows 'holes' in its data structure, this will only remove slack from holes left at the end of the structure:
+The **Shrink** function also works like the TArray equivalent in that it removes any wasted slack from the end of the container.  However, as TMap allows holes in its data structure, this will only remove slack from holes left at the end of the structure:
 
 	for (int32 i = 0; i != 10; i += 2)
 	{
@@ -399,7 +399,7 @@ The Shrink method also works like the TArray equivalent in that it removes any w
 	// 	{ Key: 1, Value: "Fruit1" }
 	// ]
 
-Note that only one invalid element has been removed from the Shrink call, because there was only one hole at the end.  The Compact method can be used to remove all the holes before shrinking:
+Note that only one invalid element has been removed from the Shrink call, because there was only one hole at the end.  The **Compact** function can be used to remove all the holes before shrinking:
 
 	FruitMap.Compact();
 	// FruitMap == [
@@ -424,15 +424,15 @@ Note that only one invalid element has been removed from the Shrink call, becaus
 
 ## KeyFuncs
 
-As long as a type has an operator== and a non-member GetTypeHash overload then it can be used as a KeyType for a TMap without any changes.  However, it may be useful to use types as keys where it is undesirable to overload those functions.  In these cases, you can provide your own custom KeyFuncs.
+As long as a type has an operator== and a non-member GetTypeHash overload then it can be used as a KeyType for a TMap without any changes.  However, it may be useful to use types as keys where it is undesirable to overload those functions.  In these cases, you can provide your own custom **KeyFuncs**.
 
 KeyFuncs requires the definition of 2 typedefs and 3 static functions:
 
-* KeyInitType - used to pass keys around
-* ElementInitType - used to pass elements around
-* KeyInitType GetSetKey(ElementInitType Element) - returns the key of an element.
-* bool Matches(KeyInitType A, KeyInitType B) - returns if A and B are equivalent
-* uint32 GetKeyHash(KeyInitType Key) - returns the hash value of Key
+* **KeyInitType** - used to pass keys around
+* **ElementInitType** - used to pass elements around
+* **KeyInitType GetSetKey(ElementInitType Element)** - returns the key of an element.
+* **bool Matches(KeyInitType A, KeyInitType B)** - returns if A and B are equivalent
+* **uint32 GetKeyHash(KeyInitType Key)** - returns the hash value of Key
 
 KeyInitType and ElementInitType are typedefs to the normal passing convention of the key type and element type.  Usually these will be a value for trivial types and a const reference for non-trivial types.  Remember that the element type of a map is a TPair.
 
@@ -487,13 +487,13 @@ An example of a custom KeyFuncs might look like this:
 
 Here, we have a type which has a unique identifier as part of its state, but also has some other state that does not contribute to its identity.  GetTypeHash and operator== would be inappropriate here, as it would be a lie for operator== to ignore part of the state and GetTypeHash needs to match operator==, which it couldn't if operator== was defined properly.  However, for the purpose of identifying this type in a map, we are happy to just use the UniqueID part of the state.
 
-First, we inherit BaseKeyFuncs as it helpfully defines some stuff for us, including KeyInitType and ElementInitType.  We pull these directly from Super into the derived class so that we can use them in the rest of the implementation.
+First, we inherit **BaseKeyFuncs** as it helpfully defines some stuff for us, including KeyInitType and ElementInitType.  We pull these directly from Super into the derived class so that we can use them in the rest of the implementation.
 
 BaseKeyFuncs takes two template parameters: the element type of the map and the type of our key: the object being returned from GetSetKey.  As with all maps, the element type is a TPair, taking FMyStruct as its KeyType and TMyStructMapKeyFuncs's template parameter as its ValueType.  We make our replacement KeyFuncs a template to allow the ValueType to be specified on a per-map basis, rather than needing to define a new KeyFuncs every time we want to create a TMap keyed on FMyStruct.  The second BaseKeyFuncs argument is the type of the key, not to be confused with TPair's 'KeyType', which is what is being stored in the Key field of the element.  As we want to use FMyStruct::UniqueID as our key, we specify FString here.
 
 We then define our three required KeyFuncs static functions.  The first is GetSetKey which, given an element type, returns the key.  Our element type is TPair, and our key is UniqueID, so we simply return that directly.
 
-The second static function is Matches which takes the keys of two elements, having already been extracted from the element types using GetSetKey, and compares them to see if they are equivalent.  FString's operator== is case-insensitive, and we want a case sensitive search, so we use the FString::Compare method with an appropriate option.
+The second static function is Matches which takes the keys of two elements, having already been extracted from the element types using GetSetKey, and compares them to see if they are equivalent.  FString's operator== is case-insensitive, and we want a case sensitive search, so we use the FString::Compare function with an appropriate option.
 
 Finally, the GetKeyHash static function takes an extracted key and returns a hashed value for it.  Again, the GetTypeHash behavior for FString is to ignore case, so we call a case-sensitive FCrc function to calculate our hash for us.
 
@@ -531,7 +531,7 @@ A note about providing your own KeyFuncs: be aware that TMap assumes that two it
 
 ## Misc
 
-The CountBytes and GetAllocatedSize methods are used to estimate how much memory is currently being utilized by the array.  CountBytes takes an FArchive and GetAllocatedSize can be called directly.  They are typically used for stats reporting.
+The **CountBytes** and **GetAllocatedSize** functions are used to estimate how much memory is currently being utilized by the array.  CountBytes takes an FArchive and GetAllocatedSize can be called directly.  They are typically used for stats reporting.
 
-The Dump method takes an FOutputDevice and writes out some implementation information about the contents of the map.  It's usually used for debugging.
+The **Dump** function takes an FOutputDevice and writes out some implementation information about the contents of the map.  It's usually used for debugging.
 
