#UE4 Docs: Replacing the old UE3 version of the image with a new UE4 version.

diff --git a/Engine/Documentation/Source/Engine/Rendering/Materials/ExpressionReference/Custom/CustomExpressions.INT.udn b/Engine/Documentation/Source/Engine/Rendering/Materials/ExpressionReference/Custom/CustomExpressions.INT.udn
--- a/Engine/Documentation/Source/Engine/Rendering/Materials/ExpressionReference/Custom/CustomExpressions.INT.udn
+++ b/Engine/Documentation/Source/Engine/Rendering/Materials/ExpressionReference/Custom/CustomExpressions.INT.udn
@@ -2,6 +2,8 @@ Availability:Public
 Title:Custom Expressions
 Crumbs: %ROOT%, Engine, Engine/Rendering/Materials, Engine/Rendering/Materials/ExpressionReference
 Description:Expressions that allow the use of custom, plain shader code.
+Skilllevel: Advanced
+Version: 4.8
 
 [EXCERPT:TOC]
 [TOC(start:2)]
@@ -34,13 +36,13 @@ Specifies the name of the input. This is the name displayed on the expression in
 [/EXCERPT:CustomInputDesc]
 -->
 
-Add as many inputs as you need to the Inputs array, and name them. You can then write code in the Code property. You can type either a full function body with return statements as shown in the example, or a simple expression such as "Input.bgr". You must also specify the output data type in OutputType.
+Add as many inputs as you need to the Inputs array, and name them. You can then write code in the Code property. You can type either a full function body with return statements as shown in the example, or a simple expression such as **Input.bgr**. You must also specify the output data type in OutputType.
 
-![Custom.jpg](Custom.jpg)
+![](T_Custom_Node.png)
 
 (#Warnings)
 [REGION:warning]
-**Using the custom node prevents constant folding and may use significantly more instructions than an equivalent version done with built in nodes!**  Constant folding is an optimization that UE4 employs under the hood to reduce shader instruction count when necessary.  For example, an expression chain of 'Time->Sin->Mul by parameter->Add to something' can and will be collapsed by UE4 into a single instruction, the final add.  This is possible because all of the inputs of that expression (Time, parameter) are constant for the whole draw call, they do not change per-pixel.  UE4 cannot collapse anything in a custom node, which can produce less efficient shaders than an equivalent version made out of existing nodes.  As a result, it is best to only use the custom node when it gives you access to functionality not possible with the existing nodes.
+**Using the custom node prevents constant folding and may use significantly more instructions than an equivalent version done with built in nodes!**  Constant folding is an optimization that UE4 employs under the hood to reduce shader instruction count when necessary.  For example, an expression chain of Time >Sin >Mul by parameter > Add to something' can and will be collapsed by UE4 into a single instruction, the final add.  This is possible because all of the inputs of that expression (Time, parameter) are constant for the whole draw call, they do not change per-pixel.  UE4 cannot collapse anything in a custom node, which can produce less efficient shaders than an equivalent version made out of existing nodes.  As a result, it is best to only use the custom node when it gives you access to functionality not possible with the existing nodes.
 [/REGION]
 
 
