#UE4doc #docs - Removed overview header.

diff --git a/Engine/Documentation/Source/Engine/Actors/DecalActor/DecalUserGuide.INT.udn b/Engine/Documentation/Source/Engine/Actors/DecalActor/DecalUserGuide.INT.udn
--- a/Engine/Documentation/Source/Engine/Actors/DecalActor/DecalUserGuide.INT.udn
+++ b/Engine/Documentation/Source/Engine/Actors/DecalActor/DecalUserGuide.INT.udn
@@ -44,7 +44,6 @@ Description:A guide to using the Deferred Decal actor.
 [/VAR]
 [TOC(start:2)]
 
-## Overview
 
 Deferred decals offer better performance and easier maintenance. Writing to the GBuffer instead of recalculating lighting has several benefits:
 
