#UE4doc #docs - removed overview header.

diff --git a/Engine/Documentation/Source/Engine/Actors/Transform/ActorTransform.INT.udn b/Engine/Documentation/Source/Engine/Actors/Transform/ActorTransform.INT.udn
--- a/Engine/Documentation/Source/Engine/Actors/Transform/ActorTransform.INT.udn
+++ b/Engine/Documentation/Source/Engine/Actors/Transform/ActorTransform.INT.udn
@@ -66,8 +66,6 @@ Description:Modifying the location, rotation, and scale of Actors placed in leve
 
 [TOC(start:2)]
 
-## Overview
-
 Transforming Actors refers to moving, rotating, or scaling them, which is an important part of the level-editing process. There are two basic ways to transform Actors in the Unreal Editor. Both ways apply the transformations to all currently selected Actors.
 
 ## Manual Transformation
