#UE4doc #docs - Adjusted bold lettering.

diff --git a/Engine/Documentation/Source/Engine/AI/BehaviorTrees/BehaviorTrees.INT.udn b/Engine/Documentation/Source/Engine/AI/BehaviorTrees/BehaviorTrees.INT.udn
--- a/Engine/Documentation/Source/Engine/AI/BehaviorTrees/BehaviorTrees.INT.udn
+++ b/Engine/Documentation/Source/Engine/AI/BehaviorTrees/BehaviorTrees.INT.udn
@@ -48,10 +48,10 @@ They're a combination of two asset types: The **Blackboard** and the **Behavior
 
 | Asset Type | Description |
 | --- | --- |
-| **Blackboard** | The **Blackboard** is the AI's memory. It stores key values for the Behavior Tree to use. |
-| **Behavior Tree** | The **Behavior Tree** is the AI's processor. It makes decisions, and then acts on those decisions. |
+| **Blackboard** | The Blackboard is the AI's memory. It stores key values for the Behavior Tree to use. |
+| **Behavior Tree** | The Behavior Tree is the AI's processor. It makes decisions, and then acts on those decisions. |
 
-Below are links to the sections in the documentation covering the use of Behavior Trees. The [Quick Start Guide](Engine/AI/BehaviorTrees/QuickStart) takes you through making an AI chase the Player. Then there's the [User Guide](Engine/AI/BehaviorTrees/UserGuide), which contains instructions on the most common tasks when using Behavior Trees. Next, the [Node Reference](Engine/AI/BehaviorTrees/NodeReference) contains a listing of all then nodes available in the Behavior Tree editor, what they do, and what their properties are. Finally, [How Unreal Engine 4 Behavior Trees Differ](Engine/AI/BehaviorTrees/HowUE4BehaviorTreesDiffer) is a paper written for those who are already familiar with how other Behavior Trees work.
+Below are links to the sections in the documentation covering the use of Behavior Trees. The [**Quick Start Guide**](Engine/AI/BehaviorTrees/QuickStart) takes you through making an AI chase the Player. Then there's the [**User Guide**](Engine/AI/BehaviorTrees/UserGuide), which contains instructions on the most common tasks when using Behavior Trees. Next, the [**Node Reference**](Engine/AI/BehaviorTrees/NodeReference) contains a listing of all then nodes available in the Behavior Tree editor, what they do, and what their properties are. Finally, [**How Unreal Engine 4 Behavior Trees Differ**](Engine/AI/BehaviorTrees/HowUE4BehaviorTreesDiffer) is a paper written for those who are already familiar with how other Behavior Trees work.
 
 [OBJECT:TopicButtonList]
 	[PARAM:icon]
