#UE4doc #docs - changed capitalization of word.

diff --git a/Engine/Documentation/Source/Engine/AI/BehaviorTrees/NodeReference/Decorators/Decorators.INT.udn b/Engine/Documentation/Source/Engine/AI/BehaviorTrees/NodeReference/Decorators/Decorators.INT.udn
--- a/Engine/Documentation/Source/Engine/AI/BehaviorTrees/NodeReference/Decorators/Decorators.INT.udn
+++ b/Engine/Documentation/Source/Engine/AI/BehaviorTrees/NodeReference/Decorators/Decorators.INT.udn
@@ -70,7 +70,7 @@ The Compare Blackboard Entries nodes will compare the values of two Blackboard k
 
 ![](Composite.png)
 
-The Composite Decorator Node enables you to setup more advanced logic than the built in nodes but not as complex as a full Blueprint. Once you have added a Composite Decorator to a node, **double-click** the Composite Decorator to bring up the Composite's Graph. By **right-clicking** in the graph area you can add Decorator Nodes as stand alone nodes, then wire them together through `AND` nodes, `OR` nodes, and `NOT` nodes, to create more advanced logic:
+The Composite Decorator node enables you to setup more advanced logic than the built in nodes but not as complex as a full Blueprint. Once you have added a Composite Decorator to a node, **double-click** the Composite Decorator to bring up the Composite's Graph. By **right-clicking** in the graph area you can add Decorator Nodes as stand alone nodes, then wire them together through `AND` nodes, `OR` nodes, and `NOT` nodes, to create more advanced logic:
 
 ![](CompositeGraph.png)(w:940)
 
