#UE4doc #docs - Adjusted bold lettering.

diff --git a/Engine/Documentation/Source/Engine/AI/BehaviorTrees/NodeReference/Tasks/Tasks.INT.udn b/Engine/Documentation/Source/Engine/AI/BehaviorTrees/NodeReference/Tasks/Tasks.INT.udn
--- a/Engine/Documentation/Source/Engine/AI/BehaviorTrees/NodeReference/Tasks/Tasks.INT.udn
+++ b/Engine/Documentation/Source/Engine/AI/BehaviorTrees/NodeReference/Tasks/Tasks.INT.udn
@@ -26,7 +26,7 @@ The Move To node will cause a Pawn with a Character Movement component to move u
 
 | Property | Description |
 | -------- | ----------- |
-| **Acceptable Radius** | How close does the pawn have to be to the target for the **Move To** Task to succeed. |
+| **Acceptable Radius** | How close does the pawn have to be to the target for the Move To Task to succeed. |
 | **Filter Class** | Which navigation data should be used? If set to None the default navigation data will be used. |
 | **Allow Strafe** | Should the pawn be allowed to strafe? |
 | **Blackboard Key** | The location to move to (Vector3). |
