#UEdoc #docs Updated grammar and added image showing Linux shell script usage.

diff --git a/Engine/Documentation/Source/GettingStarted/DownloadingUnrealEngine/DownloadZip/DownloadZip.INT.udn b/Engine/Documentation/Source/GettingStarted/DownloadingUnrealEngine/DownloadZip/DownloadZip.INT.udn
--- a/Engine/Documentation/Source/GettingStarted/DownloadingUnrealEngine/DownloadZip/DownloadZip.INT.udn
+++ b/Engine/Documentation/Source/GettingStarted/DownloadingUnrealEngine/DownloadZip/DownloadZip.INT.udn
@@ -1,5 +1,5 @@
 Availability:NoPublish
-Title:Downloading Unreal Engine As a Zip File
+Title:Downloading Unreal Engine as a Zip File
 Crumbs:%ROOT%, GettingStarted, GettingStarted/DownloadingUnrealEngine
 Description:Downloading the Unreal Engine source code as a Zip file.
 
@@ -20,6 +20,7 @@ Description:Downloading the Unreal Engine source code as a Zip file.
 		[REGION:warning]
 		The built-in Windows zip utility marks downloaded zip file contents as unsafe to execute.  
 		Before decompressing the zip file:
+
 		1.	**Right-click** the zip file.
 		1.	**Click 'Properties...'**.
 		1.	**Select 'Unblock'**.
@@ -79,7 +80,13 @@ Description:Downloading the Unreal Engine source code as a Zip file.
 	1.	**Click** the **Download ZIP** button (pictured below).
 		![](DownloadZipButton.png)
 	1.	**Decompress** 'UnrealEngine-(branch_name).zip' onto your hard disk.
-	1.	**Open** your source folder and **run Setup.sh** to **download** the engine's binary content.
+	1.	Inside the root folder, **run Setup.sh** from the terminal to **download** the engine's binary content (example pictured below).
+
+		![](RunSetupShellScript_Linux.png)
+
+		[REGION:note]
+		If you don't see a ** \*\*SUCCESS\*\* ** message after running this step, refer to the **BuildThirdParty.log** file for more information. BuildThirdParty.log is located in the `Engine/Build/BatchFiles/Linux` directory.
+		[/REGION]
 	[/PARAM]
 [/OBJECT]
 
