#UE4doc #docs Updated grammar and added image showing Linux Setup shell script usage.

diff --git a/Engine/Documentation/Source/GettingStarted/DownloadingUnrealEngine/GitHub/GitHub.INT.udn b/Engine/Documentation/Source/GettingStarted/DownloadingUnrealEngine/GitHub/GitHub.INT.udn
--- a/Engine/Documentation/Source/GettingStarted/DownloadingUnrealEngine/GitHub/GitHub.INT.udn
+++ b/Engine/Documentation/Source/GettingStarted/DownloadingUnrealEngine/GitHub/GitHub.INT.udn
@@ -77,7 +77,13 @@ Description:Downloading the Unreal Engine source code from GitHub.
 	1.	**Set up** [Git](https://help.github.com/articles/set-up-git/).
 	1.	**Fork** and **Clone** our [repository](https://help.github.com/articles/fork-a-repo/) (pictured below).
 		![](UnrealRepo.png)
-	1.	**Open** your source folder and **run Setup.sh** to **download** the engine's binary content.
+	1.	Inside the root folder, **run Setup.sh** from the terminal to **download** the engine's binary content (example pictured below).
+
+		![](RunSetupShellScript_Linux.png)
+
+		[REGION:note]
+		If you don't see a ** \*\*SUCCESS\*\* ** message after running this step, refer to the **BuildThirdParty.log** file for more information. BuildThirdParty.log is located in the `Engine/Build/BatchFiles/Linux` directory.
+		[/REGION]
 	[/PARAM]
 [/OBJECT]
 
