#UE4doc #docs Updated first draft of Building Unreal Engine from source documentation.

diff --git a/Engine/Documentation/Source/Programming/Development/BuildingUnrealEngine/BuildingUnrealEngine.INT.udn b/Engine/Documentation/Source/Programming/Development/BuildingUnrealEngine/BuildingUnrealEngine.INT.udn
--- a/Engine/Documentation/Source/Programming/Development/BuildingUnrealEngine/BuildingUnrealEngine.INT.udn
+++ b/Engine/Documentation/Source/Programming/Development/BuildingUnrealEngine/BuildingUnrealEngine.INT.udn
@@ -1,14 +1,14 @@
 Availability:Public
-Title:Building Unreal Engine
+Title:Building Unreal Engine from Source
 Crumbs: %ROOT%, Programming, Programming/Development
 Description:Compiling the Unreal Engine.
 Version: 4.9
 
-%Globals:OSSelect%
+%Globals:osselect_linux%
 
 [TOC (start:2 end:3)]
 
-## Setting Up Unreal Engine
+## Building Unreal Engine from Source
 
 [EXCERPT:Main]
 [OBJECT:ToggleButtonContent]
@@ -23,22 +23,48 @@ Version: 4.9
 	[/PARAMLITERAL]
 	[PARAM:content]
 		[REGION:note]
-		Please make sure you are running Microsoft Visual Studio 2013 before proceeding.
+		Make sure that Microsoft Visual Studio 2013 is installed prior to building Unreal Engine from Source.
 		[/REGION]
 
-		Before you dive right in and start getting your hands dirty with the engine, there are a few basic setup steps to perform:
-
-		1.	Run the **GenerateProjectFiles.bat** batch file located in the root directory.
+		1.	Inside the root directory, **run GenerateProjectFiles.bat** to set-up your project files.
 
 			[REGION:warning]
-				All of the project files and even the UE4 solution are intermediate files - located in `[UE4Root]\Engine\Intermediate\ProjectFiles`. You must generate them each time you sync a new build to ensure they exist and are up to date. This also means if you delete your `Intermediate` folder, you must regenerate the project files using the above-mentioned batch file.
+			All project files are intermediate (`[UE4Root]\Engine\Intermediate\ProjectFiles`). You must generate project files each time you sync a new build to ensure they are up to date If you delete your `Intermediate` folder, you must regenerate project files using the above-mentioned batch file.
+			[/REGION]
+
+		1.	**Load** the project into Visual Studio by **double-clicking UE4.sln**.
+
+		1.	Set your solution configuration to **Development Editor** (pictured below).
+
+			![](configMenu_updated.png)
+
+		1.	If you're developing on the 64-bit version of the Windows operating system, set your solution platform to **Win64** (pictured below).
+
+			![](platform_menu_updated.png)
+
+		1.	**Right-click** the UE4 target and select **Build** (pictured below).
+
+			![](BuildUE4.png)
+
+			[REGION:note]
+			Depending on your system's specifications, it may take anywhere between 10 and 40 minutes to finish compiling the engine from source.
 			[/REGION]
 
-		1.	Open the **UE4.sln** Visual Studio solution located in the root directory. This is the main Visual Studio solution for the entire engine. The projects for the engine core and games are contained within this solution. 
+		1.	After compiling the engine, set your startup project to **UE4** (pictured below).
+
+			![](SetUE4_StartPrj.png)
+
+		1.	**Right-click** the **UE4** project, **hover over 'Debug'** and **select 'Start New Instance'** to launch the compiled editor in debug mode (pictured below).
 
-		1.	**Right-click** the **UE4** project and choose **Set As StartUp Project**.
+			![](RunCompiledWindowsEditor.png)
 
-		1.	Save the **UE4.sln**. Make sure to save the changes you have made so the next time you open the solution you will not have to go through the process all over again.
+			[REGION:tip]
+			Alternatively, you can **press the F5 key** to start a new instance of the compiled editor in debug mode.
+			[/REGION]
+
+		1.	Congratulations! You've compiled and launched the engine from source (result pictured below).
+
+			![](CompiledEditor_Windows.png)
 	[/PARAM]
 [/OBJECT]
 [OBJECT:ToggleButtonContent]
@@ -52,26 +78,75 @@ Version: 4.9
 	[/PARAMLITERAL]
 	[PARAM:content]
 		[REGION:note]
-		Please make sure you are running Xcode 6.1 before proceeding.
+		Make sure that the latest version of XCode is installed prior to building Unreal Engine from Source.
 		[/REGION]
+		
+		1.	Inside the root directory, **run GenerateProjectFiles.command** to set-up your project files.
+		
+		1.	Load the project into XCode by **double-clicking UE4.xcodeproj**.
+
+		1.	To set your build target, select **UE4Editor - Mac** for **My Mac** in the title bar (pictured below).
 
-		Before you dive right in and start getting your hands dirty with the engine, you will need project files in order to compile: 
+			![](SettingMacTarget.png)
+
+		1.	To build the project, select **Product > Build** (pictured below).
+
+			![](BuildingMacEditor.png)
+
+			[REGION:note]
+			Depending on your system's specifications, it may take anywhere between 15 and 40 minutes to finish compiling the engine from source.
+			[/REGION]
+		
+		1.	After compiling the engine, select **Product > Run** to launch the editor (pictured below).
 		
-		1. Run **GenerateProjectFiles.command** located in the root directory. It should take less than a minute to complete. 
+			![](RunningMacEditor.png)
+
+		1.	Congratulations! You've compiled and launched the engine from source (result pictured below).
+
+			![](CompiledEditor_Mac.png)
+	[/PARAM]
+[/OBJECT]
+[OBJECT:ToggleButtonContent]
+	[PARAMLITERAL:category]
+	OS
+	[/PARAMLITERAL]
+	[PARAMLITERAL:id]
+	linux
+	[/PARAMLITERAL]
+	[PARAMLITERAL:active]
+	[/PARAMLITERAL]
+	[PARAM:content]
+		[REGION:warning]
+		At this time, we only provide support for Ubuntu 15.04.
+		[/REGION]
 		
-		1. Close the Terminal window once the operation has completed.
+		1.	Inside the root directory, **run GenerateProjectFiles.sh** from the terminal to set-up your project files (example pictured below). 
+
+			![](RunGenPrjFilesShellScript_Linux.png)
+
+		1.	To build the project, **run make** from the terminal (example pictured below).
+
+			![](RunMakeFile_Linux.png)
+
+		1.	After compiling the engine, navigate to the editor's binary path by entering **cd `Engine/Binaries/Linux/`** into the terminal.
 
-		1. Open the **UE4.xcodeproj** Xcode project located in the root directory. This is the main project for the entire engine. All of the projects for the engine core, tools, and games are contained within this project.
+		1.	Run **UE4Editor** to launch the editor (pictured below).
 
-		1. in Xcode, choose **UE4Editor - Mac** and **My Mac 64-bit** for the active scheme.
+			![](RunUE4Editor_Linux.png)
 
+		1.	Congratulations! You've compiled and launched the engine from source (result pictured below).
+
+			![](CompiledEditor_Linux.png)
 	[/PARAM]
 [/OBJECT]
 
+[/EXCERPT:Main]
+
+<!--Currently deprecated for the 'Building from Source' version
 
 ## Compiling Unreal Engine
 
-Unreal Engine 4 (UE4) uses a custom building method via the UnrealBuildTool. This tool handles all the complex aspects of compiling the engine and all of its parts. 
+Unreal Engine 4 (UE4) uses a custom building method via the UnrealBuildTool. This tool handles all the complex aspects of compiling the engine. 
 
 ### Build Configuration
 
@@ -235,5 +310,4 @@ The exception to the rule is that for PC platforms, the platform and configurati
 		1. From the **product** menu, choose **Build For > Running**.
 	[/PARAM]
 [/OBJECT]
-
-[/EXCERPT:Main]
+-->
\ No newline at end of file
