#UE4doc #docs:
Updated engine version
Changed font settings for multiple terms

diff --git a/Engine/Documentation/Source/Programming/UnrealArchitecture/TSubclassOf/TSubclassOf.INT.udn b/Engine/Documentation/Source/Programming/UnrealArchitecture/TSubclassOf/TSubclassOf.INT.udn
--- a/Engine/Documentation/Source/Programming/UnrealArchitecture/TSubclassOf/TSubclassOf.INT.udn
+++ b/Engine/Documentation/Source/Programming/UnrealArchitecture/TSubclassOf/TSubclassOf.INT.udn
@@ -2,9 +2,9 @@ Title: TSubclassOf
 Crumbs:
 Description:
 Availability:Public
-Version:4.7
+Version:4.9
 
-TSubclassOf is a template class that provides UClass type safety. For instance, let's imagine that you are creating a projectile class that allows the designer to specify the damage type. You could just create a UPROPERTY of type UClass and hope the designer always assigns a class derived from UDamageType or you could use the TSubclassOf template to enforce the choice. The sample code below illustrates the difference:
+**TSubclassOf** is a template class that provides UClass type safety. For instance, let's imagine that you are creating a projectile class that allows the designer to specify the damage type. You could just create a UPROPERTY of type UClass and hope the designer always assigns a class derived from UDamageType or you could use the TSubclassOf template to enforce the choice. The sample code below illustrates the difference:
 
 	/** type of damage */
 	UPROPERTY(EditDefaultsOnly, Category=Damage)
