#UE4doc #docs:
Added engine version field
Changed font settings on a few terms
Expanded on UE4 abbreviation at the first iteration.

diff --git a/Engine/Documentation/Source/Programming/UnrealBuildSystem/UnrealBuildSystem.INT.udn b/Engine/Documentation/Source/Programming/UnrealBuildSystem/UnrealBuildSystem.INT.udn
--- a/Engine/Documentation/Source/Programming/UnrealBuildSystem/UnrealBuildSystem.INT.udn
+++ b/Engine/Documentation/Source/Programming/UnrealBuildSystem/UnrealBuildSystem.INT.udn
@@ -2,7 +2,7 @@ Availability:Public
 Title: Unreal Build System
 Crumbs: %ROOT%, Programming
 Description:Unreal Engine's custom build system for compiling the engine and game projects.
-
+version:4.9
 [VAR:Topic]
 [OBJECT:Topic]
 	[PARAM:image]
@@ -28,7 +28,7 @@ Description:Unreal Engine's custom build system for compiling the engine and gam
 [/REGION]
 
 ### UnrealBuildTool
-UnrealBuildTool (or UBT) is a custom tool that manages the process of building the UE4 source code across the numerous possible build configurations.  Check out BuildConfiguration.cs to explore various 
+**UnrealBuildTool** (or UBT) is a custom tool that manages the process of building the **Unreal Engine 4** (UE4) source code across the numerous possible build configurations.  Check out BuildConfiguration.cs to explore various 
 user-configurable build options.
 
 [REGION:note]
@@ -36,7 +36,7 @@ If you would like to debug UnrealBuildTool or UnrealAutomationTool on Mac OS X,
 [/REGION]
 
 ### UnrealHeaderTool
-UnrealHeaderTool (or UHT) is a custom parsing and code-generation tool that supports the UObject system.  Code compilation happens in two phases:
+**UnrealHeaderTool** (or UHT) is a custom parsing and code-generation tool that supports the UObject system.  Code compilation happens in two phases:
 
 1.	UHT is invoked, which parses the C++ headers for Unreal-related class metadata and generates custom code to implement the various UObject-related features.
 1.	The normal C++ compiler is invoked to compile the results.
