#UE4doc #docs:
Updated engine version
Changed font settings for a few terms

diff --git a/Engine/Documentation/Source/Programming/UnrealBuildSystem/Configuration/BuildConfiguration.INT.udn b/Engine/Documentation/Source/Programming/UnrealBuildSystem/Configuration/BuildConfiguration.INT.udn
--- a/Engine/Documentation/Source/Programming/UnrealBuildSystem/Configuration/BuildConfiguration.INT.udn
+++ b/Engine/Documentation/Source/Programming/UnrealBuildSystem/Configuration/BuildConfiguration.INT.udn
@@ -2,30 +2,30 @@ Availability:Public
 Title:Configuring Unreal Build System
 Crumbs: %ROOT%, Programming, Programming/UnrealBuildSystem
 Description:Configuring Unreal Build System to control  how the engine is built.
-Version: 4.5
+Version: 4.9
 
-The Unreal Build System can be configured through the `BuildConfiguration.cs` to modify the build process to suit a variety of needs.
+The Unreal Build System can be configured through the BuildConfiguration.cs to modify the build process to suit a variety of needs.
 
 ## UnrealBuldTool Configuration Properties
 
-$ bStopXGECompilationAfterErrors : If _true_, compilation with XGE will stop when an error occurs.
+$ bStopXGECompilationAfterErrors : If true, compilation with XGE will stop when an error occurs.
     * Continuing compiling makes subsequent compiles much faster after you fix the odd error.
     * It lets you see compile errors in all files before quitting the build.
     * You can use **Ctrl + Break** to interrupt a failed build.
-$ bOmitPCDebugInfoInDevelopment : If _true_, debug info is not included when building in Development configuration.
-    * Set to _false_ to always include debug info.
+$ bOmitPCDebugInfoInDevelopment : If true, debug info is not included when building in Development configuration.
+    * Set to false to always include debug info.
     * Including debug info slows down compiling, but gives you line numbers and stack info in Development builds.
-$ bUseUnityBuild : Set to _true_ when rebuilding the engine as a whole. This should be disabled when using fast iteration mode. See [Build Modes](#BuildModes) for more information.
-$ bEnableFastIteration : If _true_, compile and link times are dramatically decreased. This should be disabled when creating unity builds. See [Build Modes](#BuildModes) for more information.
+$ bUseUnityBuild : Set to true when rebuilding the engine as a whole. This should be disabled when using fast iteration mode. See [Build Modes](#BuildModes) for more information.
+$ bEnableFastIteration : If true, compile and link times are dramatically decreased. This should be disabled when creating unity builds. See [Build Modes](#BuildModes) for more information.
     * Setting to _true_ gives fastest iteration times.
-    * If you encounter weird "RPC" or "PDB" link errors, then you should set `bUseIncrementalLinking=false` in `ValidateConfiguration()`
+    * If you encounter weird "RPC" or "PDB" link errors, then you should set `bUseIncrementalLinking = false` in **ValidateConfiguration()**
 
 
 ## Build Modes
 
-There are two build modes: **unity** and **fast iteration**. Unity is good for making changes that require rebuilding the whole engine e.g. changing an Engine header. Fast iteration is good when you are making localized changes - it gets your compile/link time down to less than 4 seconds. To change between the two, open `BuildConfiguration.cs` and find the lines:
+There are two build modes: **Unity** and **Fast Iteration**. Unity is good for making changes that require rebuilding the whole engine e.g. changing an Engine header. Fast iteration is good when you are making localized changes - it gets your compile/link time down to less than 4 seconds. To change between the two, open BuildConfiguration.cs and find the lines:
 
 	bEnableFastIteration = Utils.GetEnvironmentVariable( "ue.bEnableFastIteration", false );
 	bUseUnityBuild = Utils.GetEnvironmentVariable("ue.bUseUnityBuild", true);
 
-This can be modified to force the one you want to be `true`. Unity is the default.
+This can be modified to force the one you want to be true. Unity is the default.
