#UE4doc #docs:
Updated engine version
Changed font settings for multiple terms
Changed title to reflect that this is a page for configuring for Windows XP in particular.

diff --git a/Engine/Documentation/Source/Programming/UnrealBuildSystem/Configuration/Windows/BuildConfigurationWindows.INT.udn b/Engine/Documentation/Source/Programming/UnrealBuildSystem/Configuration/Windows/BuildConfigurationWindows.INT.udn
--- a/Engine/Documentation/Source/Programming/UnrealBuildSystem/Configuration/Windows/BuildConfigurationWindows.INT.udn
+++ b/Engine/Documentation/Source/Programming/UnrealBuildSystem/Configuration/Windows/BuildConfigurationWindows.INT.udn
@@ -1,25 +1,25 @@
 Availability:Public
-Title:Configuring Unreal Build System for Windows
+Title:Configuring Unreal Build System for Windows XP
 Crumbs: %ROOT%, Programming, Programming/UnrealBuildSystem/Configuration
 Description:Configuring Unreal Build System to control how the engine is built for Windows.
-Version: 4.5
+Version: 4.9
 
-The Unreal Build System can be configured through the `BuildConfiguration.cs` to modify the build process to suit a variety of needs.
-For Windows specific options, you can alter the `UEBuildWindows.cs` configuration
+The Unreal Build System can be configured through the BuildConfiguration.cs to modify the build process to suit a variety of needs.
+For Windows specific options, you can alter the UEBuildWindows.cs configuration
 
 ## Adding Windows XP Support
 
 ### Minimum Spec
 
-**Only Windows XP with Service Pack 3 is supported**.
+Only Windows XP with Service Pack 3 is supported.
 
-**You also must have a Video Card which supports OpenGL 3.2**.
+You also must have a Video Card which supports OpenGL 3.2.
 
 ### Building with Windows XP support
-$ WindowsPlatform::SupportWindowsXP : If _true_, a build of UE4 which is compatible with Windows XP will be created.
+$ WindowsPlatform::SupportWindowsXP : If true, a build of UE4 which is compatible with Windows XP will be created.
 
 
-To enable XP support, find where it is forced off in the `SetupEnvironment` and update this, with:
+To enable XP support, find where it is forced off in the SetupEnvironment and update this, with:
 	public override void SetUpEnvironment(UEBuildTarget InBuildTarget)
 	{
 		...
@@ -27,7 +27,7 @@ To enable XP support, find where it is forced off in the `SetupEnvironment` and
 		// Win32 XP is only supported at this time.
 		SupportWindowsXP = InBuildTarget.Platform == UnrealTargetPlatform.Win32;
 
-Enabling a Windows XP compatible build forces the compiler to use `Windows SDK 7.1`. This will mean that any functionality which has been provided by the `Windows SDK 8.0` will not compile. You can wrap sdk specific code in a WINVER check:
+Enabling a Windows XP compatible build forces the compiler to use Windows SDK 7.1. This will mean that any functionality which has been provided by the Windows SDK 8.0 will not compile. You can wrap sdk specific code in a WINVER check:
 	#if WINVER == 0x0502
 		// Windows SDK 7.1 code should live inside this block
 	#else
@@ -39,7 +39,7 @@ Enabling a Windows XP compatible build forces the compiler to use `Windows SDK 7
 As Windows XP does not support DirectX 10 or DirectX 11, OpenGL must be included as a target RHI for your game.
 
 You can find the option to include this in the Editor:
-_Edit->Project Settings->Platform->Windows_
+Edit->Project Settings->Platform->Windows
 [![](EnableOpenGLForWindows.png)](Programming/UnrealBuildSystem/Configuration/Windows)
 
 Once you have packaged your game, you should then be able to run the game executable from the binaries directory on a Windows XP machine.
