#UE4doc #docs:
Updated engine version
Removed overview header
Changed font settings for multiple terms

diff --git a/Engine/Documentation/Source/Programming/UnrealBuildSystem/TargetFiles/TargetFiles.INT.udn b/Engine/Documentation/Source/Programming/UnrealBuildSystem/TargetFiles/TargetFiles.INT.udn
--- a/Engine/Documentation/Source/Programming/UnrealBuildSystem/TargetFiles/TargetFiles.INT.udn
+++ b/Engine/Documentation/Source/Programming/UnrealBuildSystem/TargetFiles/TargetFiles.INT.udn
@@ -2,37 +2,37 @@ Availability:Public
 Title: Unreal Build System Target Files
 Crumbs: %ROOT%, Programming, Programming\UnrealBuildSystem
 Description:Reference for the target files used by Unreal Build System when building modules.
-Version: 4.5
+Version: 4.9
 
 [TOC (start:2 end:3)]
 
-## Overview
 
-The Unreal Build System makes use of **Target** files for each application UnrealBuildTool builds that get compiled into the dynamic Rules module. These files replace the old `UEBuild*.cs` files that were originally compiled directly into UnrealBuildTool. Each target file defines a class used by UnrealBuildtool to determine how a game or program is built.
+
+The Unreal Build System makes use of **Target** files for each application UnrealBuildTool builds that get compiled into the dynamic Rules module. These files replace the old UEBuild*.cs files that were originally compiled directly into UnrealBuildTool. Each target file defines a class used by UnrealBuildtool to determine how a game or program is built.
 
 UnrealBuildTool now retrieves the proper target for the project being built from the RulesModule.dll.
 
 There are 2 types of target:
 
-* **Game** - This is a target that utilizes the 'shared' UE4 executable.
+* **Game** - This is a target that utilizes the shared UE4 executable.
 * **Program** - This is a target that is a stand-alone executable. ShaderCompileWorker, SlateViewer, and UnrealHeaderTool are this type of target.
 
 
 ## Target File Naming and Location
 
-The Target file must be located in the root directory of a program or `Source` directory of a game and be named using the convention: `[APPNAME].Target.cs`
+The Target file must be located in the root directory of a program or Source directory of a game and be named using the convention: [APPNAME].Target.cs
 
-As an example, the ShaderCompileWorker program directory at `\UE4\Engine\Source\Programs\ShaderCompileWorker` contains the `ShaderCompileWorker.Target.cs` file.
+As an example, the ShaderCompileWorker program directory at \UE4\Engine\Source\Programs\ShaderCompileWorker contains the ShaderCompileWorker.Target.cs file.
 
-Similarly, the ExampleGame `Source` directory located at `UE4\ExampleGame|Source` contains the `ExampleGame.Target.cs` file.
+Similarly, the ExampleGame Source directory located at UE4\ExampleGame|Source contains the ExampleGame.Target.cs file.
 
 ## Target File Contents
 
-The class defined in the Target file for each game or program is derived from `TargetRules` (found in `RulesCompiler.cs`). 
+The class defined in the Target file for each game or program is derived from TargetRules (found in RulesCompiler.cs). 
 
-The name of the class is constructed by adding the name of the tool or game followed by `target`. For instance, the class defined in the target file for the ShaderCompileWorker tool would be `ShaderCompileWorkerTarget`. The `Target` at the end is **REQUIRED**. It was added to avoid name collisions with the existing `*.Build.cs` module rules classes, since they all get compiled into the same DLL.
+The name of the class is constructed by adding the name of the tool or game followed by target. For instance, the class defined in the target file for the ShaderCompileWorker tool would be ShaderCompileWorkerTarget. The Target at the end is required. It was added to avoid name collisions with the existing *.Build.cs module rules classes, since they all get compiled into the same DLL.
 
-UnrealBuildTool takes `[APPNAME]` and appends `Target` when requesting the build target from the Rules DLL. So, using the same example as above, when UnrealBuildTool is passed `ShaderCompileWorker Win64 Development` it will request `ShaderCompileWorkerTarget` as the `TargetRules` class name of interest.
+UnrealBuildTool takes [APPNAME] and appends Target when requesting the build target from the Rules DLL. So, using the same example as above, when UnrealBuildTool is passed ShaderCompileWorker Win64 Development it will request ShaderCompileWorkerTarget as the TargetRules class name of interest.
 
 ### Constructor
 
@@ -46,14 +46,14 @@ The constructor sets the name of the application and the type of target to gener
 
 ### Setup Binaries
 
-`TargetRules::SetupBinaries()` is used to generate a list of `UEBuildBinaryConfiguration` instances which UnrealBuildTool will use to fill in the AppBinaries array (these are what it will generate).
+**TargetRules::SetupBinaries()** is used to generate a list of UEBuildBinaryConfiguration instances which UnrealBuildTool will use to fill in the AppBinaries array (these are what it will generate).
 
 ### Setup Global Environment
 
-`TargetRules::SetupGlobalEnvironment()` is used to set various build configuration settings for UnrealBuildTool for the target as well as setup the global Compile and Link environments.
+**TargetRules::SetupGlobalEnvironment()** is used to set various build configuration settings for UnrealBuildTool for the target as well as setup the global Compile and Link environments.
 
 [REGION:note]
-**Note:** The current games do not contain this function, as the UEBuildGame class takes care of their needs by default. However, if a game needed to override these settings, it certainly can do so.
+The current games do not contain this function, as the UEBuildGame class takes care of their needs by default. However, if a game needed to override these settings, it certainly can do so.
 [/REGION]
 
 ## Example Target File
