#UE4doc #docs:
Removed Overview header

diff --git a/Engine/Documentation/Source/Resources/Community/SciFiBunk_MaterialInstancing/MasterMaterialsandMaterialInstancing.INT.udn b/Engine/Documentation/Source/Resources/Community/SciFiBunk_MaterialInstancing/MasterMaterialsandMaterialInstancing.INT.udn
--- a/Engine/Documentation/Source/Resources/Community/SciFiBunk_MaterialInstancing/MasterMaterialsandMaterialInstancing.INT.udn
+++ b/Engine/Documentation/Source/Resources/Community/SciFiBunk_MaterialInstancing/MasterMaterialsandMaterialInstancing.INT.udn
@@ -49,7 +49,6 @@ This tutorial was provided by community member Liam Tart (PogoP), and is publish
 
 Thank you, Liam!
 
-## Overview
 
 This tutorial will explain how master materials and material instancing can be used to quickly and efficiently build a set of tileable materials that can be applied to meshes in your scene. By texturing your meshes using tileable materials, you do not have to rely on baking from a high poly model and you can rapidly change the look and feel of your environment.
 
