#UE4doc #docs:
Removed Overview header
Changed font settings for multiple terms

diff --git a/Engine/Documentation/Source/Resources/ContentExamples/Animation/Animation.INT.udn b/Engine/Documentation/Source/Resources/ContentExamples/Animation/Animation.INT.udn
--- a/Engine/Documentation/Source/Resources/ContentExamples/Animation/Animation.INT.udn
+++ b/Engine/Documentation/Source/Resources/ContentExamples/Animation/Animation.INT.udn
@@ -49,9 +49,8 @@ Related: Engine/Content/Tools/MayaRiggingTool
 [TOC(start:2 end:2)]
 
 
-## Overview
 The **Animation** map provides several examples on how animation can be applied to an Actor through one-off or looping events, how to call an animation through a **Blueprint**, or define behavioral type sequences through **Animation Blueprints**. 
-This map also demonstrates how a **Blendspace** can be used to blend between movement directions and how a **Aimoffset** can be used to have a Actor rotate smoothly between Pitch and Yaw values. Other aspects covered include setting up IK for both **Hand IK** and **Foot IK**, the impact **Root Transformation** has on character movement as well as a fully playable 3rd person character!
+This map also demonstrates how a **Blendspace** can be used to blend between movement directions and how a **AimOffset** can be used to have a Actor rotate smoothly between Pitch and Yaw values. Other aspects covered include setting up IK for both **Hand IK** and **Foot IK**, the impact **Root Transformation** has on character movement as well as a fully playable 3rd person character!
 
 The **PhysicalAnimation** map has two sections in it which show how the **Physics Blend Weight** node can be used to apply physics to an animation at a specific location of a **Skeletal Mesh Actor**. This is useful if you are looking to blend between raw animation and physics or if you are looking to generate physics based hit reactions for Actors.  
 
@@ -64,15 +63,15 @@ Below are the examples provided inside the **Animation** map:
 | Example | What is Demonstrated |
 | ------- | ------------ |
 | **[1.1 Animation - Using Animation Asset](Engine\Animation\Sequences)** | A **SkeletalMeshActor** that is set to play an **AnimationSequence** asset assigned through the **Details** panel of the Actor (e.g. Actor loops and performs jumping jacks). |
-| **[1.2 Animation Blueprint](Engine\Animation\AnimBlueprints)** | Two examples are shown: one where a **Blueprint** is used to tell an Actor to play an animation, the other shows an **Animation Blueprint** which tells the Actor to play an animation (e.g. this example shows how to pass a variable from a **Blueprint** to an **Animation Blueprint** and demonstrates blending from an existing pose to a new animation).  |
+| **[1.2 Animation Blueprint](Engine\Animation\AnimBlueprints)** | Two examples are shown: one where a Blueprint is used to tell an Actor to play an animation, the other shows an Animation Blueprint which tells the Actor to play an animation (e.g. this example shows how to pass a variable from a Blueprint to an Animation Blueprint and demonstrates blending from an existing pose to a new animation).  |
 | **[1.3 Locomotion BlendSpace](Engine\Animation\Blendspaces)** | How a **Locomotion BlendSpace** can be used to blend an Actor's movement in different directions and at different speeds (e.g. an Actor is shown blending between walking/running backwards, forwards, left, and right at different speeds). |
-| **[1.4 AimOffset](Engine\Animation\AimOffset)** | An **AimOffset** is showcased which will allow you to control and blend between different poses for aiming along Yaw or Pitch values (e.g. an Actor is shown blending between various aiming positions). |
-| **[1.5 Animation Montage](Engine\Animation\AnimMontage)** | Playing an **Animation Montage** from a **Blueprint** as well as blending into/out of the animation is shown (e.g. an Actor blends between an aim pose and throwing punches). |
+| **[1.4 AimOffset](Engine\Animation\AimOffset)** | An AimOffset is showcased which will allow you to control and blend between different poses for aiming along Yaw or Pitch values (e.g. an Actor is shown blending between various aiming positions). |
+| **[1.5 Animation Montage](Engine\Animation\AnimMontage)** | Playing an **Animation Montage** from a Blueprint as well as blending into/out of the animation is shown (e.g. an Actor blends between an aim pose and throwing punches). |
 | **[1.6 Curve Driven Animation (Morphtarget)](Engine\Animation\Persona\MorphTargetPreviewer)** | An animation that has a **Curve Float Track** added which blends between two morph targets (e.g. the result is an Actor who's nose increases in size then returns to normal).  |
-| **[1.7 Curve Driven Animation (Bone Scale)](Engine\Animation\Persona\BoneManipulation)** |How to modify a bone of a **SkeletalMeshActor** during an animation inside the **AnimBlueprint** with a **Curve Float Track** (e.g. an Actor's head blends between a small/large head). |
-| **[1.8 Inverse Kinetics (IK)](Engine\Animation\IKSetups)** | Setting up **IK** inside the **Blueprint** of an Actor for both foot IK and hand IK (e.g. an Actor who adjusts his feet along stairs and a Character who's hands deform when hitting a wall). |
-| **[1.9 Root Motion](Engine\Animation\RootMotion)** | Enabling **Root Motion Transformation** and how it impacts collision when animating a Character (e.g. a Character without **Root Motion Transformation** loses his collision and walks through a box, while a Character with **Root Motion Transformation** enabled, keeps collision and moves around the box). |
-| **[1.10 Playable Character Animation Blueprint](Engine\Animation\CharacterSetupOverview)** | A third person controllable Character employing an **Locomotion BlendSpace** and an **AimOffset** (e.g. a playable version of the Owen Character is provided here). |
+| **[1.7 Curve Driven Animation (Bone Scale)](Engine\Animation\Persona\BoneManipulation)** |How to modify a bone of a SkeletalMeshActor during an animation inside the AnimBlueprint with a Curve Float Track (e.g. an Actor's head blends between a small/large head). |
+| **[1.8 Inverse Kinetics (IK)](Engine\Animation\IKSetups)** | Setting up **Inverse Kinetics** (IK) inside the Blueprint of an Actor for both foot IK and hand IK (e.g. an Actor who adjusts his feet along stairs and a Character who's hands deform when hitting a wall). |
+| **[1.9 Root Motion](Engine\Animation\RootMotion)** | Enabling **Root Motion Transformation** and how it impacts collision when animating a Character (e.g. a Character without Root Motion Transformation loses his collision and walks through a box, while a Character with Root Motion Transformation enabled, keeps collision and moves around the box). |
+| **[1.10 Playable Character Animation Blueprint](Engine\Animation\CharacterSetupOverview)** | A third person controllable Character employing an Locomotion BlendSpace and an AimOffset (e.g. a playable version of the Owen Character is provided here). |
 
 
 ## PhysicalAnimation Map
@@ -84,4 +83,4 @@ Below are the examples provided inside the **PhysicalAnimation** map:
 | Example | What is Demonstrated |
 | ------- | ------------ |
 | **[1.1 Physics Animation Blending](Engine\Animation\PhysicallyDrivenAnimation)** | The use of **Physics Blend Weight** is applied to three Actors, each is set to blend below different bones of the Skeletal Mesh (e.g. having physics applied to an animation below a specified bone in the Actor).  |
-| **[1.2 Physics Animation Blending (Hit Reaction)](Engine\Animation\PhysicallyDrivenAnimation)** | Using a **Physics Blend Weight** with an Impulse node to generate a Hit Reaction at the point of impact (e.g. the Actor responds to location specific damage and plays an impulse reaction while running).  |
+| **[1.2 Physics Animation Blending (Hit Reaction)](Engine\Animation\PhysicallyDrivenAnimation)** | Using a Physics Blend Weight with an Impulse node to generate a Hit Reaction at the point of impact (e.g. the Actor responds to location specific damage and plays an impulse reaction while running).  |
