#UE4doc #docs:
Updated engine version
Removed mentions of the Depot as that pertains to perforce
Replaced multiple mentions of Visual Studio with VS

diff --git a/Engine/Documentation/Source/Programming/UnrealBuildSystem/ProjectFileGenerator/ProjectFileGenerator.INT.udn b/Engine/Documentation/Source/Programming/UnrealBuildSystem/ProjectFileGenerator/ProjectFileGenerator.INT.udn
--- a/Engine/Documentation/Source/Programming/UnrealBuildSystem/ProjectFileGenerator/ProjectFileGenerator.INT.udn
+++ b/Engine/Documentation/Source/Programming/UnrealBuildSystem/ProjectFileGenerator/ProjectFileGenerator.INT.udn
@@ -2,11 +2,11 @@ Availability:Public
 Title:Automatic Project File Generation
 Crumbs:%ROOT%, Programming, Programming/UnrealBuildSystem
 Description:Process for automatically generating project files for games and modules in the current workspace.
-Version: 4.5
+Version: 4.9
 
-If you download the latest Unreal Engine 4 code, you might notice that there are no Visual Studio or Xcode project files included for compiling and running the engine or example games.  You will need to run a script that generates project files for you to load:
+If you download the latest Unreal Engine 4 (UE4) code, you might notice that there are no Visual Studio (VS) or Xcode project files included for compiling and running the engine or example games.  You will need to run a script that generates project files for you to load:
 
-* Run the `GenerateProjectFiles.bat` file which is located in the main UE4 directory (mapped to '//depot/UE4/`).
+* Run the GenerateProjectFiles.bat file which is located in the main UE4 directory<!--(mapped to '//depot/UE4/`)-->.
 
 	[REGION:note]
 	You can run that script from any folder.  (Make a desktop shortcut if you want.)
@@ -14,22 +14,22 @@ If you download the latest Unreal Engine 4 code, you might notice that there are
 
 * The project generator tool will analyze the module and target build files and generate new project files.  This can take up to 15 seconds or so.
 
-That is it!  For Visual Studio developers, this will generate a UE4.sln file in the same directory.  You will use this solution file to build and run Unreal Engine 4 games and programs.  (For your first time compiling, we recommend building using the Development configuration with Win64 platform!)
+For VS developers, this will generate a UE4.sln file in the same directory.  You will use this solution file to build and run UE4 games and programs.  (For your first time compiling, we recommend building using the Development configuration with Win64 platform)
 
 [REGION:note]
-	By default, we generate buildable projects for all available platforms that we detected SDKs for, so that you will be able to build and debug console and mobile platforms as well.  If you want to generate projects only for the current platform you are running on (e.g. Windows.), run the `GenerateProjectFiles.bat` with the '-CurrentPlatform' argument.  The project files will be a little lighter.
+	By default, we generate buildable projects for all available platforms that we detected SDKs for, so that you will be able to build and debug console and mobile platforms as well.  If you want to generate projects only for the current platform you are running on (e.g. Windows.), run the GenerateProjectFiles.bat with the -CurrentPlatform argument.  The project files will be a little lighter.
 [/REGION]
 
-**IMPORTANT**:  Remember to re-run `GenerateProjectFiles.bat` after syncing new code changes from source control!  If you forget to do this, you will most likely run into problems when compiling or running your game.
+**IMPORTANT**:  Remember to re-run GenerateProjectFiles.bat after syncing new code changes from source control.  If you forget to do this, you will most likely run into problems when compiling or running your game.
 
-If you make local changes to modules such as adding or removing source files, you should re-run `GenerateProjectFiles.bat` again.  We do not recommend manually making changes to project files.
+If you make local changes to modules such as adding or removing source files, you should re-run GenerateProjectFiles.bat again.  We do not recommend manually making changes to project files.
 
 
 ## Integration With UnrealVS Extension
 
-The [](Programming/Development/VisualStudioSetup/UnrealVS) for Visual Studio includes a tool-bar button that you can click to re-generate project files for the currently loaded solution.
+The [](Programming/Development/VisualStudioSetup/UnrealVS) for VS includes a tool-bar button that you can click to re-generate project files for the currently loaded solution.
 
-You can also bind a shortcut key to this feature.  In Visual Studio, open `Tools -> Options -> Environment -> Keyboard`, then search for `UnrealVS.RefreshProjects`.
+You can also bind a shortcut key to this feature.  In VS, open Tools -> Options -> Environment -> Keyboard, then search for UnrealVS.RefreshProjects.
 
 Note that this feature is only enabled after you have a solution loaded (because the tool needs to know which code branch to generate projects for.)  If you have no `UE4.sln` file generated yet, you will need to run the `GenerateProjectFiles.bat` script directly first.
 
@@ -41,9 +41,9 @@ The project generator has some optional command-line arguments that you can use
 | Option           |  Description  ||
  ----------------- |  ------------ |
 -CurrentPlatform   |  Generates buildable projects for only the current desktop platform (Windows or Mac), instead of all detected available target platforms. |
--2012              |  Generates projects in a format native to Visual Studio 2012.  By default, Visual Studio 2013 project files are generated. |
--2013              |  Generates projects in a format native to Visual Studio 2013.  This is the default. |
--2015              |  Generates projects in a format native to Visual Studio 2015.  The Visual C++ 2015 compiler will be used to compile when this option is set. |
+-2012              |  Generates projects in a format native to VS 2012.  By default, VS 2013 project files are generated. |
+-2013              |  Generates projects in a format native to VS 2013.  This is the default. |
+-2015              |  Generates projects in a format native to VS 2015.  The Visual C++ 2015 compiler will be used to compile when this option is set. |
 -ThirdParty        |  Adds headers and other files from third-party libraries to the project.  This could be useful if you like to see symbols and files in Visual Assist for PhysX, Steamworks or Bink for example.  However, it will make the project files take longer to load! |
 -Game `GameName`   |	 Tells the project generator to generate projects that only include code and content for the specified project name, excluding all other discovered projects.  Make sure to specify a project name as well.  For example `GenerateProjectFiles.bat ShooterGame.uproject -Game` would generate projects that only have source and targets for ShooterGame.  If you want to include Engine source code and Programs alongside a single game project, you can also pass the `-Engine` parameter. |
 -Engine            |  When used with `-Game`, causes Engine code, content, and programs to also be included in the generated solution.  This is useful if you think you will need to work with engine source a lot while working on your game project. |
@@ -58,33 +58,33 @@ The project generator has some optional command-line arguments that you can use
 
 ### Why are newly-added source files compiled even though I did not update project files?
 
-The Unreal Engine 4 build system does not actually require project files to compile the code.  Unreal Build Tool always locates source files using module and target build files.  For this reason, if you add a new source file and trigger a compile, the new source file may be included in the build even though the project files have not been refreshed yet.  Just something to keep in mind.
+The UE4 build system does not actually require project files to compile the code.  Unreal Build Tool always locates source files using module and target build files.  For this reason, if you add a new source file and trigger a compile, the new source file may be included in the build even though the project files have not been refreshed yet.  Just something to keep in mind.
 
 
 ### Where are the actual projects saved to?
 
-For Visual Studio projects, the generated solution file is `UE4.sln` and is saved to the root UE4 directory.  However, the project files are saved to the /UE4/Engine/Intermediate/ProjectFiles/ directory.  It is safe to delete these files at any time and regenerate projects, however you may lose certain project-specific preferences such as command-line argument strings if you delete these files.
+For VS projects, the generated solution file is UE4.sln and is saved to the root UE4 directory.  However, the project files are saved to the /UE4/Engine/Intermediate/ProjectFiles/ directory.  It is safe to delete these files at any time and regenerate projects, however you may lose certain project-specific preferences such as command-line argument strings if you delete these files.
 
 
 ### What are the advantages of generating project files?
 
 There are certainly both pros and cons, but here are some important reasons why we decided to generate project files for UE4:
 
-* Unreal Engine 4 is designed to work with many platforms, but different teams may only be working with a few specific platforms at any time.  By generating project files, we can omit platform-specific files and build configurations that are not relevant.
+* UE4 is designed to work with many platforms, but different teams may only be working with a few specific platforms at any time.  By generating project files, we can omit platform-specific files and build configurations that are not relevant.
 
 * UE4 programming methodology incorporates many sub-modules, and it was very important that we made it as easy as possible for programmers to add new modules.
 
-* The project generator emits highly-accurate definitions and includes paths which are used by Visual Studio IntelliSense while you are working in UE4 code.
+* The project generator emits highly-accurate definitions and includes paths which are used by VS IntelliSense while you are working in UE4 code.
 
 * Setting up a new project is much easier when project files are automatically generated.
 
-* We want to support multiple platforms and development environments (for example, Visual Studio and Xcode.)  Hand-maintaining multiple sets of project files is error-prone and tedious.
+* We want to support multiple platforms and development environments (for example, VS and Xcode.)  Hand-maintaining multiple sets of project files is error-prone and tedious.
 
 * We wanted to enable programmers to generate highly customized project files.  We expect, over time, this to become even more important.
 
 * The directory structure of source files is automatically mirrored in the project file solution hierarchy.  This is really convenient when browsing source files, but was difficult to maintain with hand-authored projects.
 
-* Unreal Engine build configurations are very complex and are difficult to hand-maintain.  The project generator makes this nearly transparent to developers.
+* UE4 build configurations are very complex and are difficult to hand-maintain.  The project generator makes this nearly transparent to developers.
 
 
 
@@ -92,11 +92,11 @@ There are certainly both pros and cons, but here are some important reasons why
 
 Along with source code for C++ modules, we automatically add several other files to the generated projects.  This is just to make it easy to locate these files when doing searches.  Some examples of other files that are added to the generated projects:
 
-* Shader source code (`*.usf` files)
-* Engine documentation files (`*.udn` files)
-* Program configuration files (`*.ini` files)
-* Localization files (`*.int` files)
-* Program resource files and manifests (`*.rc`, `*.manifest`)
+* Shader source code (*.usf files)
+* Engine documentation files (*.udn files)
+* Program configuration files (*.ini files)
+* Localization files (*.int files)
+* Program resource files and manifests (*.rc, .manifest)
 * Certain external (non-generated) project files (e.g. UnrealBuildTool, Clean)
 
 
@@ -113,16 +113,16 @@ This script is a simple wrapper around Unreal Build Tool, which is launched in a
 ### I am making changes to the project generator code.  How do I debug it?
 
 * Change your startup project to UnrealBuildTool.
-* Set the command-line parameters for debugging to:  `-ProjectFiles`
-* Set the working directory to your local path for:  `//depot/UE4/Engine/Source/`
+* Set the command-line parameters for debugging to:  -ProjectFiles
+* Set the working directory to your local path for:  <!--//depot/UE4-->[Your UE4 Directory]/Engine/Source/
 * Compile and debug as per normal.
 
 
-Be aware that the project generator may clobber the project files you are using in Visual Studio as you are working.  For this reason, it is sometimes useful to load the UnrealBuildTool project directly into Visual Studio instead of debugging through the regular UE4 solution file.
+Be aware that the project generator may clobber the project files you are using in VS as you are working.  For this reason, it is sometimes useful to load the UnrealBuildTool project directly into VS instead of debugging through the regular UE4 solution file.
 
 
 ### How do I build multiple configurations at the same time?
 
-You can use the `Batch Build` feature in Visual Studio.  You can find this under the Build menu.  Simply select all of the configurations that you want to compile and click Build.  We are looking at ways to make this interface much easier to work with in the future.
+You can use the Batch Build feature in VS.  You can find this under the Build menu.  Simply select all of the configurations that you want to compile and click Build.  We are looking at ways to make this interface much easier to work with in the future.
 
 
