Re-placing the code that was used with code that runs better. Also added the HLSL code to the page to make it easier to copy into your UE4 project.

diff --git a/Engine/Documentation/Source/Engine/Rendering/Materials/ExpressionReference/Custom/CustomExpressions.INT.udn b/Engine/Documentation/Source/Engine/Rendering/Materials/ExpressionReference/Custom/CustomExpressions.INT.udn
--- a/Engine/Documentation/Source/Engine/Rendering/Materials/ExpressionReference/Custom/CustomExpressions.INT.udn
+++ b/Engine/Documentation/Source/Engine/Rendering/Materials/ExpressionReference/Custom/CustomExpressions.INT.udn
@@ -5,10 +5,6 @@ Description:Expressions that allow the use of custom, plain shader code.
 Skilllevel: Advanced
 Version: 4.8
 
-[EXCERPT:TOC]
-[TOC(start:2)]
-[/EXCERPT:TOC]
-
 ## Custom
 
 The Custom expression allows you to write custom HLSL shader code operating on an arbitrary amount of inputs and outputting the result of the operation.
@@ -40,6 +36,29 @@ Add as many inputs as you need to the Inputs array, and name them. You can then
 
 ![](T_Custom_Node.png)
 
+Here is the code that was used above so you can try out the Custom node for yourself.
+
+	float3 blur = Texture2DSample(Tex, TexSampler, UV);
+
+	for (int i = 0; i < r; i++)
+	{
+
+	  blur += Texture2DSample(Tex, TexSampler, UV + float2(i * dist, 0));
+	  blur += Texture2DSample(Tex, TexSampler, UV - float2(i * dist, 0));
+
+	}
+
+	for (int j = 0; j < r; j++)
+	{ 
+
+	  blur += Texture2DSample(Tex, TexSampler, UV + float2(0, j * dist));
+	  blur += Texture2DSample(Tex, TexSampler, UV - float2(0, j * dist));
+
+	}
+
+	blur /= 2*(2*r)+1;
+	return blur;
+
 (#Warnings)
 [REGION:warning]
 **Using the custom node prevents constant folding and may use significantly more instructions than an equivalent version done with built in nodes!**  Constant folding is an optimization that UE4 employs under the hood to reduce shader instruction count when necessary.  For example, an expression chain of Time >Sin >Mul by parameter > Add to something' can and will be collapsed by UE4 into a single instruction, the final add.  This is possible because all of the inputs of that expression (Time, parameter) are constant for the whole draw call, they do not change per-pixel.  UE4 cannot collapse anything in a custom node, which can produce less efficient shaders than an equivalent version made out of existing nodes.  As a result, it is best to only use the custom node when it gives you access to functionality not possible with the existing nodes.
