updated DOF documentation

diff --git a/Engine/Documentation/Source/Engine/Rendering/PostProcessEffects/DepthOfField/DepthOfField.INT.udn b/Engine/Documentation/Source/Engine/Rendering/PostProcessEffects/DepthOfField/DepthOfField.INT.udn
--- a/Engine/Documentation/Source/Engine/Rendering/PostProcessEffects/DepthOfField/DepthOfField.INT.udn
+++ b/Engine/Documentation/Source/Engine/Rendering/PostProcessEffects/DepthOfField/DepthOfField.INT.udn
@@ -23,13 +23,13 @@ There are three methods available for performing depth of field effects in Unrea
 [/REGION]
 
 
-### Gaussian
+### Gaussian Depth of Field
 
 [EXCERPT:DepthOfField_Gaussian]
 Gaussian DoF blurs the scene using a standard Guassian blur. The Guassian method is fast and generally well-suited for use in-game, where performance is critical.
 [/EXCERPT:DepthOfField_Gaussian]
 
-### Bokeh
+### Bokeh Depth of Field
 
 [EXCERPT:DepthOfField_Bokeh]
 Bokeh is the name of the shape that can be seen in photos or movies when objects are out of focus. Bokeh DoF is a method where a textured quad is rendered for each pixel using a texture to define the shape to reproduce the effect produced by camera lenses, which gives a filmic look to the scene. 
@@ -38,12 +38,16 @@ The actual implementation does this performance intensive effect only in half re
 The Bokeh method is far more costly performance-wise, making cinematics and showcases prime candidates for the Bokeh method as the appealing visuals often outweigh any performance concerns in those situations.
 [/EXCERPT:DepthOfField_Bokeh]
 
-### Circle
+### Circle Depth of Field
 
 [EXCERPT:DepthOfField_Circle]
 The Circle DoF is the newest addition and matches real world cameras more closely. Simialar to the Bokeh DoF you can see circular shaped Bokeh with sharp HDR content. Very large Bokeh are a weakness of this
 algorithm, they appear noisy and not as clear as the BokehDOF can do them.
 The method has good performance (quite better than Bokeh DoF) and good quality, especially in the transition areas.
+
+The defaults of CircleDOF are intentionally weak to be more pleasant. To get the maximum effect you can tweak settings. Use a low Aperture for a large Bokeh, get close to an object and change the Field of View to be lower.
+Play with the focal distance to
+
 [/EXCERPT:DepthOfField_Circle]
 
 ## Implementation
@@ -63,6 +67,17 @@ These layers, including the transition regions, can be visualized using the **De
 | Scene with DoF | DoF Layers Visualized |
 [/REGION]
 
+We further improved this visualization to also include a lot of interesting information:
+
+[REGION:imagetable]
+| ![Scene](NewVisualize.png) |
+| :----: |
+| VisualizeDOF |
+[/REGION]
+
+When using the CircleDOF you also can see a crosshair attached to the mouse cursor showing the computed circle of confusion (BokehSize) together with the depth and the pixel size in numbers.
+Note that the circle size might not always match to the actual effect as the implementation is not well suited / intended for large Bokeh shapes. We might want to improve on that in the future.
+
 ## Properties
 
 | Property | Description |
