#UE4Doc #docs Updated downloading and building unreal engine from source.

diff --git a/Engine/Documentation/Source/GettingStarted/DownloadingUnrealEngine/DownloadZip/DownloadZip.INT.udn b/Engine/Documentation/Source/GettingStarted/DownloadingUnrealEngine/DownloadZip/DownloadZip.INT.udn
--- a/Engine/Documentation/Source/GettingStarted/DownloadingUnrealEngine/DownloadZip/DownloadZip.INT.udn
+++ b/Engine/Documentation/Source/GettingStarted/DownloadingUnrealEngine/DownloadZip/DownloadZip.INT.udn
@@ -1,5 +1,5 @@
 Availability:NoPublish
-Title:Downloading Unreal Engine as a Zip File
+Title:Downloading Unreal as a Zip from GitHub
 Crumbs:%ROOT%, GettingStarted, GettingStarted/DownloadingUnrealEngine
 Description:Downloading the Unreal Engine source code as a Zip file.
 
diff --git a/Engine/Documentation/Source/GettingStarted/DownloadingUnrealEngine/DownloadingUnrealEngine4.INT.udn b/Engine/Documentation/Source/GettingStarted/DownloadingUnrealEngine/DownloadingUnrealEngine4.INT.udn
--- a/Engine/Documentation/Source/GettingStarted/DownloadingUnrealEngine/DownloadingUnrealEngine4.INT.udn
+++ b/Engine/Documentation/Source/GettingStarted/DownloadingUnrealEngine/DownloadingUnrealEngine4.INT.udn
@@ -5,14 +5,19 @@ Description:Steps for connecting to the source code repository and downloading t
 
 [TOC (start:2 end:3)]
 
-
 [EXCERPT:Main]
 [PUBLISH:Licensee]
 Unreal Engine 4 is available to be downloaded by licensees directly from Epic's Perforce server. Before you can connect to the server, a secure VPN connection must be established. The login credentials for both the VPN and Perforce server are provided to your team's technical lead by Epic Games in your UDN welcome email. Only one set of login credentials is provided to each team.
 [/PUBLISH:Licensee]
 [PUBLISH:Rocket]
-Source code for Unreal Engine 4 (UE4) can be downloaded by our subscribers from the Unreal Engine GitHub repository. With access to our source code, subscribers can customize the engine and tools to meet their needs. Our subscribers are building the editor in Windows, Mac, Linux, and compiling games for both Android and iOS. Source code for all of our tools (such as Unreal Lightmass and Frontend) is included as well.
+
 [REGION:warning]
+This page shows subscribers how to download and build Unreal Engine from our source code repository on GitHub. If you would like to download the binary version of Unreal Engine, go to our main page to [get unreal](https://www.unrealengine.com/) (click the image pictured below).
+[/REGION]
+
+Source code for Unreal Engine 4 (UE4) can be downloaded by our subscribers from the Unreal Engine GitHub repository. With access to our source code, subscribers can customize the engine and tools to meet their needs. Our subscribers are building the editor in Windows, Mac, Linux, and compiling games for both Android and iOS. Source code for all of our tools (such as Unreal Lightmass and Frontend) is included as well.
+
+[REGION:note]
 Before you can access the [repository](https://github.com/EpicGames/UnrealEngine), you must be an Unreal Engine subscriber with a GitHub account.  
 *	If you haven't created a GitHub account, go to [GitHub's](https://github.com/) page to sign-up.   
 *	For instructions on associating your GitHub account with your Unreal Engine account, visit the [UE4 on GitHub](https://www.unrealengine.com/ue4-on-github) page.
@@ -66,8 +71,8 @@ Other branches may appear from time-to-time as we work to stabilize new releases
 
 Subscribers can download Unreal Engine's source code by:
 
-*	Downloading Unreal Engine from GitHub
-*	Downloading Unreal Engine as a Compressed File
+*	%GettingStarted/DownloadingUnrealEngine/GitHub:title%
+*	%GettingStarted/DownloadingUnrealEngine/DownloadZip:title%
 
 ### %GettingStarted/DownloadingUnrealEngine/GitHub:title%
 
diff --git a/Engine/Documentation/Source/GettingStarted/DownloadingUnrealEngine/GitHub/GitHub.INT.udn b/Engine/Documentation/Source/GettingStarted/DownloadingUnrealEngine/GitHub/GitHub.INT.udn
--- a/Engine/Documentation/Source/GettingStarted/DownloadingUnrealEngine/GitHub/GitHub.INT.udn
+++ b/Engine/Documentation/Source/GettingStarted/DownloadingUnrealEngine/GitHub/GitHub.INT.udn
@@ -1,7 +1,7 @@
 ﻿Availability:NoPublish
-Title:Downloading Unreal Engine from GitHub
+Title:Forking and Cloning Unreal from GitHub
 Crumbs:%ROOT%, GettingStarted, GettingStarted/DownloadingUnrealEngine
-Description:Downloading the Unreal Engine source code from GitHub.
+Description:Forking and cloning the Unreal Engine source code from GitHub.
 
 %Globals:osselect_linux%
 
