#UE4doc #docs - updated images and bold lettering.

diff --git a/Engine/Documentation/Source/Engine/AI/EnvironmentQuerySystem/UserGuide/UserGuide.INT.udn b/Engine/Documentation/Source/Engine/AI/EnvironmentQuerySystem/UserGuide/UserGuide.INT.udn
--- a/Engine/Documentation/Source/Engine/AI/EnvironmentQuerySystem/UserGuide/UserGuide.INT.udn
+++ b/Engine/Documentation/Source/Engine/AI/EnvironmentQuerySystem/UserGuide/UserGuide.INT.udn
@@ -52,7 +52,7 @@ Description:This document covers the interface and most commonly used workflows
 
 ## Parts of a Environment Query
 
-An Environment Query is actually made up a number of different pieces. You must call an Environment Query from a Behavior Tree, and then the actual Environment Query will use its Generator, reference its Contexts, and use its Tests to give the Behavior Tree the highest weighted result.
+An **Environment Query** is actually made up a number of different pieces. You must call an Environment Query from a Behavior Tree, and then the actual Environment Query will use its Generator, reference its Contexts, and use its Tests to give the Behavior Tree the highest weighted result.
 
 ### Behavior Tree
 
@@ -64,30 +64,30 @@ Selecting it in the Behavior Tree will enable you to set the **Query Template**
 
 ![](runEQS_details.png)
 
-See the [Behavior Tree](Engine\AI\BehaviorTrees) documentation for more information on their usage.
+See the [**Behavior Tree**](Engine\AI\BehaviorTrees) documentation for more information on their usage.
 
 
 ### Environment Query
 
 The Environment Query asset is what you call in the Behavior Tree using the **Run EQS Query** node, to test a number of Actors or locations to find the "best" one. What is really defined as "best" is really up to the way the Tests are structured in the Environment Query.
 
-**Environment Queries** can be created in Blueprint and C++.
+Environment Queries can be created in Blueprint and C++.
 
 
 ### Generator 
 
-A Generator produces the locations or Actors, referred to as **Items**, that will actually be tested and weighted, with the highest weighted location or Actor being returned to the Behavior Tree.
+A **Generator** produces the locations or Actors, referred to as **Items**, that will actually be tested and weighted, with the highest weighted location or Actor being returned to the Behavior Tree.
 
-**Generators** can be created in Blueprint and C++.
+Generators can be created in Blueprint and C++.
 
 
 ### Contexts
 
-Contexts are a frame of reference for various tests and Generators. They can be as simple as **The Querier** or more complex such as **All Actors of a Type**.
+**Contexts** are a frame of reference for various tests and Generators. They can be as simple as **The Querier** or more complex such as **All Actors of a Type**.
 
 Of note is the fact a Generator, such as the Simple Grid Generator can use a Context that returns multiple locations or Actors. This will create a Simple Grid, of the defined size and density, at the location of each Context.
 
-**Contexts** can be created in Blueprint and C++.
+Contexts can be created in Blueprint and C++.
 
 
 ### Tests
@@ -128,9 +128,9 @@ Currently, **Tests** must be created in C++; there is no Blueprint implementable
 
 ![](connectRootToGen.png)
 
-Just left-click and drag from the dark bar at the bottom of the Root node to the dark bar at the top of the Generator node to make a connection.
+Just **Left-click** and **drag** from the dark bar at the bottom of the Root node to the dark bar at the top of the Generator node to make a connection.
 
-Alternatively, if you left-click and drag into the graph area, you'll get an option to create a new Generator node.
+Alternatively, if you **Left-click** and **drag** into the graph area, you'll get an option to create a new Generator node.
 
 [REGION:note]
 While you can have more than one Generator connected to the Root, only the left most, valid Generator is used in the Query.
@@ -142,7 +142,7 @@ To use an Environment Query in a Behavior Tree, you need a **Run EQS Query** Tas
 
 ![](runEQS.png)
 
-1. Right-click in the Behavior Tree graph area.
+1. **Right-click** in the Behavior Tree graph area.
 1. Create a **Task** -> **Run EQSQuery** node.
 1. Select the newly created **Run EQSQuery** node, and change the Query Template property to your Environment Query asset (or you can create a new one from within the pull down).
 1. Optionally, set any **Query Params** the Environment Query may be looking for.
@@ -164,15 +164,15 @@ To use an Environment Query in a Behavior Tree, you need a **Run EQS Query** Tas
 
 ### Modifying Tests
 
-1. Select a Test on a Generator and modify it in the Details Panel
+1. Select a Test on a Generator and modify it in the **Details** panel
 
 ![](editTest.png)
 
-See the [Environmental Query System Node Reference](Engine/AI/EnvironmentQuerySystem/NodeReference) for usage information on individual tests.
+See the [**Environmental Query System Node Reference**](Engine/AI/EnvironmentQuerySystem/NodeReference) for usage information on individual tests.
 
 ### Previewing the Query Results
 
-See the [Environmental Query System Test Pawn](Engine/AI/EnvironmentQuerySystem/EQSPawn) documentation for information on how to preview the results of a Query
+See the [**Environmental Query System Test Pawn**](Engine/AI/EnvironmentQuerySystem/EQSPawn) documentation for information on how to preview the results of a Query
 
 ![](Engine\AI\EnvironmentQuerySystem\EQSPawn\testingPawnBanner.png)
 
