UEDOC-2098

Gif has been converted into a YouTube Video

diff --git a/Engine/Documentation/Source/Engine/AI/BehaviorTrees/QuickStart/BehaviorTreeQuickStart.INT.udn b/Engine/Documentation/Source/Engine/AI/BehaviorTrees/QuickStart/BehaviorTreeQuickStart.INT.udn
--- a/Engine/Documentation/Source/Engine/AI/BehaviorTrees/QuickStart/BehaviorTreeQuickStart.INT.udn
+++ b/Engine/Documentation/Source/Engine/AI/BehaviorTrees/QuickStart/BehaviorTreeQuickStart.INT.udn
@@ -37,7 +37,17 @@ Description:Overviews and examples AI
 [/OBJECT]
 [/VAR]
 
-![](Engine\AI\BehaviorTrees\QuickStart\14\BTFinal.gif)(convert:false)
+[OBJECT:EmbeddedVideo]
+[PARAMLITERAL:width]
+640
+[/PARAMLITERAL]
+[PARAMLITERAL:height]
+360
+[/PARAMLITERAL]
+[PARAMLITERAL:videoid]
+1td5wSD3PPU
+[/PARAMLITERAL]
+[/OBJECT]
 
 The Behavior Tree Quick Start Guide walks you through the process of creating a NavMesh, creating an AI Controller, creating a Character that will be controlled by that AI Controller, and creating all the parts necessary for a simple Behavior Tree.
 
