#UE4doc #docs - updated images and bold letter. Doc needs to be updated with Gray guy and gif converted to video.

diff --git a/Engine/Documentation/Source/Engine/Animation/AnimHowTo/AimOffset/AimOffset.INT.udn b/Engine/Documentation/Source/Engine/Animation/AnimHowTo/AimOffset/AimOffset.INT.udn
--- a/Engine/Documentation/Source/Engine/Animation/AnimHowTo/AimOffset/AimOffset.INT.udn
+++ b/Engine/Documentation/Source/Engine/Animation/AnimHowTo/AimOffset/AimOffset.INT.udn
@@ -7,7 +7,7 @@ Related: Engine/Animation/AnimBlueprints
 Related: Engine\Content\Types\SkeletalMeshes\Sockets
 Related: Engine\Animation\AnimHowTo\AdditiveAnimations
 
-In this guide, we will create an [Aim Offset](Engine/Animation/AimOffset), which is an asset that stores a blendable series of poses to help a character aim a weapon. We will take an animation sequence and chop it up into usable poses for an Aim Offset, get the player's pitch/yaw location of the mouse and use that to determine which blended pose to use so that the character moves and aims in the location of the mouse (with some restrictions) as seen in the example below. 
+In this guide, we will create an [**Aim Offset**](Engine/Animation/AimOffset), which is an asset that stores a blendable series of poses to help a character aim a weapon. We will take an animation sequence and chop it up into usable poses for an Aim Offset, get the player's pitch/yaw location of the mouse and use that to determine which blended pose to use so that the character moves and aims in the location of the mouse (with some restrictions) as seen in the example below. 
 
 ![](AimGIF.gif)(convert:false)
 
@@ -185,7 +185,7 @@ With the poses created, next we will create the Aim Offset asset and set up the
 	Set the **X Axis Label** to **Yaw** and **Y Axis Label** to **Pitch**, then set both Axis Ranges to **-90 to 90** and click **Apply Parameter Changes**. 
 
 	[REGION:note]
-	For information on how **Blend Spaces** work (which are what Aim Offsets are based upon), see [](Engine/Animation/Blendspaces). 
+	For information on how **Blend Spaces** work (which are what Aim Offsets are based upon), see [**Blend Spaces**](Engine/Animation/Blendspaces). 
 	[/REGION]
 
 1.	In the **Asset Browser**, search for **Aim**, then drag the **Aim_Center** animation onto the graph in the center position as shown below. 
@@ -220,14 +220,14 @@ In this section, we will implement the Aim Offset inside the character's **Anima
 	This is the Animation Blueprint that is assigned to the playable character and governs the logic for the character's animations. 
 
 	[REGION:note]
-	For information on how **Animation Blueprints** work, see the [](Engine/Animation/AnimBlueprints) documentation. 
+	For information on how **Animation Blueprints** work, see the [**Animation Blueprints**](Engine/Animation/AnimBlueprints) documentation. 
 	[/REGION]
 
 1.	Inside the Animation Blueprint, in the lower center window, click the **MyBlueprint** tab, then **Double-click** the **AnimGraph**. 
 
 	![](AimOffset32.png)
 
-	This will open the [](Engine/Animation/AnimBlueprints/AnimGraph) tab which is where we will implement our Aim Offset. 
+	This will open the [**Anim Graph**](Engine/Animation/AnimBlueprints/AnimGraph) tab which is where we will implement our Aim Offset. 
 
 1.	In the **MyBlueprint** window, click the **Add Variable** button to create a new variable. 
 
@@ -259,7 +259,7 @@ In this section, we will implement the Aim Offset inside the character's **Anima
 
 	![](AimOffset38.png)
 
-	This will open the [](Engine/Animation/AnimBlueprints/EventGraph) tab which is used to keep the character animations updated with what is occurring in-game.
+	This will open the [**Event Graph**](Engine/Animation/AnimBlueprints/EventGraph) tab which is used to keep the character animations updated with what is occurring in-game.
 
 1.	In the **EventGraph**, locate the **Sequence** node in the graph. 
 
@@ -381,7 +381,7 @@ In this final section, we will change the default idle pose and update Character
 
 When you play in the editor, the character will now aim while stationary and will react to mouse movement, pointing in the direction of your mouse. There are some restrictions in place that prevent the character from turning and facing backwards while the mouse is behind them as well and while the character is jogging, moving the mouse also turns the character's torso in the direction you are pointing. 
 
-You can take this a few steps further by allowing the character to aim in a direction and play a shooting animation in that direction and (or) allow the character to play a shoot animation while moving in a direction or crouching by taking a look at the [](Engine\Animation\AnimHowTo\AdditiveAnimations) How To page to learn about blending multiple animations together. Or, you can take a look at [Skeletal Mesh Sockets](Engine\Content\Types\SkeletalMeshes\Sockets) for examples on how to attach an weapon to the character's hands now that they can aim. 
+You can take this a few steps further by allowing the character to aim in a direction and play a shooting animation in that direction and (or) allow the character to play a shoot animation while moving in a direction or crouching by taking a look at the [**Using Additive Animations**](Engine\Animation\AnimHowTo\AdditiveAnimations) How To page to learn about blending multiple animations together. Or, you can take a look at [**Skeletal Mesh Sockets**](Engine\Content\Types\SkeletalMeshes\Sockets) for examples on how to attach an weapon to the character's hands now that they can aim. 
 
 
 
