#UEDoc #docs - updated image and bold lettering. Gif needs to be converted to video.

diff --git a/Engine/Documentation/Source/Engine/Animation/Tools/VertexAnimationTool/VertexAnimationTools.INT.udn b/Engine/Documentation/Source/Engine/Animation/Tools/VertexAnimationTool/VertexAnimationTools.INT.udn
--- a/Engine/Documentation/Source/Engine/Animation/Tools/VertexAnimationTool/VertexAnimationTools.INT.udn
+++ b/Engine/Documentation/Source/Engine/Animation/Tools/VertexAnimationTool/VertexAnimationTools.INT.udn
@@ -8,9 +8,9 @@ Version: 4.8
 
 [TOC(start:2 end:3)]
 
-## Overview
+
 Storing complex animation data into 2D Textures or a meshes UV's is a great way to reduce animation overhead while still maintaining the required look and feel of the animation. 
-In the past you might have used [Morph Targets](Engine/Content/FBX/MorphTargets) to do the exact same thing, however this new method has a few advantages that Morph Targets do not.
+In the past you might have used [**Morph Targets**](Engine/Content/FBX/MorphTargets) to do the exact same thing, however this new method has a few advantages that Morph Targets do not.
 One of those advantages is allowing you to use complex animation data inside systems that would not otherwise be capable of utilizing Morph Targets like the Cascade Particle Editor.
 In the following documents we will cover how you can use this script to generate content for your UE4 projects.
 
@@ -36,7 +36,7 @@ The tool generates the data in the Texture using the following formula.
 
 Due to this limitation this tool is best suited for Visual Effects or ambient Static Meshes that require animations but might not justify using a complex animation rig to do so. 
 This tool will also **not work** with Skeletal Mesh Animations as the bone transforms are not available to the Material Editor.
-This means that if you want to affect the vertices of a Skeletal Mesh in a similar manner you will need to use [Morph Targets](Engine/Content/FBX/MorphTargets) instead.
+This means that if you want to affect the vertices of a Skeletal Mesh in a similar manner you will need to use [**Morph Targets**](Engine/Content/FBX/MorphTargets) instead.
 
 ### Vertex Animation Tools Break Down
 Inside of the Vertex Animation Tools you will find two entirely different methods for affecting the vertices of Static Meshes. 
@@ -72,7 +72,7 @@ To install the script just drag it from its location in **4.8\Engine\Extras\3dsM
 ![](VAT_Max_Install.gif)(convert:false)
 
 [region:tip]
-If you find that you are using this script a lot you can always add the script to one of the tool bars or quad menus. If you are unfamiliar how to do this the [Autodesk site has a very detailed walk through](http://knowledge.autodesk.com/support/3ds-max/learn-explore/caas/CloudHelp/cloudhelp/2015/ENU/3DSMax/files/GUID-A2CF8BAA-7B52-40A8-8C40-803B1AB5FC05-htm.html) that will explain how you go about doing this.
+If you find that you are using this script a lot you can always add the script to one of the tool bars or quad menus. If you are unfamiliar how to do this the [**Autodesk site has a very detailed walk through**](http://knowledge.autodesk.com/support/3ds-max/learn-explore/caas/CloudHelp/cloudhelp/2015/ENU/3DSMax/files/GUID-A2CF8BAA-7B52-40A8-8C40-803B1AB5FC05-htm.html) that will explain how you go about doing this.
 [/region]
 
 ## 3Ds Max Unit Setup
