Fix for Jira UEDOC-1965

diff --git a/Engine/Documentation/Source/Gameplay/HowTo/CharacterMovement/Blueprints/Setup_2/Setup_2.INT.udn b/Engine/Documentation/Source/Gameplay/HowTo/CharacterMovement/Blueprints/Setup_2/Setup_2.INT.udn
--- a/Engine/Documentation/Source/Gameplay/HowTo/CharacterMovement/Blueprints/Setup_2/Setup_2.INT.udn
+++ b/Engine/Documentation/Source/Gameplay/HowTo/CharacterMovement/Blueprints/Setup_2/Setup_2.INT.udn
@@ -55,7 +55,11 @@ In this section, we will finish setting up the Character by implementing the scr
 
 1.	Set up the node network making sure to check "Jog Pressed" in the appropriate **Set** node as shown below:
 
-	![](Character_JogB.png)(w:720)
+	[REGION:lightbox]
+	[![](Character_JogB.png)(w:720)](Character_JogB.png)
+	[/REGION:lightbox]
+
+	_Click image for full view._
 
 	Here we are setting **Jog Pressed** to _True_ when Jog is Pressed, then off the **Character Movement** Component, we are setting the characters **Max Walk Speed** to **375** (you can adjust this value to increase/decrease Jog speed). 
 	
