#UE4doc #docs - adjusted bold lettering

diff --git a/Engine/Documentation/Source/Engine/Audio/AmbientZones/AmbientZones.INT.udn b/Engine/Documentation/Source/Engine/Audio/AmbientZones/AmbientZones.INT.udn
--- a/Engine/Documentation/Source/Engine/Audio/AmbientZones/AmbientZones.INT.udn
+++ b/Engine/Documentation/Source/Engine/Audio/AmbientZones/AmbientZones.INT.udn
@@ -5,31 +5,29 @@ Description:Creating ambient audio zones using the Ambient Zones properties of A
 
 [TOC(start:2 end:2)]
 
-## Overview
-
-**Ambient Zones** are an inexpensive interior/exterior ambient audio solution.  Benefits of **Ambient Zones** include:
+**Ambient Zones** are an inexpensive interior/exterior ambient audio solution.  Benefits of Ambient Zones include:
 
 * Makes the transition from interior to exteriors areas sound good.
 * Reduces labor for audio designers, as less **Sound Actors** will need to be placed and iterated on inside/outside spaces.
-* Reduces game systems overhead, due to the savings in number of needed **Sound Actors** to define interiors and exteriors.
+* Reduces game systems overhead, due to the savings in number of needed Sound Actors to define interiors and exteriors.
 
 
 
 ## Ambient Zone Settings
 
-The settings for an **Ambient Zone** define how **Sound Actors** located inside and outside of its associated **Audio Volume** are altered by the player's location.
+The settings for an Ambient Zone define how Sound Actors located inside and outside of its associated **Audio Volume** are altered by the player's location.
 
 A simple example of how this works is shown below: 
 
 ![](AmbientZone_Example.png)(w:480)
 
-In the image above, the gold box around the section marked as "Inside" is an **Ambient Zone** volume; the two speaker icons are **Sound Actors**. As you can see, one **Sound Actor** is located inside the **Ambient Zone** and the other is located outside the **Ambient Zone**.
+In the image above, the gold box around the section marked as "Inside" is an Ambient Zone volume; the two speaker icons are Sound Actors. As you can see, one Sound Actor is located inside the Ambient Zone and the other is located outside the Ambient Zone.
 
-Here is the behavior of the **Ambient Zone** volume when playing in the level:
+Here is the behavior of the Ambient Zone volume when playing in the level:
 
-* When the player is outside of the **Ambient Zone**, the **Sound Actor** located inside the **Ambient Zone** will have a volume multiplier and LPF effect applied to it. 
-* When the player enters the **Ambient Zone**, the **Sound Actor** located inside the **Ambient Zone** will return to a 1.0 multiplier for volume and LPF over a specified fade time.  At the same time, the **Sound Actor** located outside of the **Ambient Zone** will have a volume and LPF multiplier applied to it over a specified fade time.  
-* When the player exits the **Ambient Zone**, the **Sound Actor** outside the **Ambient Zone** will return to its previous default volume and LPF values using the previous fade time, and the **Sound Actor** located inside the **Ambient Zone** will have its volume multiplier and LPF values changed via the **Ambient Zone Settings** for the **Ambient Zone**.
+* When the player is outside of the Ambient Zone, the Sound Actor located inside the Ambient Zone will have a volume multiplier and LPF effect applied to it. 
+* When the player enters the Ambient Zone, the Sound Actor located inside the Ambient Zone will return to a 1.0 multiplier for volume and LPF over a specified fade time.  At the same time, the Sound Actor located outside of the Ambient Zone will have a volume and LPF multiplier applied to it over a specified fade time.  
+* When the player exits the Ambient Zone, the Sound Actor outside the Ambient Zone will return to its previous default volume and LPF values using the previous fade time, and the Sound Actor located inside the Ambient Zone will have its volume multiplier and LPF values changed via the Ambient Zone Settings for the Ambient Zone.
 
 You can access the **Ambient Zone Settings** from the **Details** panel of an **Audio Volume** by expanding it as seen below. 
 
@@ -37,11 +35,11 @@ You can access the **Ambient Zone Settings** from the **Details** panel of an **
 
 | **Property** | **Description** |
 | --- | --- |
-| **Exterior Volume** | The final volume of exterior sounds when player inside this **Ambient Zone**. |
+| **Exterior Volume** | The final volume of exterior sounds when player inside this Ambient Zone. |
 | **Exterior Time** | Time to fade to new exterior volume in seconds. |
 | **Exterior LPF** | Lowpass filter multiplier applied to exterior sounds when inside (0.1 to apply the maximum LPF). |
 | **Exterior LPFTime** | Time to fade to new lowpass filter level in seconds. |
-| **Interior Volume** | The final volume of interior sounds when outside this **Ambient Zone**. |
+| **Interior Volume** | The final volume of interior sounds when outside this Ambient Zone. |
 | **Interior Time** | Time to fade to new interior volume in seconds. |
 | **Interior LPF** | Lowpass filter multiplier applied to interior sounds when outside (0.1 to apply the maximum LPF). |
 | **Interior LPFTime** | Time to fade to new lowpass filter level in seconds. |
@@ -56,22 +54,22 @@ You can access the **Ambient Zone Settings** from the **Details** panel of an **
 
 	![](AmbientZone_Settings.png)
 1. In the **Ambient Zone** properties, set the **Exterior Volume** to .5  
-	* This is the final volume of **Ambient Sounds** that are outside the **Ambient Zone** when the player is inside the **Ambient Zone**.
+	* This is the final volume of Ambient Sounds that are outside the Ambient Zone when the player is inside the Ambient Zone.
 1. Set the **Exterior LPF** to .5  
-	* This is the final **LowPassFilter** on **Ambient Sounds** that are outside the **Ambient Zone** when the player is inside the **Ambient Zone**.
+	* This is the final LowPassFilter on Ambient Sounds that are outside the Ambient Zone when the player is inside the Ambient Zone.
 1. Set the **Interior Volume** to .2  
-	* This is the final volume of **Ambient Sounds** that are inside the **Ambient Zone** when the player is outside the **Ambient Zone**.
+	* This is the final volume of Ambient Sounds that are inside the Ambient Zone when the player is outside the Ambient Zone.
 1. Set the **Interior LPF** to .2 
-	* This is the final **LowPassFilter** on **Ambient Sounds** that are inside the **Ambient Zone** when the player is outside the **Ambient Zone**.
+	* This is the final LowPassFilter on Ambient Sounds that are inside the Ambient Zone when the player is outside the Ambient Zone.
 1. Create a new **Sound Class** and in its properties, check **Apply Ambient Volumes**: 
 	![](AmbientZone_SoundClass.png)
 1. Assign the **Sound Class** to both the interior and exterior sounds from the **Details** panel of each sound: 
 	![](SoundClass_Override.png)
 1. Rebuild geometry, and enter the **Audio Volume** with the **Ambient Zone** you just set up.
-	* The **Ambient Sound** that is outside the **Ambient Zone** will have its volume and **LowPassFilter** multiplied by .5
+	* The Ambient Sound that is outside the Ambient Zone will have its volume and LowPassFilter multiplied by .5
 1. Exit the **Audio Volume**.  
-	* Now the **Ambient Sound** that was outside the **Audio Volume** is restored to its previous volume and **LPF** settings.  
-	* The **Ambient Sound** inside the **Audio Volume** now has a volume multiplier and **LPF** multiplier of .2 applied, making it very quiet.
+	* Now the Ambient Sound that was outside the **Audio Volume** is restored to its previous volume and **LPF** settings.  
+	* The Ambient Sound inside the Audio Volume now has a volume multiplier and LPF multiplier of .2 applied, making it very quiet.
 	
 [Region: note]
 Note: When placing a **Sound Actor** inside of an **Ambient Zone** volume, you must rebuild the level's geometry, as this is when the check is performed on the **Sound Actor** to determine its location.  After geometry has been rebuilt, all **Ambient Zone** parameters can be freely changed in the editor.
