#UE4doc #docs:
Removed Overview header
Changed font settings for multiple terms

diff --git a/Engine/Documentation/Source/Resources/ContentExamples/Blueprints/Blueprints.INT.udn b/Engine/Documentation/Source/Resources/ContentExamples/Blueprints/Blueprints.INT.udn
--- a/Engine/Documentation/Source/Resources/ContentExamples/Blueprints/Blueprints.INT.udn
+++ b/Engine/Documentation/Source/Resources/ContentExamples/Blueprints/Blueprints.INT.udn
@@ -49,11 +49,11 @@ Related: Gameplay/HowTo/SetUpInput/Blueprints
 [TOC(start:2 end:2)]
 
 
-## Overview
+
 
 There are several **Blueprint Maps** located inside the **Content Examples** project that demonstrate the power and flexibility of the **Blueprints** system. Whether you are looking to define player controls, implement a health/pick-up system, or simply have a door that opens when a player approaches it, the examples illustrated in these maps should give you some ideas on how to approach whatever it is you are after. 
 
-This page indicates what examples are included in each map and where to find them so that you may jump to your area of interest and begin scripting right away. If you are new to **Blueprints** however, it is suggested that you begin with the **Blueprints_Overview** map as it will cover the core aspects of the system and will help get you started with using them. 
+This page indicates what examples are included in each map and where to find them so that you may jump to your area of interest and begin scripting right away. If you are new to Blueprints however, it is suggested that you begin with the **Blueprints_Overview** map as it will cover the core aspects of the system and will help get you started with using them. 
 
 [region:note]
 You can find additional resources on how an example was constructed (if applicable) by clicking on the example's name in their respective tables.
@@ -63,18 +63,18 @@ You can find additional resources on how an example was constructed (if applicab
 
 ![](blueprint_ce_header.png)(w:640)
 
-The **Blueprints_Overview** map covers the essential aspects of **Blueprints** including: **Components**, using the **Construction Script** and **Event Graph**, **Variables** and how to edit/set them, simple mathematics, and controlling the flow of a **Blueprint** by using different types of **Flow Control** nodes. 
-Take your time going over these sections as you will use these elements quite often when working with **Blueprints** in the future. 
+The **Blueprints_Overview** map covers the essential aspects of Blueprints including: **Components**, using the **Construction Script** and **Event Graph**, **Variables** and how to edit/set them, simple mathematics, and controlling the flow of a Blueprint by using different types of **Flow Control** nodes. 
+Take your time going over these sections as you will use these elements quite often when working with Blueprints in the future. 
 
 | Example | What is Demonstrated |
 | ------- | ------------ |
-| **[1.1 Hello World Example](Resources\ContentExamples\Blueprints\1_1)** | A basic **Blueprint** is used to display floating text in the level (e.g. floating text "Hello World" is displayed when playing in the level). |
-| **[1.2 Components Example](Resources\ContentExamples\Blueprints\1_2)** | Examples of **Components** and how multiple **Components** can be used in a single **Blueprint** (e.g. a sprite, Static Mesh, and light all used in one **Blueprint**). |
-| **[1.3 Variables Example](Resources\ContentExamples\Blueprints\1_3)** | Shows how **Variables** can be set to `public` or `private` and how to edit them (e.g. a light color that can be changed via the **Details** panel or only changed in the **Blueprint**). |
-| **[1.4 Construction Script Example](Resources\ContentExamples\Blueprints\1_4)** | Shows how **Components** of a **Blueprint** can be disabled/enabled with the **Construction Script** (e.g. use checkboxes to turn on/off **Components** in a **Blueprint**). |
+| **[1.1 Hello World Example](Resources\ContentExamples\Blueprints\1_1)** | A basic Blueprint is used to display floating text in the level (e.g. floating text "Hello World" is displayed when playing in the level). |
+| **[1.2 Components Example](Resources\ContentExamples\Blueprints\1_2)** | Examples of Components and how multiple Components can be used in a single Blueprint (e.g. a sprite, Static Mesh, and light all used in one Blueprint). |
+| **[1.3 Variables Example](Resources\ContentExamples\Blueprints\1_3)** | Shows how Variables can be set to public or private and how to edit them (e.g. a light color that can be changed via the **Details** panel or only changed in the Blueprint). |
+| **[1.4 Construction Script Example](Resources\ContentExamples\Blueprints\1_4)** | Shows how Components of a Blueprint can be disabled/enabled with the **Construction Script** (e.g. use checkboxes to turn on/off Components in a Blueprint). |
 | **[1.5 Event Graph Example](Resources\ContentExamples\Blueprints\1_5)** | The **Construction Script** and **Event Graph** are used to allow custom colors for a mesh without editing the material (e.g. user defines the color of a rocket via the **Details** panel).  |
-| **[1.6 Simple Math Example](Resources\ContentExamples\Blueprints\1_6)** | An **Event Tick** is used in this **Blueprint** to perform simple addition over time (e.g. a timer is displayed counting upwards every second from 0).  |
-| **[1.7 Flow Example](Resources\ContentExamples\Blueprints\1_7)** | A **Blueprint** that uses different types of **Flow Control** to modify movement, color and effects of an object (e.g. the direction a coin moves, its color, and if effects are applied via the **Details** panel). |
+| **[1.6 Simple Math Example](Resources\ContentExamples\Blueprints\1_6)** | An **Event Tick** is used in this Blueprint to perform simple addition over time (e.g. a timer is displayed counting upwards every second from 0).  |
+| **[1.7 Flow Example](Resources\ContentExamples\Blueprints\1_7)** | A Blueprint that uses different types of Flow Control to modify movement, color and effects of an object (e.g. the direction a coin moves, its color, and if effects are applied via the **Details** panel). |
 
 
 ## Blueprints_Advanced Map
@@ -95,14 +95,14 @@ The **Construction Script** section illustrates how to: randomly spawn meshes in
 
 
 ### Event Graph Examples
-The **Event Graph** examples will show you how to: animate a Static Mesh with a timeline, automatically rotate an Actor to track player movement, how to pass variables between **Blueprints**, create a door that opens automatically when triggered or create different types of pick-ups. 
+The **Event Graph** examples will show you how to: animate a Static Mesh with a timeline, automatically rotate an Actor to track player movement, how to pass variables between Blueprints, create a door that opens automatically when triggered or create different types of pick-ups. 
 
 | Example | What is Demonstrated |
 | ------- | ------------ |
 | **[2.1 Timeline Animation Example](Resources\ContentExamples\Blueprints_Advanced\2_1)** | How to use a **Timeline** to animate a Static Mesh wth an **Event Track** to spawn an emitter at a specified point (e.g. creating an animated bouncing ball). |
 | **[2.2 Actor Tracking Example](Resources\ContentExamples\Blueprints_Advanced\2_2)** | How to make an Actor turn to face the player when coming within a specified range (e.g. a spotlight that reacts to and follows the player).  |
-| **[2.3 Button Trigger (Blueprint Interface) Example](Resources\ContentExamples\Blueprints_Advanced\2_3)** | How to pass a variable from one **Blueprint** to another using **Blueprint Interfaces** (e.g. an on/off variable is passed from a button to a set of gears that animate). |
-| **[2.4 Opening Door (Trigger) Example](Resources\ContentExamples\Blueprints_Advanced\2_4)** | How to open/close a door automatically with a **Trigger** and a **Timeline**. |
+| **[2.3 Button Trigger (Blueprint Interface) Example](Resources\ContentExamples\Blueprints_Advanced\2_3)** | How to pass a variable from one Blueprint to another using **Blueprint Interfaces** (e.g. an on/off variable is passed from a button to a set of gears that animate). |
+| **[2.4 Opening Door (Trigger) Example](Resources\ContentExamples\Blueprints_Advanced\2_4)** | How to open/close a door automatically with a **Trigger** and a Timeline. |
 | **[2.5 Child Blueprints Example](Resources\ContentExamples\Blueprints_Advanced\2_5)** | How to implement different types of pick-ups derived from the same **Parent Class** (e.g. a pick-up for health, a coin, and a cupcake each with different functionality when picked up). |
 
 
@@ -110,46 +110,46 @@ The **Event Graph** examples will show you how to: animate a Static Mesh with a
 
 ![](BlueprintCommunication_Header.png)(w:640)
 
-When working with **Blueprints**, you may need to execute a function or access certain variables from another **Blueprint** and depending on your needs, different methods of communication are provided. 
-The **Blueprint_Communication** map shows the various ways that you can pass information from one **Blueprint** to another through **Direct Blueprint Communication**, **Event Dispatchers**, or **Blueprint Interfaces**. 
+When working with Blueprints, you may need to execute a function or access certain variables from another Blueprint and depending on your needs, different methods of communication are provided. 
+The **Blueprint_Communication** map shows the various ways that you can pass information from one Blueprint to another through **Direct Blueprint Communication**, **Event Dispatchers**, or **Blueprint Interfaces**. 
 
 ### Blueprint Communication Examples
-The **Direct Blueprint Communication** or **Blueprint Communication** section shows the basic ways in which you can access and provide information between **Blueprints** through targeting or casting to a **Blueprint**. 
+The **Direct Blueprint Communication** or **Blueprint Communication** section shows the basic ways in which you can access and provide information between Blueprints through targeting or casting to a Blueprint. 
 
 | Example | What is Demonstrated |
 | ------- | ------------ |
-| **1.1 Basic Communication (Target Blueprint)** | An example of a button **Blueprint** that toggles a light **Blueprint** on and off through a custom event. |
-| **1.2 Basic Communication (Actor Casting)** | How to use an overlap event to detect a **Blueprint** of a particular class (e.g. a light bulb is turned on/off by connecting a battery).|
+| **1.1 Basic Communication (Target Blueprint)** | An example of a button Blueprint that toggles a light Blueprint on and off through a custom event. |
+| **1.2 Basic Communication (Actor Casting)** | How to use an overlap event to detect a Blueprint of a particular class (e.g. a light bulb is turned on/off by connecting a battery).|
 | **1.3 Communication (Actor Casting To Child Blueprints)** | How use **Child Blueprints** which share functionality of their parent but have independent functionality as well (e.g. turning lights on/off by connecting batteries each with different energy values).  |
 | **1.4 Communication with All Actors of a Specific Class** | How to communicate with several Actors of a particular class (e.g. turn on several lights at once each with different functionality).  |
 
 
 ### Event Dispatcher Examples
-The **Event Dispatcher** examples cover another method in which you can communicate between **Blueprints** and **Level Blueprint**.
+The **Event Dispatcher** examples cover another method in which you can communicate between Blueprints and **Level Blueprint**.
 In addition, by **Binding** one or more events to an **Event Dispatcher**, you can cause multiple events that are bound to the event to fire at once when the **Event Dispatcher** is called. 
 
 | Example | What is Demonstrated |
 | ------- | ------------ |
-| **2.1 Using Event Dispatchers** | How to use an **Event Dispatcher** to call an event inside the **Level Blueprint** (e.g. toggling a light on/off using an **Event Dispatcher**). |
-| **2.2 Binding Event Dispatchers to A Custom Event** | Illustrates binding an **Event Dispatcher** to a custom event inside of the **Level Blueprint** (e.g. toggling a button that rotates an object over time). |
-| **2.3 Binding Event Dispatchers to A Custom Event on Spawn** | Illustrates how to enable/disable a trigger with an **Event Dispatcher** that is called when a spawned Actor is destroyed (e.g. press a button to spawn a bomb, no other bombs can be spawned until the previous bomb explodes). |
+| **2.1 Using Event Dispatchers** | How to use an **Event Dispatcher** to call an event inside the **Level Blueprint** (e.g. toggling a light on/off using an Event Dispatcher). |
+| **2.2 Binding Event Dispatchers to A Custom Event** | Illustrates binding an Event Dispatcher to a custom event inside of the Level Blueprint (e.g. toggling a button that rotates an object over time). |
+| **2.3 Binding Event Dispatchers to A Custom Event on Spawn** | Illustrates how to enable/disable a trigger with an Event Dispatcher that is called when a spawned Actor is destroyed (e.g. press a button to spawn a bomb, no other bombs can be spawned until the previous bomb explodes). |
 
 
 ### Blueprint Interface Examples
-There may be instances where you have multiple **Blueprints** which share some basic commonalities (like being shot by a projectile and taking some form of damage) but are entirely different entities (like a car and a tree). 
+There may be instances where you have multiple Blueprints which share some basic commonalities (like being shot by a projectile and taking some form of damage) but are entirely different entities (like a car and a tree). 
 The **Blueprint Interface** will allow you to have a method in which both the car and tree can communicate and share how to react when damaged but have different components and functionality in-and-of themselves. 
 
 | Example | What is Demonstrated |
 | ------- | ------------ |
-| **3.1 Blueprint Interface (Array of Blueprints)** | Illustrates how to communicate with multiple **Blueprints** each of different classes through an array (e.g. turn on/off different lights and rotating objects with one switch). |
-| **3.2 Blueprint Interface (Interactive Blueprints)** | How to communicate from the player **Blueprint** to Actor **Blueprints** based on input (e.g. an interactive switch, crank, or lever with a pulley chain). |
-| **3.3 Blueprint Interface (Projectile Collision)** | How to use a **Blueprint Interface** to check what an Actor of a particular type was hit by (e.g. three objects react differently when shot with fire or water). |
+| **3.1 Blueprint Interface (Array of Blueprints)** | Illustrates how to communicate with multiple Blueprints each of different classes through an array (e.g. turn on/off different lights and rotating objects with one switch). |
+| **3.2 Blueprint Interface (Interactive Blueprints)** | How to communicate from the player Blueprint to Actor Blueprints based on input (e.g. an interactive switch, crank, or lever with a pulley chain). |
+| **3.3 Blueprint Interface (Projectile Collision)** | How to use a Blueprint Interface to check what an Actor of a particular type was hit by (e.g. three objects react differently when shot with fire or water). |
 
 ## Blueprint_Input_Examples Map
 
 ![](BlueprintInput_Header.png)(w:640)
 
-The **Blueprint_Input_Examples** map gives you three examples of creating player input through **Blueprints** in three different styles of games: a 2D side scroller, a top-down 2D space shooter, and a 3D flying UFO. 
+The **Blueprint_Input_Examples** map gives you three examples of creating player input through Blueprints in three different styles of games: a 2D side scroller, a top-down 2D space shooter, and a 3D flying UFO. 
 
 
 | Example | What is Demonstrated |
@@ -163,7 +163,7 @@ The **Blueprint_Input_Examples** map gives you three examples of creating player
 
 ![](BlueprintMouse_Header.png)(w:640)
 
-The **Blueprint_Mouse_Interaction** map shows how to incorporate mouse input (or mobile touch inputs) as well as how to have objects react to mouse or touch input through **Blueprints**. 
+The **Blueprint_Mouse_Interaction** map shows how to incorporate mouse input (or mobile touch inputs) as well as how to have objects react to mouse or touch input through Blueprints. 
 The example given allows you to click-and-drag objects around on screen that react when interacted with as well as interacting with moving components. 
 
 More detailed information for this map can be found here: 
