#UE4doc #docs:
Removed Overview header
Removed italics from a term

diff --git a/Engine/Documentation/Source/Resources/ContentExamples/Blueprints/1_1/1_1.INT.udn b/Engine/Documentation/Source/Resources/ContentExamples/Blueprints/1_1/1_1.INT.udn
--- a/Engine/Documentation/Source/Resources/ContentExamples/Blueprints/1_1/1_1.INT.udn
+++ b/Engine/Documentation/Source/Resources/ContentExamples/Blueprints/1_1/1_1.INT.udn
@@ -41,7 +41,7 @@ This example is best seen during Play In Editor or Simulation. To Simulate, clic
 
 
 
-## Overview
+
 
 This example simply shows how to make an update to a Component's property using the Event Graph by using the classic programmer introductory project of "Hello World." Simply put, the Blueprint script for the object sets the Text property of a Text Render Component to "Hello World" as soon as gameplay begins. Prior to gameplay, such as when using the editor, the Text property will be blank so no text will appear. This is one of the simplest and most basic approaches to demonstrating the Event Graph and Blueprint setup.
 
@@ -51,7 +51,7 @@ As with all examples in this level, the object displayed is a Blueprint Actor, m
 There are only two Components in this setup:
 
 * A **Billboard Component**, which is used for placement and easy selection. This is named **Billboard1**, which is the default name. We will not need to refer to this in our Blueprint.
-* A **Text Render Component**, which will have its _Text_ property updated at the beginning of gameplay. This is named **HelloWorld**, and we will refer to it in the Event Graph to change the text.
+* A **Text Render Component**, which will have its Text property updated at the beginning of gameplay. This is named **HelloWorld**, and we will refer to it in the Event Graph to change the text.
 
 ## Construction Script
 In this example, we do not use the Construction Script. This is because we do not need anything to happen during level editing (construction) time. 
