#UE4doc #docs:
Removed Overview header
Changed font settings for multiple terms

diff --git a/Engine/Documentation/Source/Resources/ContentExamples/Blueprints/1_3/1_3.INT.udn b/Engine/Documentation/Source/Resources/ContentExamples/Blueprints/1_3/1_3.INT.udn
--- a/Engine/Documentation/Source/Resources/ContentExamples/Blueprints/1_3/1_3.INT.udn
+++ b/Engine/Documentation/Source/Resources/ContentExamples/Blueprints/1_3/1_3.INT.udn
@@ -36,7 +36,7 @@ Related: Gameplay/HowTo/SetUpInput/Blueprints
 
 [TOC(start:2 end:3)]
 
-## Overview
+
 
 In this example, we have a Blueprint Actor that has two spot lights set up. The color of each spot light is controlled by a **variable**. In programming, a variable is defined as a location in the computer's memory that stores some amount of data. In Blueprint scripts, variables are used in precisely the same way as virtual containers into which we can place, and out of which we can query, some kind of information. In this case, each light in our example uses a Linear Color variable, which stores an RGB value for a specific color. We can put in default values for this color, query it to use it in various locations, and set new colors to it, all as parts of our Blueprint script.
 
@@ -48,10 +48,10 @@ On the left spot light, this variable has been made **public**, meaning that it
 
 Making a variable public simply requires that you check its **Editable** property, or click the eye icon next to a variable's name in the **My Blueprint** window. Once a variable is editable, the eye-shaped button next to it will highlight yellow. This indicates that although the variable is publicly editable, it has no tooltip, which may confuse some users. To remedy this, enter a description in the variable's **Tooltip** property. Once this is done, the eye-shaped button will appear green.
 
-It is important for beginners to note that you will not always want your variables to be public. Sometimes, giving a user the ability to easily change a critical variable can end up breaking or potentially altering important functionality. Always make sure that your public variables are things that an artist or level designer will _need_ access to, and leave anything they are not supposed to touch as private.
+It is important for beginners to note that you will not always want your variables to be public. Sometimes, giving a user the ability to easily change a critical variable can end up breaking or potentially altering important functionality. Always make sure that your public variables are things that an artist or level designer will need access to, and leave anything they are not supposed to touch as private.
 
 [REGION:tip]
-Always make sure your public variables have tooltips!
+You should always make sure your public variables have tooltips to explain what they are used for.
 [/REGION]
 
 
@@ -64,7 +64,7 @@ The Components for this Blueprint Actor are fairly extensive, because this examp
 There are two separate setups, though both are essentially the same. The Component setup for each light is as follows:
 
 [REGION:note]
-Note that the words **Public** and **Private** will be appended to these names as a prefix for each of the respective light setups! Also, be aware that indentation indicates parenting relationships from one Component to another, just as it does within the Components list.
+Note that the words **Public** and **Private** will be appended to these names as a prefix for each of the respective light setups. Also, be aware that indentation indicates parenting relationships from one Component to another, just as it does within the Components list.
 [/REGION]
 
 *  **Wall Mount** - This is the Static Mesh Component that appears to be fastened to the wall.
@@ -105,7 +105,7 @@ Not all variables need to be created manually within the **My Blueprint** panel.
 
 ## Construction Script
 
-Most of the work in this effect is actually done within the Construction Script. **Since the lights do not need to update once the game begins, we do not use the Event Graph.**
+Most of the work in this effect is actually done within the Construction Script. Since the lights do not need to update once the game begins, we do not use the Event Graph.
 
 To make understanding the Construction Script easy in this example, we have broken it up into a series of commented blocks, as you can see below:
 
@@ -148,7 +148,7 @@ In these portions of the script, we simply take the default values that were app
 
 The next portion of the network creates a Dynamic Material Instance for each of the Light Fixture Static Meshes. A Dynamic Material Instance is an instantiated copy of a Material that can be updated in some way during gameplay. In this case, we can change the emissive color of the Material so that it matches the color used in the Spot Light Components. You will notice that we create two separate Dynamic Material Instances. This is because we need different instances for the two different colors. 
 
-After each Create Dynamic Material Instance node, you will also see that we store the newly created Material in a variable. This gives us access to the Material later, which we will need if we want to be able to adjust the color. We could just connect the _Return Value_  output to other places later on in the network, but storing in a variable makes the Blueprint script appear neater and gives us access to the Material if we should ever need it in the Event Graph.
+After each Create Dynamic Material Instance node, you will also see that we store the newly created Material in a variable. This gives us access to the Material later, which we will need if we want to be able to adjust the color. We could just connect the Return Value  output to other places later on in the network, but storing in a variable makes the Blueprint script appear neater and gives us access to the Material if we should ever need it in the Event Graph.
 
 ### Assigning the new Dynamic Material Instances 
 
