#UE4doc #docs:
Removed Overview header
Changed font settings for multiple terms
Changed text and picture referring to the 'Component' mode in the Blueprint Editor, as it no longer exists and has been replaced by the Component tab.

diff --git a/Engine/Documentation/Source/Resources/ContentExamples/Blueprints/1_2/1_2.INT.udn b/Engine/Documentation/Source/Resources/ContentExamples/Blueprints/1_2/1_2.INT.udn
--- a/Engine/Documentation/Source/Resources/ContentExamples/Blueprints/1_2/1_2.INT.udn
+++ b/Engine/Documentation/Source/Resources/ContentExamples/Blueprints/1_2/1_2.INT.udn
@@ -34,9 +34,9 @@ Related: Gameplay/HowTo/SetUpInput/Blueprints
 ![](ComponentsExample.png)
 [/REGION]
 
-## Overview
 
-Components can be thought of as the objects that will be needed for a Blueprint object to fulfill its purpose. There are many types of Components that can be added to a Blueprint. For example, if you know your Blueprint will need some sort of polygonal mesh to help define it, then you can add a **Static Mesh Component**. If your Blueprint object needs to emit light, you can add a **Point Light Component** or perhaps a **Spot Light Component**. 
+
+**Components** can be thought of as the objects that will be needed for a **Blueprint** object to fulfill its purpose. There are many types of Components that can be added to a Blueprint. For example, if you know your Blueprint will need some sort of polygonal mesh to help define it, then you can add a **Static Mesh Component**. If your Blueprint object needs to emit light, you can add a **Point Light Component** or perhaps a **Spot Light Component**. 
 
 In this example, we see a series of Blueprint Actors, each showing some of the available Components. 
 
@@ -50,7 +50,7 @@ From left to right, these Blueprint Actors are:
 * **BP_SeveralComponents** - This contains one of all of the above components, including a Billboard Component, a Static Mesh Component, and a Point Light Component.
 
 
-For each of these Blueprint Actors, their respective Components were each added within the Blueprint Editor while in Component mode. Component Mode is set using the button at the upper-right of the Blueprint Editor window, as shown below. To see various ways to add and remove components, please see the [Blueprint Components documentation.](Engine/Blueprints/UserGuide/Components)
+For each of these Blueprint Actors, their respective Components were each added within the Blueprint Editor in the **Component** tab. The **Component** tab is displayed at the top-left of the Blueprint Editor window by default, as shown below. To see various ways to add and remove components, please see the [Blueprint Components documentation.](Engine/Blueprints/UserGuide/Components)
 
 [REGION:fullwidth]
 ![](BPComponentMode.png)
