#UE4doc #docs:
Removed Overview header
Changed font settings for multiple terms

diff --git a/Engine/Documentation/Source/Resources/ContentExamples/Blueprints/1_4/1_4.INT.udn b/Engine/Documentation/Source/Resources/ContentExamples/Blueprints/1_4/1_4.INT.udn
--- a/Engine/Documentation/Source/Resources/ContentExamples/Blueprints/1_4/1_4.INT.udn
+++ b/Engine/Documentation/Source/Resources/ContentExamples/Blueprints/1_4/1_4.INT.udn
@@ -34,14 +34,14 @@ Related: Gameplay/HowTo/SetUpInput/Blueprints
 ![](ConstructionScript_BP.png)
 [/REGION]
 
-## Overview
 
-In this very simple Blueprint setup, we explore the nature of the Construction Script. The Construction Script is designed to execute upon creation of the Blueprint object, and again whenever that object is updated in some way. This means that whenever it is moved, rotated, scaled, or has any of its properties adjusted, the Construction Script fires again. This is useful for elements of a Blueprint Actor that need to be calculated prior to gameplay.
+
+In this very simple Blueprint setup, we explore the nature of the **Construction Script**. The Construction Script is designed to execute upon creation of the Blueprint object, and again whenever that object is updated in some way. This means that whenever it is moved, rotated, scaled, or has any of its properties adjusted, the Construction Script fires again. This is useful for elements of a Blueprint Actor that need to be calculated prior to gameplay.
 
 For this example, we show how we can add Components to a Blueprint Actor via the Construction Script, contrasting that against adding them manually in the Components List.
 
 [REGION:note]
-It is important to remember that once gameplay begins, the Construction script stops executing. At that point, everything that the Construction Script was doing is considered complete. Also, it should be noted that the Level Blueprint has no Construction Script.
+It is important to remember that once gameplay begins, the Construction Script stops executing. At that point, everything that the Construction Script was doing is considered complete. Also, it should be noted that the Level Blueprint has no Construction Script.
 [/REGION]
 
 
@@ -49,7 +49,7 @@ It is important to remember that once gameplay begins, the Construction script s
 
 You will notice in this example that both of the objects look exactly the same. However, one was assembled via the Components List and the other by the Construction Script. While in this case the example is primarily an academic exercise, it is important to know the benefits of each approach. 
 
-When you set up the Components of your Blueprint in the Components List, you are given the benefit of being able to use the Components Tab of the Blueprint Editor to adjust and place each Component while you work. This makes precision placement of Components very easy and intuitive. Generally, this type of setup is ideal when you know exactly what types of Components you will need in advance. 
+When you set up the Components of your Blueprint in the Components List, you are given the benefit of being able to use the **Components Tab** of the **Blueprint Editor** to adjust and place each Component while you work. This makes precision placement of Components very easy and intuitive. Generally, this type of setup is ideal when you know exactly what types of Components you will need in advance. 
 
 When you may need to change out the Components used in your Blueprints, it will be more useful to set those Components up in the Construction Script. In the example on the right, we have used the Construction Script to add each Component, and set up a system by which each type of Component can be shown or hidden by checking a Boolean value. 
 
