#UE4doc #docs:
Removed Overview header
Bolded a single term

diff --git a/Engine/Documentation/Source/Resources/ContentExamples/Blueprints/1_5/1_5.INT.udn b/Engine/Documentation/Source/Resources/ContentExamples/Blueprints/1_5/1_5.INT.udn
--- a/Engine/Documentation/Source/Resources/ContentExamples/Blueprints/1_5/1_5.INT.udn
+++ b/Engine/Documentation/Source/Resources/ContentExamples/Blueprints/1_5/1_5.INT.udn
@@ -34,9 +34,9 @@ Related: Gameplay/HowTo/SetUpInput/Blueprints
 [REGION:fullwidth]
 ![](1.5BlueprintOverview.png)
 
-## Overview
 
-In this example, we have a rocket mesh that has colored highlights. In the **Details** panel, we can see that there are two available colors we can set. One is set in the Construction Script and the other is set in the Event Graph. The example shows the difference between the two graphs and when they execute. If we change the property for the Construction Script, we see an immediate change in the editor. However, if we change the color for the Event Graph, there is no change until we play or simulate. 
+
+In this example, we have a rocket mesh that has colored highlights. In the **Details** panel, we can see that there are two available colors we can set. One is set in the Construction Script and the other is set in the **Event Graph**. The example shows the difference between the two graphs and when they execute. If we change the property for the Construction Script, we see an immediate change in the editor. However, if we change the color for the Event Graph, there is no change until we play or simulate. 
 
 ## Construction Script
 
