#UE4doc #docs:
Removed Overview header
Changed font settings for multiple terms

diff --git a/Engine/Documentation/Source/Resources/ContentExamples/Blueprints/1_7/1_7.INT.udn b/Engine/Documentation/Source/Resources/ContentExamples/Blueprints/1_7/1_7.INT.udn
--- a/Engine/Documentation/Source/Resources/ContentExamples/Blueprints/1_7/1_7.INT.udn
+++ b/Engine/Documentation/Source/Resources/ContentExamples/Blueprints/1_7/1_7.INT.udn
@@ -33,19 +33,19 @@ Related: Gameplay/HowTo/SetUpInput/Blueprints
 [REGION:fullwidth]
 ![](1.7Overview.png)
 
-## Overview
 
-In this example, we show several different types of flow control. Flow control nodes exist as a way to provide explicit control of the flow of execution in your Blueprint scripts. Many different types of flow control exist, such as:
 
-* Branch - allows you to execute one of a pair of scripts based on the result of a Boolean value.
-* Sequence - Allows any number of scripts to execute in a specific order.
-* Switch - Allows for execution paths to be selected based on the result of an enumeration.
+In this example, we show several different types of **Flow Control**. Flow Control nodes exist as a way to provide explicit control of the flow of execution in your Blueprint scripts. Many different types of flow control exist, such as:
+
+* **Branch** - allows you to execute one of a pair of scripts based on the result of a Boolean value.
+* **Sequence** - Allows any number of scripts to execute in a specific order.
+* **Switch** - Allows for execution paths to be selected based on the result of an enumeration.
 
 There are many other types available, and can be found under the Flow Control category when creating Blueprints.
 
 ## Construction Script
 
-The Construction Script for this example creates and sets a Dynamic Material Instance, and then uses a Branch flow control node to determine what color to apply to the dynamic Material. The Branch is powered by a Boolean variable named _Branch Condition_ which has been made public and can be edited in the **Details** panel.
+The Construction Script for this example creates and sets a Dynamic Material Instance, and then uses a Branch Flow Control node to determine what color to apply to the dynamic Material. The Branch is powered by a Boolean variable named Branch Condition which has been made public and can be edited in the **Details** panel.
 
 [REGION:lightbox]
 [![](1.7Construction.png)](1.7Construction.png)
@@ -56,19 +56,19 @@ The Construction Script for this example creates and sets a Dynamic Material Ins
 The Event Graph uses several different types of flow control. These include:
 
 * A Sequence, which executes all other operations in order. This is useful when you have a one-to-many relationship in which one event needs to execute many different things.
-* A Branch, which is controlled via the _Branch Condition_ variable.
+* A Branch, which is controlled via the Branch Condition variable.
 * A Switch on Enum, which controls the direction the coin moves (horizontally or vertically) based on the result of an enumeration.
 * A Do N, spawns a particle effect a specific number of times.
 
 ### Creating Enumerations
 
-Enumerations are lists of options in which each entry can be given a specific name. This is useful when you need to provide specific options that need more descriptive labels than just a number. For instance, you could create an enumeration for a character's stance that had the following options:
+**Enumerations** are lists of options in which each entry can be given a specific name. This is useful when you need to provide specific options that need more descriptive labels than just a number. For instance, you could create an enumeration for a character's stance that had the following options:
 
 * Standing
 * Crouched
 * Prone
 
-Enumerations are created in the **Content Browser** by **right-clicking** and choosing **Blueprints > Enumeration** from the Context Menu. They can then be edited within the Enum Editor by **double-clicking** them. Within, you can add the various entries that need to appear within your enumeration.
+Enumerations are created in the **Content Browser** by **right-clicking** and choosing **Blueprints > Enumeration** from the Context Menu. They can then be edited within the **Enum Editor** by **double-clicking** them. Within, you can add the various entries that need to appear within your enumeration.
 
 ![](EnumEditor.png)
 
