#UE4doc #docs:
Removed Overview header
Changed font settings for multiple terms

diff --git a/Engine/Documentation/Source/Resources/ContentExamples/Blueprints_Advanced/1_1/1_1.INT.udn b/Engine/Documentation/Source/Resources/ContentExamples/Blueprints_Advanced/1_1/1_1.INT.udn
--- a/Engine/Documentation/Source/Resources/ContentExamples/Blueprints_Advanced/1_1/1_1.INT.udn
+++ b/Engine/Documentation/Source/Resources/ContentExamples/Blueprints_Advanced/1_1/1_1.INT.udn
@@ -39,9 +39,9 @@ AdvBP_1_1_ConstructionScript.png
 ![](blueprint_advanced_1_1.png)
 
 
-## Overview
 
-The **BP_Random_Meshes** _Blueprint_ uses the Construction Script to randomly scatter _StaticMeshes_ within a radius, using exposed variables to allow users to customize the scale, rotation, and spacing of the _StaticMeshes_. All of the _StaticMeshes_ are actually components within the **BP_Random_Meshes** _Blueprint_.
+
+The **BP_Random_Meshes** Blueprint uses the Construction Script to randomly scatter Static Meshes within a radius, using exposed variables to allow users to customize the scale, rotation, and spacing of the Static Meshes. All of the Static Meshes are actually components within the BP_Random_Meshes Blueprint.
 
 
 
@@ -56,9 +56,9 @@ Click for full size or **right-click** and Save As.
 
 
 
-*	A **ForLoop** node is used to repeat a section of script a specified number of times. In this case, the **Number of meshes** variable is also the number of times the **ForLoop** will repeat, adding a new _StaticMeshComponent_ each time.
-*	**RandomFloatInRange** nodes are used to randomize the location, scale, and rotation of the _StaticMeshComponents_, as well as the _StaticMesh_ assets being used.
-*	An array is used to allow the user to specify a list of possible _StaticMeshes_ that will be set for the _StaticMeshComponents_ in the Blueprint. After each _StaticMeshComponent_ is added, a _StaticMesh_ asset is picked from the array using a random index number and assigned to the newly added  _StaticMeshComponent_.
+*	A **ForLoop** node is used to repeat a section of script a specified number of times. In this case, the Number of meshes variable is also the number of times the ForLoop will repeat, adding a new StaticMeshComponent each time.
+*	**RandomFloatInRange** nodes are used to randomize the location, scale, and rotation of the StaticMeshComponents, as well as the StaticMesh assets being used.
+*	An array is used to allow the user to specify a list of possible StaticMeshes that will be set for the StaticMeshComponents in the Blueprint. After each StaticMeshComponent is added, a StaticMesh asset is picked from the array using a random index number and assigned to the newly added  StaticMeshComponent.
 
 [/REGION]
 
