#UE4 DOCS: updating stat commands page

diff --git a/Engine/Documentation/Source/Engine/Performance/StatCommands/StatCommands.INT.udn b/Engine/Documentation/Source/Engine/Performance/StatCommands/StatCommands.INT.udn
--- a/Engine/Documentation/Source/Engine/Performance/StatCommands/StatCommands.INT.udn
+++ b/Engine/Documentation/Source/Engine/Performance/StatCommands/StatCommands.INT.udn
@@ -1,31 +1,44 @@
-Availability:Public
+Availability:Doc
 Title: Stat Commands
 Crumbs: %ROOT%, Engine
 Description:Console commands specific to displaying game statistics.
 Navigation:topic
 
-Here is a list of currently available stat commands that can be entered into the console to aid with profiling.
-
-| Command | Description |
-| -------- | ----------- |
-| **stat Unit** | Overall frame time as well as the game thread, rendering thread, and GPU times. |
-| **stat UnitGraph** | To see a graph with the stat unit data, use **stat Raw** to see the unfiltered data. |
-| **stat Hitches** | Set **t.HitchThreshold** to define the time in seconds that is considered a hitch. Will also dump all hitches to the log/visual studio debug e.g. `[0327.87] LogEngine:Warning: HITCH @ 00m:01s:643ms,1643,72,2.`  |
-| **stat InitViews** | Displays information on how long visibility culling took and how effective it was. Visible section count is the single most important stat with respect to rendering thread performance, and that is dominanted by Visible Static Mesh Elements under STAT INITVIEWS, but Visible Dynamic Primitives also factors in.  |
-| **stat Scenerendering** | Shows general rendering statistics. This is a good starting point to find general areas of slow performance in the rendering process. |
-| **stat D3d11rhi** **stat OpenGL** | Direct3D 11 or OpenGL RHI.  |
-| **stat Sceneupdate** | Displays information about updating the world, including the time taken to add, update, and remove lights as well as add and remove primitives in the scene.  |
-| **stat Engine** | Engine.  |
-| **stat Game** | Game. |
-| **stat Anim** | Animation. |
-| **stat Particles** | Particles. |
-| **stat ShadowRendering** | Shadow rendering. |
-| **stat LightRendering** | Light rendering. |
-| **stat Memory** | Memory. |
-| **stat Streaming** | Streaming. |
-| **stat StreamingDetails** | More details for Streaming. |
-| **stat Slow 0.012 10** | Show any stat that has been >12ms in the last 10 seconds (good starting point). |
-| **stat Grouped** | To disable stat Slow.|
+Here is a list of currently available stat commands that can be entered into the console to aid with profiling. Remember, if you want anything to dump to the log, you'll need to actually set the editor/your game to have a log file by running it with the -log command:
+
+	UE4Editor.exe -silent LOG=MyLog.txt
+
+This is necessary to get anything useful out of any of the "dump" commands.
+
+[REGION:raw]
+| Command | %spacer% | Description |
+| -------- | :---: | :----------- |
+| **stat Unit** || Overall frame time as well as the game thread, rendering thread, and GPU times. |
+| **stat UnitGraph** || To see a graph with the stat unit data, use **stat Raw** to see the unfiltered data. |
+| **stat Hitches** || Set **t.HitchThreshold** to define the time in seconds that is considered a hitch. Will also dump all hitches to the log/visual studio debug e.g. `[0327.87] LogEngine:Warning: HITCH @ 00m:01s:643ms,1643,72,2.`  |
+| **stat dumphitches** || Anytime a "hitch" is detected based on **t.HithThreshold** it will be written to the log. |
+| **stat InitViews** || Displays information on how long visibility culling took and how effective it was. Visible section count is the single most important stat with respect to rendering thread performance, and that is dominated by Visible Static Mesh Elements under STAT INITVIEWS, but Visible Dynamic Primitives also factors in.  |
+| **stat Scenerendering** || Shows general rendering statistics. This is a good starting point to find general areas of slow performance in the rendering process. |
+| **stat D3d11rhi** **stat OpenGL** || Direct3D 11 or OpenGL RHI.  |
+| **stat Sceneupdate** || Displays information about updating the world, including the time taken to add, update, and remove lights as well as add and remove primitives in the scene.  |
+| **stat dumpevents [-ms=0.1] [-all]** || Whenever an event is called (say from Blueprints, or for things like PhysX Task functions), they will be written to the log. |
+| **stat dumpframe** || This will write to the log information about the frame rendered right after the command is recieved. |
+| **stat Engine** || Shows general rendering stats like frame time as well as counters from the number of triangles being rendered.  |
+| **stat Game** || Gives feedback on how long the various Gameplay Ticks are taking. |
+| **stat Anim** || Shows how long skinned meshes are taking to compute per tick. |
+| **stat Particles** || Displays how long particle calculations are taking as well as sprite render time. |
+| **stat ShadowRendering** || Shows how long shadow calculations are taking, separate from actual shadow render time which is covered in **stat LightRendering**. |
+| **stat LightRendering** || Gives feedback on how long lighting and shadows are taking to render. |
+| **stat Memory** || Shows stats on how much memory is being used by various sub systems in Unreal Engine. |
+| **stat Streaming** || Displays various statistics on streaming assets, like how much memory streaming textures are using, or how many streaming textures there are in the scene. |
+| **stat StreamingDetails** || More detailed statistics on streaming, like breaking down general texture streaming into more specific groups (lightmaps, static textures, and dynamic textures). |
+| **stat Slow [-ms=0.3] [-maxdepth=5]** || Show any stat that has recently taken longer than 0.3 ms up to 5 levels deep (as things can be layered). The defaults are **-ms=1.0 -maxdepth=4**. Be careful with really small values in **-ms=** as it will start displaying everything and will drastically reduce performance (thus making more things show up on the list...). |
+| **stat Grouped** || Disables stat Slow.|
+[/REGION]
+
+[VAR:spacer]![](%ROOT%/spacer.png)(w:150 convert:false)[/VAR]
+
+[VAR:spacer2]![](%ROOT%/spacer.png)(w:40 convert:false)[/VAR]
 
 <!--
 Note: removed until this can be checked against profiler info. StatsViewer didn't appear to be in P4 or a binary build
