#UE4doc #docs:
Removed Overview header
Changed font settings for multiple terms

diff --git a/Engine/Documentation/Source/Resources/ContentExamples/Blueprints_Advanced/1_2/1_2.INT.udn b/Engine/Documentation/Source/Resources/ContentExamples/Blueprints_Advanced/1_2/1_2.INT.udn
--- a/Engine/Documentation/Source/Resources/ContentExamples/Blueprints_Advanced/1_2/1_2.INT.udn
+++ b/Engine/Documentation/Source/Resources/ContentExamples/Blueprints_Advanced/1_2/1_2.INT.udn
@@ -35,14 +35,14 @@ Related: Gameplay/HowTo/SetUpInput/Blueprints
 
 ![](blueprint_advanced_1_2.png)
 
-## Overview
 
-The **BP_Fence** _Blueprint_ uses the Construction Script to generate a row of _Static Mesh Components_ extending from a start point to an end point. The location of the root _Scene Component_ is treated as the starting point, while a vector variable represented as a _3D widget_ is treated as the end point. The script determines how many _StaticMeshComponents_ are needed to bridge the gap between the start and end points, and then adds them in a row, attaching each new _Static Mesh Component_ to the initial _Static Mesh Component_ at the start of the wall. Finally, the initial _Static Mesh Component_ is rotated to face the end point along its X axis, causing the rest of the row to rotate with it.
+
+The **BP_Fence** Blueprint uses the Construction Script to generate a row of Static Mesh Components extending from a start point to an end point. The location of the root Scene Component is treated as the starting point, while a vector variable represented as a 3D widget is treated as the end point. The script determines how many StaticMeshComponents are needed to bridge the gap between the start and end points, and then adds them in a row, attaching each new Static Mesh Component to the initial Static Mesh Component at the start of the wall. Finally, the initial Static Mesh Component is rotated to face the end point along its X axis, causing the rest of the row to rotate with it.
 
 
 
 [REGION:tip]
-To show a vector variable's 3D widget, simply check the **Show 3D Widget** option, available in the **Details** panel of the Blueprint Editor when that variable is selected.
+To show a vector variable's 3D widget, check the **Show 3D Widget** option, available in the **Details** panel of the **Blueprint Editor** when that variable is selected.
 [/REGION]
 
 ## Construction Script
@@ -54,7 +54,7 @@ To show a vector variable's 3D widget, simply check the **Show 3D Widget** optio
 Click for full size or **right-click** and Save As.
 [/REGION]
 
--	A **ForLoop** node is used to repeat a section of script. The **Index** output from the **ForLoop** node is multiplied by the variable **MeshWidth** to calculate the location offset for each added _Static Mesh Component_. 
--	The **End point** variable is a vector variable, which is **Editable** so users can type in exact X, Y, and Z values for the vector. **Show 3D Widget** is set to _true_ for this variable, so users can also select and drag the 3D widget around in the level to set the end point for the fence.
+-	A **ForLoop** node is used to repeat a section of script. The **Index** output from the ForLoop node is multiplied by the variable **MeshWidth** to calculate the location offset for each added Static Mesh Component. 
+-	The **End point** variable is a vector variable, which is **Editable** so users can type in exact X, Y, and Z values for the vector. **Show 3D Widget** is set to true for this variable, so users can also select and drag the 3D widget around in the level to set the end point for the fence.
 
 [/REGION]
