#UE4doc #docs:
Removed Overview header
Changed font settings for multiple terms

diff --git a/Engine/Documentation/Source/Resources/ContentExamples/Blueprints_Advanced/1_3/1_3.INT.udn b/Engine/Documentation/Source/Resources/ContentExamples/Blueprints_Advanced/1_3/1_3.INT.udn
--- a/Engine/Documentation/Source/Resources/ContentExamples/Blueprints_Advanced/1_3/1_3.INT.udn
+++ b/Engine/Documentation/Source/Resources/ContentExamples/Blueprints_Advanced/1_3/1_3.INT.udn
@@ -33,10 +33,10 @@ Related: Gameplay/HowTo/SetUpInput/Blueprints
 
 ![](blueprint_advanced_1_3.png)
 
-## Overview
 
-The **BP_Radial_Array** _Blueprint_ uses the Construction Script to generate an evenly spaced ring of _Static Mesh Components_. Several variables have been set to public here so that they can be edited in the **Details** tab for the _Actor_.  These include variables 
- specifying the number of _Static Mesh Components_ to add, the radius of the ring, the ability to add an additional rotation offset, and the ability to insert a gap into the ring.
+
+The **BP_Radial_Array** Blueprint uses the Construction Script to generate an evenly spaced ring of Static Mesh Components. Several variables have been set to public here so that they can be edited in the **Details** tab for the Actor.  These include variables 
+ specifying the number of Static Mesh Components to add, the radius of the ring, the ability to add an additional rotation offset, and the ability to insert a gap into the ring.
 
 
 
@@ -51,6 +51,6 @@ Click for full size or **right-click** and Save As.
 
 
 
-- The **Index** output of the **ForLoop** node is used to determine the location offset for each added _StaticMeshComponent_. The index number is combined with the **Number of meshes** and **Radius** variables to produce both a rotation value for each **StaticMeshComponent** and a location offset that spaces them out in a circle.
+- The **Index** output of the **ForLoop** node is used to determine the location offset for each added StaticMeshComponent. The index number is combined with the **Number of meshes** and **Radius** variables to produce both a rotation value for each StaticMeshComponent and a location offset that spaces them out in a circle.
 
 [/REGION]
