#UE4doc #docs:
Removed Overview header
Changed font settings for multiple terms

diff --git a/Engine/Documentation/Source/Resources/ContentExamples/Blueprints_Advanced/1_4/1_4.INT.udn b/Engine/Documentation/Source/Resources/ContentExamples/Blueprints_Advanced/1_4/1_4.INT.udn
--- a/Engine/Documentation/Source/Resources/ContentExamples/Blueprints_Advanced/1_4/1_4.INT.udn
+++ b/Engine/Documentation/Source/Resources/ContentExamples/Blueprints_Advanced/1_4/1_4.INT.udn
@@ -35,9 +35,8 @@ Related: Gameplay/HowTo/SetUpInput/Blueprints
 
 ![](blueprint_advanced_1_4.png)
 
-## Overview
 
-The **BP_Spotlight** _Blueprint_ is an example of a _Blueprint_ used to create a customizable Actor that can be re-used in multiple situations in the creation of a game. It combines a variety of components into one Actor that is easy to place, and uses the Construction Script to allow users to tweak a number of variables.
+The **BP_Spotlight** Blueprint is an example of a Blueprint used to create a customizable Actor that can be re-used in multiple situations in the creation of a game. It combines a variety of components into one Actor that is easy to place, and uses the Construction Script to allow users to tweak a number of variables.
 
 
 
@@ -47,7 +46,7 @@ The **BP_Spotlight** _Blueprint_ is an example of a _Blueprint_ used to create a
 ![](AdvBP_1_4_ComponentsList.png)
 
 
-The _Blueprint_ contains three Static Mesh Components: A base that attaches it to a wall, the spotlight itself, and a hinge that connects the two together. The spotlight is attached to a socket within the hinge, while the hinge is attached to a socket within the base. Sockets, which can be created and positioned in the [](Engine/Content\Types\StaticMeshes\Editor), are a quick and easy way of attaching Static Meshes together without the need to manually type in offsets and rotation values in a _Blueprint_.
+The Blueprint contains three Static Mesh Components: A base that attaches it to a wall, the spotlight itself, and a hinge that connects the two together. The spotlight is attached to a socket within the hinge, while the hinge is attached to a socket within the base. Sockets, which can be created and positioned in the [](Engine/Content\Types\StaticMeshes\Editor), are a quick and easy way of attaching Static Meshes together without the need to manually type in offsets and rotation values in a Blueprint.
 
 ## Construction Script
 
@@ -58,6 +57,6 @@ The _Blueprint_ contains three Static Mesh Components: A base that attaches it t
 Click for full size or **right-click** and Save As.
 [/REGION]
 
-The **Light color** variable not only changes the color of the **Spotlight** _SpotLightComponent_, but is also used to change vector parameters in Material Instances applied to the **Spotlight mesh** and **Light beam mesh** _StaticMeshComponents_. These Material Instances are created and applied to the _StaticMeshComponents_ in the ConstructionScript, so each instance of the Blueprint will have its own unique Dynamic Material Instances.
+The **Light color** variable not only changes the color of the **Spotlight** SpotLightComponent, but is also used to change vector parameters in Material Instances applied to the **Spotlight mesh** and **Light beam mesh** StaticMeshComponents. These Material Instances are created and applied to the StaticMeshComponents in the ConstructionScript, so each instance of the Blueprint will have its own unique Dynamic Material Instances.
 
 [/REGION]
