#UE4 DOCS: removed reference to Asking the Env Smart Question docs that is currently in question if we can publish it.

diff --git a/Engine/Documentation/Source/Gameplay/AI/AI.INT.udn b/Engine/Documentation/Source/Gameplay/AI/AI.INT.udn
--- a/Engine/Documentation/Source/Gameplay/AI/AI.INT.udn
+++ b/Engine/Documentation/Source/Gameplay/AI/AI.INT.udn
@@ -35,7 +35,7 @@ Description:Information over Artificial Intelligence including the use of Behavi
 	* [](Engine/AI/EnvironmentQuerySystem/UserGuide "%Engine/AI/EnvironmentQuerySystem/UserGuide:description%")
 	* [](Engine/AI/EnvironmentQuerySystem/EQSPawn "%Engine/AI/EnvironmentQuerySystem/EQSPawn:description%")
 	* [](Engine/AI/EnvironmentQuerySystem/NodeReference "%Engine/AI/EnvironmentQuerySystem/NodeReference:description%")
-	* [](Engine/AI/EnvironmentQuerySystem/AskingSmartQuestions "%Engine/AI/EnvironmentQuerySystem/AskingSmartQuestions:description%")
+	<!--* [](Engine/AI/EnvironmentQuerySystem/AskingSmartQuestions "%Engine/AI/EnvironmentQuerySystem/AskingSmartQuestions:description%")-->
 	[/PARAM]
 [/OBJECT]
 
