#UE4doc #docs:
Removed Overview header
Changed font settings for multiple terms

diff --git a/Engine/Documentation/Source/Resources/ContentExamples/Blueprints_Advanced/2_1/2_1.INT.udn b/Engine/Documentation/Source/Resources/ContentExamples/Blueprints_Advanced/2_1/2_1.INT.udn
--- a/Engine/Documentation/Source/Resources/ContentExamples/Blueprints_Advanced/2_1/2_1.INT.udn
+++ b/Engine/Documentation/Source/Resources/ContentExamples/Blueprints_Advanced/2_1/2_1.INT.udn
@@ -33,9 +33,9 @@ Related: Gameplay/HowTo/SetUpInput/Blueprints
 
 ![](blueprint_advanced_2_1.png)
 
-## Overview
 
-The **BP_Timeline_Ball** Blueprint demonstrates a basic [Timeline](Engine/Blueprints/UserGuide/Timelines) in the Event Graph, which automatically plays and loops, triggering an event at a certain point along the Timeline that spawns a particle effect.
+
+The **BP_Timeline_Ball** Blueprint demonstrates a basic [**Timeline**](Engine/Blueprints/UserGuide/Timelines) in the Event Graph, which automatically plays and loops, triggering an event at a certain point along the Timeline that spawns a particle effect.
 
 
 
