#UE4doc #docs:
Removed Overview header
Changed font settings for multiple terms

diff --git a/Engine/Documentation/Source/Resources/ContentExamples/Blueprints_Advanced/2_2/2_2.INT.udn b/Engine/Documentation/Source/Resources/ContentExamples/Blueprints_Advanced/2_2/2_2.INT.udn
--- a/Engine/Documentation/Source/Resources/ContentExamples/Blueprints_Advanced/2_2/2_2.INT.udn
+++ b/Engine/Documentation/Source/Resources/ContentExamples/Blueprints_Advanced/2_2/2_2.INT.udn
@@ -34,9 +34,9 @@ Related: Gameplay/HowTo/SetUpInput/Blueprints
 
 ![](blueprint_advanced_2_2.png)
 
-## Overview
 
-The **BP_Spotlight_Follow** _Blueprint_ contains script in the EventGraph that makes it detect how far away the player's Pawn is and turn to face the Pawn when it comes within a specified **Detection Range**. The **Hinge mesh** component controls the spotlight's yaw, while the **Spotlight mesh** component only changes its pitch, which is clamped to prevent it from clipping through the hinge.
+
+The **BP_Spotlight_Follow** Blueprint contains script in the EventGraph that makes it detect how far away the player's Pawn is and turn to face the Pawn when it comes within a specified **Detection Range**. The **Hinge mesh** component controls the spotlight's yaw, while the **Spotlight mesh** component only changes its pitch, which is clamped to prevent it from clipping through the hinge.
 
 
 ## Event Graph
@@ -48,9 +48,9 @@ The **BP_Spotlight_Follow** _Blueprint_ contains script in the EventGraph that m
 Click for full size or **right-click** and Save As.
 [/REGION]
 
-The **Find Look at Rotation** node generates a rotation value when given a start location and an end location. The resulting rotation value can then be applied to a component located at the start location to make it face the end location. In this _Blueprint_, the start location is the spotlight's location, while the end location is the player Pawn's location. The Rotator value is broken apart, and the resulting **Pitch** and **Yaw** values are applied to the spotlight and hinge _Static Mesh Components_ respectively.
+The **Find Look at Rotation** node generates a rotation value when given a start location and an end location. The resulting rotation value can then be applied to a component located at the start location to make it face the end location. In this Blueprint, the start location is the spotlight's location, while the end location is the player Pawn's location. The Rotator value is broken apart, and the resulting **Pitch** and **Yaw** values are applied to the spotlight and hinge Static Mesh Components respectively.
 
 Interpolation nodes, such as **RInterp** (for interpolating between two rotation values), are intended to feed values into nodes that are executed every frame, such as those connected to the **ReceiveTick** event or the **Update** pin on a timeline. Essentially, they blend smoothly from the **Current Value** to the **Target Value** based on the specified **Interp Speed**. For these nodes to work, a **Delta Time** value must be provided.
-In this _Blueprint_, the **Delta seconds** variable is set from the **Event Tick** output. In this case, **RInterp** nodes are connected to **Set World Rotation** nodes to smoothly rotate the **Spotlight mesh** and **Hinge mesh** components, with a low **Interp Speed** so that the spotlight lags behind its target slightly.
+In this Blueprint, the **Delta seconds** variable is set from the **Event Tick** output. In this case, **RInterp** nodes are connected to **Set World Rotation** nodes to smoothly rotate the **Spotlight mesh** and **Hinge mesh** components, with a low **Interp Speed** so that the spotlight lags behind its target slightly.
 
 [/REGION]
