#UE4doc #docs Updating 'Downloading' and 'Building' Unreal Engine from Source documentation.

diff --git a/Engine/Documentation/Source/GettingStarted/DownloadingUnrealEngine/DownloadZip/DownloadZip.INT.udn b/Engine/Documentation/Source/GettingStarted/DownloadingUnrealEngine/DownloadZip/DownloadZip.INT.udn
--- a/Engine/Documentation/Source/GettingStarted/DownloadingUnrealEngine/DownloadZip/DownloadZip.INT.udn
+++ b/Engine/Documentation/Source/GettingStarted/DownloadingUnrealEngine/DownloadZip/DownloadZip.INT.udn
@@ -3,6 +3,8 @@ Title:Downloading Unreal as a Zip from GitHub
 Crumbs:%ROOT%, GettingStarted, GettingStarted/DownloadingUnrealEngine
 Description:Downloading the Unreal Engine source code as a Zip file.
 
+If you prefer to download our source code as a Zip file from GitHub, this section is for you.
+
 [OBJECT:ToggleButtonContent]
 	[PARAMLITERAL:category]
 	OS
@@ -34,10 +36,10 @@ Description:Downloading the Unreal Engine source code as a Zip file.
 		[/REGION]
 	1.	**Open** your source folder in Explorer and **run Setup.bat**.
 		[REGION:note]
-		Our Setup batch file downloads binary content for the engine, installs prerequesites, and sets up Unreal file associations. Downloading engine binaries may take some time because they are currently take up three to four gigabytes of memory. Subsequent checkouts take less time to download.
+		Our Setup batch file downloads binary content for the engine, installs prerequesites, and sets up Unreal file associations. Downloading engine binaries may take some time because they currently take up three to four gigabytes of memory. Subsequent checkouts take less time to download.
 		[/REGION]
 		[REGION:warning]
-		A warning from SmartScreen may appear in Windows 8. **Click "More Info"**, then **"Run Anyway"** to continue running our Setup batch file.
+		A warning from SmartScreen may appear in Windows 8. **Click 'More Info'**, then **'Run Anyway'** to continue running our Setup batch file.
 		[/REGION]
 	[/PARAM]
 [/OBJECT]
@@ -76,6 +78,11 @@ Description:Downloading the Unreal Engine source code as a Zip file.
 	active_button_content
 	[/PARAMLITERAL]
 	[PARAM:content]
+
+	[REGION:warning]
+	At this time, we only provide development support for Ubuntu 15.04.
+	[/REGION]
+
 	1.	**Navigate** to our GitHub page.
 	1.	**Click** the **Download ZIP** button (pictured below).
 		![](DownloadZipButton.png)
@@ -85,7 +92,7 @@ Description:Downloading the Unreal Engine source code as a Zip file.
 		![](RunSetupShellScript_Linux.png)
 
 		[REGION:note]
-		If you don't see a ** \*\*SUCCESS\*\* ** message after running this step, refer to the **BuildThirdParty.log** file for more information. BuildThirdParty.log is located in the `Engine/Build/BatchFiles/Linux` directory.
+		If you don't see a ** \*\*SUCCESS\*\* ** message after running this step, refer to the **'BuildThirdParty.log'** file for more information. 'BuildThirdParty.log' is located in the `[UE4Root]/Engine/Build/BatchFiles/Linux` directory.
 		[/REGION]
 	[/PARAM]
 [/OBJECT]
diff --git a/Engine/Documentation/Source/GettingStarted/DownloadingUnrealEngine/DownloadingUnrealEngine4.INT.udn b/Engine/Documentation/Source/GettingStarted/DownloadingUnrealEngine/DownloadingUnrealEngine4.INT.udn
--- a/Engine/Documentation/Source/GettingStarted/DownloadingUnrealEngine/DownloadingUnrealEngine4.INT.udn
+++ b/Engine/Documentation/Source/GettingStarted/DownloadingUnrealEngine/DownloadingUnrealEngine4.INT.udn
@@ -1,7 +1,11 @@
 Availability:Public
 Title:Downloading Unreal Engine Source Code
-Crumbs: %ROOT%, Programming, Programming/Development
 Description:Steps for connecting to the source code repository and downloading the latest build of Unreal Engine.
+Crumbs: %ROOT%, Programming, Programming/Development
+Prereq: GettingStarted/RecommendedSpecifications
+Prereq: Support/Builds
+Related: Programming/Development/BuildingUnrealEngine
+Version: 4.8
 
 [TOC (start:2 end:3)]
 
@@ -12,9 +16,13 @@ Unreal Engine 4 is available to be downloaded by licensees directly from Epic's
 [PUBLISH:Rocket]
 
 [REGION:warning]
-This page shows subscribers how to download and build Unreal Engine from our source code repository on GitHub. If you would like to download the binary version of Unreal Engine, go to our main page to [get unreal](https://www.unrealengine.com/) (click the image pictured below).
+This page shows subscribers how to download and build Unreal Engine from our source code repository on GitHub. If you'd like to download the binary version of Unreal Engine, go to our main page to [Get Unreal](https://www.unrealengine.com/).
 [/REGION]
 
+----
+
+![](WindowsMacLinux.png)
+
 Source code for Unreal Engine 4 (UE4) can be downloaded by our subscribers from the Unreal Engine GitHub repository. With access to our source code, subscribers can customize the engine and tools to meet their needs. Our subscribers are building the editor in Windows, Mac, Linux, and compiling games for both Android and iOS. Source code for all of our tools (such as Unreal Lightmass and Frontend) is included as well.
 
 [REGION:note]
@@ -27,24 +35,30 @@ Before you can access the [repository](https://github.com/EpicGames/UnrealEngine
 
 Once you've accessed the repository, you'll notice that we've published UE4's source code in three rolling branches, namely:
 
-*	**Release Branch**
-*	**Promoted Branch**
-*	**Master Branch**
+*	[Release Branch](#Bookmark1)
+*	[Promoted Branch](#Bookmark2)
+*	[Master Branch](#Bookmark3)
 
 ### Release Branch
 
+(#Bookmark1)
+
 ![](ReleaseBranch.png)
 
 The [release branch](https://github.com/EpicGames/UnrealEngine/tree/release) (pictured above) is extensively tested by our QA team and makes a great starting point for learning UE4 or making Unreal games. We work hard to stabilize this branch with the goal of releasing new builds every one to two months.
 
 ### Promoted Branch
 
+(#Bookmark2)
+
 ![](PromotedBranch.png)
 
 The [promoted branch](https://github.com/EpicGames/UnrealEngine/tree/promoted) (pictured above) is updated with builds for use by our designers and artists. We try to update this branch daily so that our subscribers have access to UE4's latest stable features. This branch is a great place to learn how feature development works in UE4.
 
 ### Master Branch
 
+(#Bookmark3)
+
 [REGION:warning]
 Because the master branch contains cutting-edge technology, UE4 isn't guaranteed to compile from source.
 [/REGION]
@@ -71,15 +85,19 @@ Other branches may appear from time-to-time as we work to stabilize new releases
 
 Subscribers can download Unreal Engine's source code by:
 
-*	%GettingStarted/DownloadingUnrealEngine/GitHub:title%
-*	%GettingStarted/DownloadingUnrealEngine/DownloadZip:title%
+*	[%GettingStarted/DownloadingUnrealEngine/GitHub:title%](#Bookmark4)
+*	[%GettingStarted/DownloadingUnrealEngine/DownloadZip:title%](#Bookmark5)
 
 ### %GettingStarted/DownloadingUnrealEngine/GitHub:title%
 
+(#Bookmark4)
+
 [INCLUDE:GettingStarted/DownloadingUnrealEngine/GitHub]
 
 ### %GettingStarted/DownloadingUnrealEngine/DownloadZip:title%
 
+(#Bookmark5)
+
 [INCLUDE:GettingStarted/DownloadingUnrealEngine/DownloadZip]
 
 [/PUBLISH:Rocket]
diff --git a/Engine/Documentation/Source/GettingStarted/DownloadingUnrealEngine/GitHub/GitHub.INT.udn b/Engine/Documentation/Source/GettingStarted/DownloadingUnrealEngine/GitHub/GitHub.INT.udn
--- a/Engine/Documentation/Source/GettingStarted/DownloadingUnrealEngine/GitHub/GitHub.INT.udn
+++ b/Engine/Documentation/Source/GettingStarted/DownloadingUnrealEngine/GitHub/GitHub.INT.udn
@@ -3,6 +3,8 @@ Title:Forking and Cloning Unreal from GitHub
 Crumbs:%ROOT%, GettingStarted, GettingStarted/DownloadingUnrealEngine
 Description:Forking and cloning the Unreal Engine source code from GitHub.
 
+Forking and cloning our source code from GitHub is a great way to get started with Unreal.
+
 %Globals:osselect_linux%
 
 [OBJECT:ToggleButtonContent]
@@ -29,10 +31,10 @@ Description:Forking and cloning the Unreal Engine source code from GitHub.
 		[/REGION]
 	1.	**Open** your source folder in Explorer and **run Setup.bat**.
 		[REGION:note]
-		Our Setup batch file downloads binary content for the engine, installs prerequesites, and sets up Unreal file associations. Downloading engine binaries may take some time because they are currently take up three to four gigabytes of memory. Subsequent checkouts take less time to download.
+		Our Setup batch file downloads binary content for the engine, installs prerequesites, and sets up Unreal file associations. Downloading engine binaries may take some time because they currently take up three to four gigabytes of memory. Subsequent checkouts take less time to download.
 		[/REGION]
 		[REGION:warning]
-		A warning from SmartScreen may appear in Windows 8. **Click "More Info"**, then **"Run Anyway"** to continue running our Setup batch file.
+		A warning from SmartScreen may appear in Windows 8. **Click 'More Info'**, then **'Run Anyway'** to continue running our Setup batch file.
 		[/REGION]
 	[/PARAM]
 [/OBJECT]
@@ -74,6 +76,11 @@ Description:Forking and cloning the Unreal Engine source code from GitHub.
 	active_button_content
 	[/PARAMLITERAL]
 	[PARAM:content]
+
+	[REGION:warning]
+	At this time, we only provide development support for Ubuntu 15.04.
+	[/REGION]
+
 	1.	**Set up** [Git](https://help.github.com/articles/set-up-git/).
 	1.	**Fork** and **Clone** our [repository](https://help.github.com/articles/fork-a-repo/) (pictured below).
 		![](UnrealRepo.png)
@@ -82,7 +89,7 @@ Description:Forking and cloning the Unreal Engine source code from GitHub.
 		![](RunSetupShellScript_Linux.png)
 
 		[REGION:note]
-		If you don't see a ** \*\*SUCCESS\*\* ** message after running this step, refer to the **BuildThirdParty.log** file for more information. BuildThirdParty.log is located in the `Engine/Build/BatchFiles/Linux` directory.
+		If you don't see a ** \*\*SUCCESS\*\* ** message after running this step, refer to the **'BuildThirdParty.log'** file for more information. 'BuildThirdParty.log' is located in the `[UE4Root]/Engine/Build/BatchFiles/Linux` directory.
 		[/REGION]
 	[/PARAM]
 [/OBJECT]
diff --git a/Engine/Documentation/Source/GettingStarted/Installation/InstallingUnrealEngine.INT.udn b/Engine/Documentation/Source/GettingStarted/Installation/InstallingUnrealEngine.INT.udn
--- a/Engine/Documentation/Source/GettingStarted/Installation/InstallingUnrealEngine.INT.udn
+++ b/Engine/Documentation/Source/GettingStarted/Installation/InstallingUnrealEngine.INT.udn
@@ -12,7 +12,7 @@ You can download the installer from the [UnrealEngine.com](http://unrealengine.c
 1. Go to [UnrealEngine.com](http://unrealengine.com) and sign in. 
 
 	[REGION:note]
-	You need to create an account and subscribe in order to download Unreal Engine. If you do not already have an account, click the orange **Get Unreal** button at the top of the page and follow the instructions.
+	You need to create an account and subscribe in order to download Unreal Engine. If you do not already have an account, click the **Get Unreal** button at the top of the page and follow the instructions.
 	[/REGION]
 
 1. From your account page, click one of the download buttons, depending on your operating system:
diff --git a/Engine/Documentation/Source/Programming/Development/BuildingUnrealEngine/BuildingUnrealEngine.INT.udn b/Engine/Documentation/Source/Programming/Development/BuildingUnrealEngine/BuildingUnrealEngine.INT.udn
--- a/Engine/Documentation/Source/Programming/Development/BuildingUnrealEngine/BuildingUnrealEngine.INT.udn
+++ b/Engine/Documentation/Source/Programming/Development/BuildingUnrealEngine/BuildingUnrealEngine.INT.udn
@@ -1,8 +1,11 @@
 Availability:Public
 Title:Building Unreal Engine from Source
+Description:Compiling Unreal Engine from source.
 Crumbs: %ROOT%, Programming, Programming/Development
-Description:Compiling the Unreal Engine.
-Version: 4.9
+Prereq: GettingStarted/DownloadingUnrealEngine
+Prereq: GettingStarted/RecommendedSpecifications
+Related: GettingStarted/DownloadingUnrealEngine
+Version: 4.8
 
 %Globals:osselect_linux%
 
@@ -23,13 +26,13 @@ Version: 4.9
 	[/PARAMLITERAL]
 	[PARAM:content]
 		[REGION:note]
-		Make sure that Microsoft Visual Studio 2013 is installed prior to building Unreal Engine from Source.
+		Make sure that Microsoft Visual Studio 2013 is installed prior to building Unreal Engine from source.
 		[/REGION]
 
 		1.	Inside the root directory, **run GenerateProjectFiles.bat** to set-up your project files.
 
 			[REGION:warning]
-			All project files are intermediate (`[UE4Root]\Engine\Intermediate\ProjectFiles`). You must generate project files each time you sync a new build to ensure they are up to date If you delete your `Intermediate` folder, you must regenerate project files using the above-mentioned batch file.
+			All project files are intermediate (`[UE4Root]\Engine\Intermediate\ProjectFiles`). You must generate project files each time you sync a new build to ensure they are up to date. If you delete your `Intermediate` folder, you must regenerate project files using the 'GenerateProjectFiles' batch file.
 			[/REGION]
 
 		1.	**Load** the project into Visual Studio by **double-clicking UE4.sln**.
@@ -38,7 +41,7 @@ Version: 4.9
 
 			![](configMenu_updated.png)
 
-		1.	If you're developing on the 64-bit version of the Windows operating system, set your solution platform to **Win64** (pictured below).
+		1.	Depending on your target-version of the Windows operating system (32-bit versus 64-bit), set your solution platform to either **Win32** or **Win64** (pictured below).
 
 			![](platform_menu_updated.png)
 
@@ -47,7 +50,7 @@ Version: 4.9
 			![](BuildUE4.png)
 
 			[REGION:note]
-			Depending on your system's specifications, it may take anywhere between 10 and 40 minutes to finish compiling the engine from source.
+			Depending on your system's specifications, it may take between 10 and 40 minutes to compile the engine from source.
 			[/REGION]
 
 		1.	After compiling the engine, set your startup project to **UE4** (pictured below).
@@ -94,7 +97,7 @@ Version: 4.9
 			![](BuildingMacEditor.png)
 
 			[REGION:note]
-			Depending on your system's specifications, it may take anywhere between 15 and 40 minutes to finish compiling the engine from source.
+			Depending on your system's specifications, it may take between 15 and 40 minutes to compile the engine from source.
 			[/REGION]
 		
 		1.	After compiling the engine, select **Product > Run** to launch the editor (pictured below).
@@ -117,7 +120,7 @@ Version: 4.9
 	[/PARAMLITERAL]
 	[PARAM:content]
 		[REGION:warning]
-		At this time, we only provide support for Ubuntu 15.04.
+		At this time, we only provide development support for Ubuntu 15.04.
 		[/REGION]
 		
 		1.	Inside the root directory, **run GenerateProjectFiles.sh** from the terminal to set-up your project files (example pictured below). 
@@ -128,18 +131,26 @@ Version: 4.9
 
 			![](RunMakeFile_Linux.png)
 
-		1.	After compiling the engine, navigate to the editor's binary path by entering **cd `Engine/Binaries/Linux/`** into the terminal.
+		1.	After compiling the engine, navigate to the editor's binary path by entering ** cd Engine/Binaries/Linux/ ** into the terminal.
 
 		1.	Run **UE4Editor** to launch the editor (pictured below).
 
 			![](RunUE4Editor_Linux.png)
 
+			[REGION:warning]
+			If you're unable to create an OpenGL context, you may need to [update](http://www.nvidia.com/Download/index.aspx?lang=en-us) your Nvidia graphics driver. Refer to [Nvidia's website](http://www.nvidia.com/page/home.html) for more information.
+			[/REGION]
+
 		1.	Congratulations! You've compiled and launched the engine from source (result pictured below).
 
 			![](CompiledEditor_Linux.png)
 	[/PARAM]
 [/OBJECT]
 
+## Previous Steps
+
+If you want a refresher on our GitHub download process, refer to the [](GettingStarted/DownloadingUnrealEngine) documentation.
+
 [/EXCERPT:Main]
 
 <!--Currently deprecated for the 'Building from Source' version
