added close bracket to make table of contents functional

diff --git a/Engine/Documentation/Source/Engine/Rendering/LightingAndShadows/Lightmass/Basics/Lightmass.INT.udn b/Engine/Documentation/Source/Engine/Rendering/LightingAndShadows/Lightmass/Basics/Lightmass.INT.udn
--- a/Engine/Documentation/Source/Engine/Rendering/LightingAndShadows/Lightmass/Basics/Lightmass.INT.udn
+++ b/Engine/Documentation/Source/Engine/Rendering/LightingAndShadows/Lightmass/Basics/Lightmass.INT.udn
@@ -46,7 +46,7 @@ Description:High level introduction to Lightmass.
 ![](Engine/Rendering/LightingAndShadows/lighting_overview_banner_001.png)(convert:false)
 [/REGION]
 
-[TOC(start:2)
+[TOC(start:2)]
 
 [INCLUDE:Engine/Rendering/LightingAndShadows/Lightmass#LightMass_002]
 
