TypoCorrection

diff --git a/Engine/Documentation/Source/Programming/Development/BuildingUnrealEngine/BuildingUnrealEngine.INT.udn b/Engine/Documentation/Source/Programming/Development/BuildingUnrealEngine/BuildingUnrealEngine.INT.udn
--- a/Engine/Documentation/Source/Programming/Development/BuildingUnrealEngine/BuildingUnrealEngine.INT.udn
+++ b/Engine/Documentation/Source/Programming/Development/BuildingUnrealEngine/BuildingUnrealEngine.INT.udn
@@ -1,7 +1,7 @@
-Availability:Public
+﻿Availability:Public
 Title:Building Unreal Engine from Source
-Description:Compiling Unreal Engine from source.
 Crumbs: %ROOT%, Programming, Programming/Development
+Description:Compiling Unreal Engine from source.
 Prereq: GettingStarted/DownloadingUnrealEngine
 Prereq: GettingStarted/RecommendedSpecifications
 Related: GettingStarted/DownloadingUnrealEngine
