#UE4doc #docs - adjusted bold lettering.

diff --git a/Engine/Documentation/Source/Engine/Basics/Projects/Browser/ProjectBrowser.INT.udn b/Engine/Documentation/Source/Engine/Basics/Projects/Browser/ProjectBrowser.INT.udn
--- a/Engine/Documentation/Source/Engine/Basics/Projects/Browser/ProjectBrowser.INT.udn
+++ b/Engine/Documentation/Source/Engine/Basics/Projects/Browser/ProjectBrowser.INT.udn
@@ -33,12 +33,12 @@ The **Project Browser** provides a jumping off point from which you can create p
 
 
 [REGION:tip]
-To bypass the **Project Browser** and instantly open the editor for a specific game, you can pass the name of the game as a command line argument to the executable. 
+To bypass the Project Browser and instantly open the editor for a specific game, you can pass the name of the game as a command line argument to the executable. 
 [/REGION]
 
 ## Opening Projects
 
-The first tab in the **Project Browser** is the **Projects** tab. 
+The first tab in the Project Browser is the **Projects** tab. 
 
 ![](OpenProject_single.png)(w:800)
 
@@ -57,19 +57,19 @@ You can also click on the **Browse...** button to navigate your computer and sel
 [REGION:tip]
 In the **Projects** tab, you can check the **Always load last project on startup** checkbox to automatically open your most recent project and skip the **Project Browser**. To change
 this setting in-editor: 
-1. Open the [Editor Preferences](Engine/UI/EditorSettings) window.
-1. Navigate to **General > Loading & Saving**.
+1. Open the [**Editor Preferences**](Engine/UI/EditorSettings) window.
+1. Navigate to **General** > **Loading & Saving**.
 1. Under the **Startup** heading, toggle **Load the Most Recently Loaded Project at Startup**.
 [/REGION]
 
-This tab also opens when you select **File > Open Project** from within Unreal Editor.
+This tab also opens when you select **File** > **Open Project** from within Unreal Editor.
 
 ![](open_project_menu.png)
 
 ## Creating New Projects
 
 [EXCERPT:creatingprojects]
-The second tab in the **Project Browser** is the **New Project** tab. 
+The second tab in the Project Browser is the **New Project** tab. 
 
 ![](NewProject_FPSExample.png)(w:800)
 
@@ -90,7 +90,7 @@ _Blueprints Only_ projects can be extended with C++ code, it is simple to add Bl
 
 **To create a new project:**
 1. Select a template from the list.
-1. Select whether or not to copy starter content into your project.
+1. Select whether or not to copy Starter Content into your project.
 1. Enter a name for your project.
 1. Click **Create**.
 
@@ -100,7 +100,7 @@ To save your project outside of the default save location, you can click on the
 
  Your new project will open in the editor if it is based on the Blank template or a Blueprints template, and in Visual Studio or XCode if it is based on a C++ template.
 
-This tab also opens when you select **File > New Project** from within Unreal Editor.
+This tab also opens when you select **File** > **New Project** from within Unreal Editor.
 
 ![](newprojectmenu.png)
 
@@ -116,8 +116,8 @@ This tab also opens when you select **File > New Project** from within Unreal Ed
 | ![](Mobile_Option.png) | Use this option if the game you are making is going to target Mobile devices. |
 | ![](Maximum_Quality_Option.png)| Use this option if the game you are making is going to target high end devices like Desktops or Consoles. |
 | ![](Scalable_Quality_Option.png)| Use this option if the game you are making is going to target Mobile devices. |
-|![](With_Starter_Content_Option.png) | Use this option if you want to include the [Starter Content](Engine/Content/Packs) with your game. |
-|![](Without_Starter_Content_Option.png) | Use this option if you do not want to include the [Starter Content](Engine/Content/Packs) with your game. |
+|![](With_Starter_Content_Option.png) | Use this option if you want to include the [**Starter Content**](Engine/Content/Packs) with your game. |
+|![](Without_Starter_Content_Option.png) | Use this option if you do not want to include the [**Starter Content**](Engine/Content/Packs) with your game. |
 
 [region:note]
 Please note that the Starter Content is not optimized for mobile devices and could cause performance issues if used with mobile devices. 
