#UE4doc #docs - bolded link in text.

diff --git a/Engine/Documentation/Source/Engine/Blueprints/UserGuide/Types/ClassBlueprint/ClassBlueprint.INT.udn b/Engine/Documentation/Source/Engine/Blueprints/UserGuide/Types/ClassBlueprint/ClassBlueprint.INT.udn
--- a/Engine/Documentation/Source/Engine/Blueprints/UserGuide/Types/ClassBlueprint/ClassBlueprint.INT.udn
+++ b/Engine/Documentation/Source/Engine/Blueprints/UserGuide/Types/ClassBlueprint/ClassBlueprint.INT.udn
@@ -70,7 +70,7 @@ While the above are the most common, all existing classes can be used as the Par
 For example, say you created an **Actor Blueprint** called _Animals_ and in it provided some script that all animals share such as _Hunger_, _Thirst_, _Energy_, or whatever script you wanted. Then you created another Blueprint called _Dogs_ and selected your _Animals_ Blueprint Class as the Parent Class; you can then provide specific functionality that applies to only dogs within the _Dogs_ Blueprint such as _Play Fetch_ or _Roll Over_ while inheriting the functionality that all animals share from the _Animals_ Blueprint. 
 
 [REGION:note]
-Depending upon the method used to [create a Blueprint Class](Engine\Blueprints\UserGuide\Types\ClassBlueprint\Creation), a Parent Class may already be assigned upon creation. 
+Depending upon the method used to [**create a Blueprint Class**](Engine\Blueprints\UserGuide\Types\ClassBlueprint\Creation), a Parent Class may already be assigned upon creation. 
 [/REGION]
 
 ## Working with Blueprint Classes
