#UE4doc #docs - adjusted bold lettering.

diff --git a/Engine/Documentation/Source/Engine/Commandlets/Commandlets.INT.udn b/Engine/Documentation/Source/Engine/Commandlets/Commandlets.INT.udn
--- a/Engine/Documentation/Source/Engine/Commandlets/Commandlets.INT.udn
+++ b/Engine/Documentation/Source/Engine/Commandlets/Commandlets.INT.udn
@@ -8,7 +8,7 @@ Description:Command-line programs that run within the Unreal Engine environment,
 ![commandlets.jpg](commandlets.jpg)
 [/REGION]
 
-Commandlets are command line programs that run inside of the Unreal Engine environment. They are most often used to make bulk changes to content, iterate over content to get information about it, or as a unit testing mechanism. The engine contains many commandlets by default and new commandlets can be added to perform additional functions.
+**Commandlets** are command line programs that run inside of the Unreal Engine environment. They are most often used to make bulk changes to content, iterate over content to get information about it, or as a unit testing mechanism. The engine contains many commandlets by default and new commandlets can be added to perform additional functions.
 
 Commandlets are executed by passing the name of the commandlet along with any necessary parameters to the game on the command line, either by running the game from a command prompt in Windows or through the **target** of a shortcut to the executable (i.e. UE4Editor.exe).
 
