#UE4doc #docs - Removed overview header

diff --git a/Engine/Documentation/Source/Engine/Commandlets/BatchExport/BatchExport.INT.udn b/Engine/Documentation/Source/Engine/Commandlets/BatchExport/BatchExport.INT.udn
--- a/Engine/Documentation/Source/Engine/Commandlets/BatchExport/BatchExport.INT.udn
+++ b/Engine/Documentation/Source/Engine/Commandlets/BatchExport/BatchExport.INT.udn
@@ -9,9 +9,8 @@ Description:Commandlet for exporting resource from packages to various file type
 
 [TOC(start:2)]
 
-## Overview
 
-The batchexport commandlet is used for exporting resources from Unreal package files to various external file types.
+The **Batchexport** commandlet is used for exporting resources from Unreal package files to various external file types.
 
 
 
