#UE4doc #docs:
Removed Overview header
Changed font settings for multiple terms

diff --git a/Engine/Documentation/Source/Resources/ContentExamples/Blueprints_Advanced/2_3/2_3.INT.udn b/Engine/Documentation/Source/Resources/ContentExamples/Blueprints_Advanced/2_3/2_3.INT.udn
--- a/Engine/Documentation/Source/Resources/ContentExamples/Blueprints_Advanced/2_3/2_3.INT.udn
+++ b/Engine/Documentation/Source/Resources/ContentExamples/Blueprints_Advanced/2_3/2_3.INT.udn
@@ -35,13 +35,13 @@ Related: Gameplay/HowTo/SetUpInput/Blueprints
 
 ![](blueprint_advanced_2_3.png)
 
-## Overview
 
-[Blueprint Interfaces](Engine/Blueprints\UserGuide\Types\Interface) allow one _Blueprint_ to trigger an event or function in another _Blueprint_. Both of these _Blueprints_ must implement the _Blueprint Interface_ in their **Blueprint Properties**.
-In this example, the **BP_Button** _Blueprint_ causes the gears in the **BP_Gears** _Blueprint_ to start or stop spinning.
 
+[**Blueprint Interfaces**](Engine/Blueprints\UserGuide\Types\Interface) allow one Blueprint to trigger an event or function in another Blueprint. Both of these Blueprints must implement the Blueprint Interface in their **Blueprint Properties**.
+In this example, the **BP_Button** Blueprint causes the gears in the **BP_Gears** Blueprint to start or stop spinning.
 
-There are two separate _Blueprints_ here: A button _Blueprint_ with a box component to use as a trigger, and a _Blueprint_ made up of a collection of gears that will start rotating while the button is pressed. A third asset, a _Blueprint Interface_, is employed as a means of sending information from one _Blueprint_ to another. 
+
+There are two separate Blueprints here: A button Blueprint with a box component to use as a trigger, and a Blueprint made up of a collection of gears that will start rotating while the button is pressed. A third asset, a Blueprint Interface, is employed as a means of sending information from one Blueprint to another. 
 
 
 ## BP_Button
@@ -53,9 +53,9 @@ There are two separate _Blueprints_ here: A button _Blueprint_ with a box compon
 Click for full size or **right-click** and Save As.
 [/REGION]
 
-Shape components can either be a box, a capsule, or a sphere, and are simple shapes used primarily for collision detection. The box component in **BP_Button** is set to only overlap Pawns, so the **OnComponentBeginOverlap** event in the script will only fire when a Pawn overlaps the box.
+Shape components can either be a box, a capsule, or a sphere, and are simple shapes used primarily for collision detection. The box component in BP_Button is set to only overlap Pawns, so the **OnComponentBeginOverlap** event in the script will only fire when a Pawn overlaps the box.
 
-The **Toggle** interface message nodes have the **Target Blueprint** _Actor_ variable connected to their **Target** pins.  The **Target Blueprint** variable is public, and set to **BP_Gears** for the instance of the button placed in the level.  These interface message nodes execute the **Toggle** event present in the **BP_Gears** _Blueprint_ as a result.
+The **Toggle** interface message nodes have the **Target Blueprint** Actor variable connected to their **Target** pins.  The **Target Blueprint** variable is public, and set to BP_Gears for the instance of the button placed in the level.  These interface message nodes execute the **Toggle** event present in the BP_Gears Blueprint as a result.
 
 
 ## BP_Gears
@@ -68,10 +68,10 @@ The **Toggle** interface message nodes have the **Target Blueprint** _Actor_ var
 Click for full size or **right-click** and Save As.
 [/REGION]
 
-The gears _Blueprint_ implements the interface in the form of a special **Toggle** interface event, connected to nodes that causes the gear _Static Mesh Components_ to stop and start turning. This event is triggered by the **Toggle** interface message nodes in **BP_Button**.
+The gears Blueprint implements the interface in the form of a special **Toggle** interface event, connected to nodes that causes the gear Static Mesh Components to stop and start turning. This event is triggered by the **Toggle** interface message nodes in BP_Button.
 
-Interface variables - The interface event employed by these _Blueprints_ also outputs a bool variable, which is used to tell the gears _Blueprint_ whether it should stop or start turning when the interface event is fired. Interfaces can also return variables to the _Blueprint_ that called the event, which can be useful for passing information back and forth between _Blueprints_ in both directions.
+Interface variables - The interface event employed by these Blueprints also outputs a bool variable, which is used to tell the gears Blueprint whether it should stop or start turning when the interface event is fired. Interfaces can also return variables to the Blueprint that called the event, which can be useful for passing information back and forth between Blueprints in both directions.
 
-In those cases, however, the implementation in the **Target** _Blueprint_ is done with functions instead of events.
+In those cases, however, the implementation in the **Target Blueprint** is done with functions instead of events.
 
 [/REGION]
