#UE4doc #docs:
Removed Overview header
Changed font settings for multiple terms

diff --git a/Engine/Documentation/Source/Resources/ContentExamples/Blueprints_Advanced/2_4/2_4.INT.udn b/Engine/Documentation/Source/Resources/ContentExamples/Blueprints_Advanced/2_4/2_4.INT.udn
--- a/Engine/Documentation/Source/Resources/ContentExamples/Blueprints_Advanced/2_4/2_4.INT.udn
+++ b/Engine/Documentation/Source/Resources/ContentExamples/Blueprints_Advanced/2_4/2_4.INT.udn
@@ -34,11 +34,11 @@ Related: Gameplay/HowTo/SetUpInput/Blueprints
 
 ![](blueprint_advanced_2_4.png)
 
-## Overview
 
-The **BP_Door** _Blueprint_ creates a door that opens with a trigger _Box Component_, and is another example of a _Blueprint_ that might commonly be used in the creation of a game, with script and animation built in that needs no additional setup. 
 
-When a player touches the _BoxComponent_, the _Blueprint_ retrieves the player's location relative to the door and uses this information in conjunction with the Actor's rotation to determine which side of the door the player is on. Then, a [Timeline](Engine/Blueprints\UserGuide\Timelines) plays that opens the door in the correct direction, so that it swings away from the player.
+The **BP_Door** Blueprint creates a door that opens with a trigger Box Component, and is another example of a Blueprint that might commonly be used in the creation of a game, with script and animation built in that needs no additional setup. 
+
+When a player touches the BoxComponent, the Blueprint retrieves the player's location relative to the door and uses this information in conjunction with the Actor's rotation to determine which side of the door the player is on. Then, a [Timeline](Engine/Blueprints\UserGuide\Timelines) plays that opens the door in the correct direction, so that it swings away from the player.
 
 ## Event Graph
 
@@ -49,6 +49,6 @@ When a player touches the _BoxComponent_, the _Blueprint_ retrieves the player's
 Click for full size or **right-click** and Save As.
 [/REGION]
 
-This Blueprint employs two [Timelines](Engine/Blueprints\UserGuide\Timelines): one for opening the door, and one for closing it. When the player first overlaps the _Box Component_, the **Open direction** Boolean variable is set depending on which side of the door the player is on. Next, the timeline for opening the door plays. The timeline is connected to a **Lerp (Rotator)** node, which linearly interpolates between the starting rotation and the desired final rotation. A **Select Float** node selects the correct final rotation based on the value of **Open Direction**. This way, the same timeline can be used to open the door in either direction.
+This Blueprint employs two [Timelines](Engine/Blueprints\UserGuide\Timelines): one for opening the door, and one for closing it. When the player first overlaps the Box Component, the **Open direction** Boolean variable is set depending on which side of the door the player is on. Next, the timeline for opening the door plays. The timeline is connected to a **Lerp (Rotator)** node, which linearly interpolates between the starting rotation and the desired final rotation. A **Select Float** node selects the correct final rotation based on the value of **Open Direction**. This way, the same timeline can be used to open the door in either direction.
 
 [/REGION]
