#UE4doc #docs:
Removed Overview header
Changed font settings for multiple terms

diff --git a/Engine/Documentation/Source/Resources/ContentExamples/Blueprints_Advanced/2_5/2_5.INT.udn b/Engine/Documentation/Source/Resources/ContentExamples/Blueprints_Advanced/2_5/2_5.INT.udn
--- a/Engine/Documentation/Source/Resources/ContentExamples/Blueprints_Advanced/2_5/2_5.INT.udn
+++ b/Engine/Documentation/Source/Resources/ContentExamples/Blueprints_Advanced/2_5/2_5.INT.udn
@@ -32,11 +32,11 @@ Related: Gameplay/HowTo/SetUpInput/Blueprints
 
 ![](blueprint_advanced_2_5.png)
 
-## Overview
 
-It is possible to create a child _Blueprint_ based on a parent _Blueprint_, similar to the way you would create a Material Instance of a master Material. The children automatically inherit the functionality of their parent, which can then be augmented, added to, or altered. In this example, the three collectible pickup _Blueprints_ are children of a generic pickup _Blueprint_. 
 
-The parent _Blueprint_ is **BP_Pickup_Parent_C**, and the three child _Blueprints_ are:
+It is possible to create a child Blueprint based on a parent Blueprint, similar to the way you would create a Material Instance of a master Material. The children automatically inherit the functionality of their parent, which can then be augmented, added to, or altered. In this example, the three collectible pickup Blueprints are children of a generic pickup Blueprint. 
+
+The parent Blueprint is **BP_Pickup_Parent_C**, and the three child Blueprints are:
 
 *  **BP_Pickup_Child_Health**
 * **BP_Pickup_Child_Coin**
@@ -54,13 +54,13 @@ Click for full size or **right-click** and Save As.
 [/REGION]
 
 
-The parent Blueprint, BP_Pickup_Parent, contains script for basic pickup behavior, such as a timeline animation to make the pickup spin and float, and a sequence for collecting the pickup. 
+The parent Blueprint, **BP_Pickup_Parent**, contains script for basic pickup behavior, such as a timeline animation to make the pickup spin and float, and a sequence for collecting the pickup. 
 The child Blueprints inherit all of this script, but have additional script for when they are collected, so that each pickup does something different while sharing the same basic features. 
 This is useful because any updates made to the parent Blueprint's script will also apply to all its child Blueprints.
 
 
-The OnPickedUp [event dispatcher](Engine/Blueprints\UserGuide\EventDispatcher) is called when the player's Pawn overlaps the **Pickup Trigger** _SphereComponent_. Each child _Blueprint_ binds a different 
-event to OnPickedUp, so that a different string prints when the player collects each pickup and OnPickedUp is called.
+The **OnPickedUp** [event dispatcher](Engine/Blueprints\UserGuide\EventDispatcher) is called when the player's Pawn overlaps the **Pickup Trigger** SphereComponent. Each child Blueprint binds a different 
+event to **OnPickedUp**, so that a different string prints when the player collects each pickup and **OnPickedUp** is called.
 [/EXCERPT:CoinParent]
 
 ## Child Blueprints' Event Graph
@@ -75,7 +75,7 @@ Click for full size or **right-click** and Save As.
 
 
 
-As a rule, Blueprint graphs cannot contain duplicates of the same event, so if both the parent script and the child script each contain a ReceiveBeginPlay event, only the event in the parent script will fire. You can get around this in the child script by **right-clicking** an event and selecting **Add call to parent function**. This will create a function that will trigger the corresponding event in the parent script when it is called.
+As a rule, Blueprint graphs cannot contain duplicates of the same event, so if both the parent script and the child script each contain a **ReceiveBeginPlay** event, only the event in the parent script will fire. You can get around this in the child script by **right-clicking** an event and selecting **Add call to parent function**. This will create a function that will trigger the corresponding event in the parent script when it is called.
 
 [/EXCERPT:CoinChild]
 
