#UE4doc #docs:
Changed font settings for multiple terms

diff --git a/Engine/Documentation/Source/Resources/ContentExamples/Blueprints_HUD/1_1/1_1.INT.udn b/Engine/Documentation/Source/Resources/ContentExamples/Blueprints_HUD/1_1/1_1.INT.udn
--- a/Engine/Documentation/Source/Resources/ContentExamples/Blueprints_HUD/1_1/1_1.INT.udn
+++ b/Engine/Documentation/Source/Resources/ContentExamples/Blueprints_HUD/1_1/1_1.INT.udn
@@ -36,7 +36,7 @@ The HUD Blueprint can be found by searching for **BP_HUD_Example** in the **Cont
 
 ## GameMode
 
-The **Game Mode** sets the default Pawn class, HUD class, and PlayerController class, in addition to containing any game-specific gameplay behavior or settings.  In this example,
+The **GameMode** sets the default Pawn class, HUD class, and PlayerController class, in addition to containing any game-specific gameplay behavior or settings.  In this example,
 the **GameInfo_HUD** Blueprint is the GameMode Blueprint which sets the **BP_HUD_Example** Blueprint as the HUD class, and it must be set as the **GameMode Override** in the **World Settings** tab.
  
 ![](gamemode_hud.png)
@@ -52,15 +52,15 @@ The **ReceiveDrawHUD** event that is used to draw each element of the HUD provid
 ![](BP_HUD_healthbar.png)
 
 The health bar is a dynamic Material Instance created in the HUD Blueprint's Construction Script. The Material contains nodes that change the length of the health bar 
-according to **Current Health** and **Maximum Health** scalar parameters, and the HUD Blueprint updates the **Current Health** parameter based on the Character's current health.
+according to Current Health and Maximum Health scalar parameters, and the HUD Blueprint updates the Current Health parameter based on the Character's current health.
 
 ### Pickups
 ![](BP_HUD_pickup.png)
 
 The objective pickup in this level is the **BP_Pickup_Child_Health** also used in the [Advanced Blueprints - Child Blueprints](Resources\ContentExamples\Blueprints_Advanced\2_5) content example. It
-is set as the **Objective** variable in the Level Blueprint. 
+is set as the Objective variable in the Level Blueprint. 
 
-In the HUD Blueprint, if the **Current Objective** is not null:
+In the HUD Blueprint, if the Current Objective is not null:
 -	The objective marker uses a **Project** function to convert a 3D location vector (representing the objective's location) into 2D screen coordinates, so that the 
 objective label and icon drawn at those coordinates will always be hovering over the objective Actor when it is on the screen.
 -	The text is centered over the objective by using the **Get Text Size** node. This node measures the height and width a string will be when drawn on the HUD, taking into account the font used and the scale of the text. 
@@ -76,7 +76,7 @@ Key features shown in the pause menu include:
 -	Custom functions: A custom function named **Draw Button** is used as a convenient way to draw all the necessary components of a button (button texture, button text, and a hit box)
  with a single function node that can be re-used as many times as needed. This is particularly useful in the case of the HUD, because a lot of math nodes are necessary to determine
   screen coordinate values, offsets, and so on, and the function allows these to be contained and reused easily. 
--	Hit boxes: The **Add Hit Box** function in the **DrawButton** function is used as a simple way to detect when the mouse is clicked or released on a specified region of the screen. The position and size of the hit box 
+-	Hit boxes: The **Add Hit Box** function in the DrawButton function is used as a simple way to detect when the mouse is clicked or released on a specified region of the screen. The position and size of the hit box 
 can be specified, as well as a name. A generic **ReceiveHitBoxClick** event in the main EventGraph fires when a hit box is clicked and outputs the name of the hit box. A **Switch on Name** node can then
 direct the execution flow according to this name output.  
 
@@ -89,7 +89,7 @@ direct the execution flow according to this name output.
 The **BP_PlayerController_HUD** Blueprint is the default PlayerController class for this example, and is also set in the **GameInfo_HUD** GameMode Blueprint. Although the **Press "M" for Menu** string is set up by the **BP_HUD_Example** Blueprint,
 the PlayerController Blueprint actually sets the initial response to pressing the **M** key.
 
-Pressing **M** executes the **Toggle Menu** interface message with the HUD Blueprint as the target. The PlayerController then:
+Pressing **M** executes the Toggle Menu interface message with the HUD Blueprint as the target. The PlayerController then:
 
 * Sets **Show Mouse Cursor** to _true_.
 * Sets **Enable Click Events** to _true_.
