#UE4doc #docs:
Removed Overview header
Changed font settings for multiple terms

diff --git a/Engine/Documentation/Source/Resources/ContentExamples/Blueprint_Splines/Blueprint_Splines.INT.udn b/Engine/Documentation/Source/Resources/ContentExamples/Blueprint_Splines/Blueprint_Splines.INT.udn
--- a/Engine/Documentation/Source/Resources/ContentExamples/Blueprint_Splines/Blueprint_Splines.INT.udn
+++ b/Engine/Documentation/Source/Resources/ContentExamples/Blueprint_Splines/Blueprint_Splines.INT.udn
@@ -42,30 +42,30 @@ Related: Engine/Rendering/Materials
 [TOC(start:2 end:2)]
 
 
-## Overview
-The **Blueprint Splines Content Examples** show how you can use **Spline Components** with **Blueprints** to create paths that can be used in a number of different ways. By simply adding a **Spline Component** to your **Blueprint** and selecting an instance of that **Blueprint**, you will gain access to editing splines within your level. You will then be able to select spline keys inside the viewport and manipulate them like any other Actor using the **Translation**, **Rotation**, and **Scale** widgets. 
 
-The **Content Examples** provided on this page illustrate how **Blueprint Splines** can be used with a **Construction Script** to create paths of **Static Mesh** objects or how to dynamically alter splines paths or **Spline Meshes** during gameplay to create content or animations. 
+The Blueprint Splines Content Examples show how you can use **Spline Components** with **Blueprints** to create paths that can be used in a number of different ways. By simply adding a Spline Component to your Blueprint and selecting an instance of that Blueprint, you will gain access to editing splines within your level. You will then be able to select spline keys inside the viewport and manipulate them like any other Actor using the Translation, Rotation, and Scale widgets. 
+
+The Content Examples provided on this page illustrate how **Blueprint Splines** can be used with a Construction Script to create paths of Static Mesh objects or how to dynamically alter splines paths or Spline Meshes during gameplay to create content or animations. 
 
  
 ## Construction Script Examples
 
 ![](blueprintSplines_topic_construction.png)(w:640)
 
-The examples below showcase how a **Spline Component** can be used with a **Blueprint's Construction Script** to paths or scalable **Static Meshes** inside the viewport in real-time. 
+The examples below showcase how a Spline Component can be used with a Blueprint's Construction Script to paths or scalable Static Meshes inside the viewport in real-time. 
 
 | Example | What is Demonstrated |
 | ------- | ------------ |
-| **1.1 Blueprint Spline Component** | Demonstrates how a **Spline Component** can be edited easily in the viewport (e.g. move four spline points of a **Spline Component** inside the viewport). |
-| **1.2 Placing Components Along Splines** | Shows how a **Blueprint** can retrieve various kinds of information from a **Spline Component** (e.g. shows how **Static Mesh Components** can be placed along the length of a spline where their orientation is determined by the direction of the spline in which they were placed). |
-| **1.3 Spline Mesh Components** | **Spline Meshes** which are another type of **Spline Component** used to stretch and bend a **Static Mesh** asset (e.g. a pipe that can be stretched and orientated using widgets in the viewport). |
+| **1.1 Blueprint Spline Component** | Demonstrates how a Spline Component can be edited easily in the viewport (e.g. move four spline points of a Spline Component inside the viewport). |
+| **1.2 Placing Components Along Splines** | Shows how a Blueprint can retrieve various kinds of information from a Spline Component (e.g. shows how Static Mesh Components can be placed along the length of a spline where their orientation is determined by the direction of the spline in which they were placed). |
+| **1.3 Spline Mesh Components** | Spline Meshes which are another type of Spline Component used to stretch and bend a Static Mesh asset (e.g. a pipe that can be stretched and orientated using widgets in the viewport). |
 
 
 ## Dynamic Examples
 
 ![](blueprintSplines_topic_dynamic.png)(w:640)
 
-The **Dynamic Examples** show how a **Spline Component** can be updated dynamically during gameplay to provide ambient background movement of objects or to animate objects or even a playable character! 
+The **Dynamic Examples** show how a Spline Component can be updated dynamically during gameplay to provide ambient background movement of objects or to animate objects or even a playable character! 
 
 | Example | What is Demonstrated |
 | ------- | ------------ |
