#UE4doc #docs:
Removed Overview header
Changed font settings for multiple terms

diff --git a/Engine/Documentation/Source/Resources/ContentExamples/Brushes/BSP.INT.udn b/Engine/Documentation/Source/Resources/ContentExamples/Brushes/BSP.INT.udn
--- a/Engine/Documentation/Source/Resources/ContentExamples/Brushes/BSP.INT.udn
+++ b/Engine/Documentation/Source/Resources/ContentExamples/Brushes/BSP.INT.udn
@@ -55,11 +55,11 @@ Related: Engine/UI/LevelEditor/Modes
 [TOC(start:2 end:2)]
 
 
-## Overview
+
 There are a number of methods in which you can create the geometry that comprise your levels in Unreal Engine 4. Generally speaking, it is better (more efficient performance wise) to use [Static Meshes](Engine/Content\Types\StaticMeshes) when creating the geometry.
 That said, since creating Static Meshes requires some sort of 3D modeling software, you can use the **Geometry Editing** tools to quickly block out your levels for gameplay prototyping. 
 
-This page covers the **Geometry_Editing** map and illustrates how to begin blocking out levels by using **Additive** and **Subtractive Brushes** to define the areas where geometry exists. By the end of the examples, you will learn how to build a house entirely out of **Brushes** and should have enough knowledge to start blocking out your own levels. 
+This page covers the **Geometry_Editing** map and illustrates how to begin blocking out levels by using **Additive** and **Subtractive Brushes** to define the areas where geometry exists. By the end of the examples, you will learn how to build a house entirely out of Brushes and should have enough knowledge to start blocking out your own levels. 
 
 [region:note]
 You can find additional resources about each example by clicking on the example's name in table below.
@@ -68,14 +68,14 @@ You can find additional resources about each example by clicking on the example'
  
 ## Geometry_Editing Map
 
-Listed below are the examples provided in the **Geometry_Editing** map: 
+Listed below are the examples provided in the Geometry_Editing map: 
 
 | Example | What is Demonstrated |
 | ------- | ------------ |
-| **1.1 Additive and Subtractive Geometry Example** | Shows how **Additive** and **Subtractive Brushes** interact with one another. |
-| **1.2 Order Example** | Shows how the **Order** in which you place **Additive** or **Subtractive Brushes** can affect how they interact with one another. |
-| **1.3 Solidity Example** | Shows how a **Brush** can collide differently with objects based on its **Solidity** setting. (e.g. whether a **Brush** causes cuts to be created in surrounding brushes or if it collides with objects in the world). |
+| **1.1 Additive and Subtractive Geometry Example** | Shows how Additive and Subtractive Brushes interact with one another. |
+| **1.2 Order Example** | Shows how the order in which you place Additive or Subtractive Brushes can affect how they interact with one another. |
+| **1.3 Solidity Example** | Shows how a Brush can collide differently with objects based on its **Solidity** setting. (e.g. whether a Brush causes cuts to be created in surrounding brushes or if it collides with objects in the world). |
 | **1.4 Surface Properties Example** | Shows how to control how textures that are placed across your surfaces as well as lightmap resolution via the **Surface Properties** options. |
-| **1.5 Edit Geometry Example** | This example demonstrates how to edit a **Brush's** shape with **Edit Mode** by clicking on the Edit (**Shift+5**) button. |
-| **1.6 Example House** | This is an example of using **Additive** and **Subtractive Brushes** to construct a house with windows and a door opening. |
+| **1.5 Edit Geometry Example** | This example demonstrates how to edit a Brush's shape with **Edit Mode** by clicking on the Edit (**Shift+5**) button. |
+| **1.6 Example House** | This is an example of using Additive and Subtractive Brushes to construct a house with windows and a door opening. |
 
