#UE4doc #docs:
Removed Overview header
Changed font settings for multiple terms

diff --git a/Engine/Documentation/Source/Resources/ContentExamples/Cloth/Cloth.INT.udn b/Engine/Documentation/Source/Resources/ContentExamples/Cloth/Cloth.INT.udn
--- a/Engine/Documentation/Source/Resources/ContentExamples/Cloth/Cloth.INT.udn
+++ b/Engine/Documentation/Source/Resources/ContentExamples/Cloth/Cloth.INT.udn
@@ -43,7 +43,7 @@ Related: Engine/Physics/Apex
 [TOC(start:2 end:2)]
 
 
-## Overview
+
 
 The **Cloth** map inside the Content Examples project shows how to create and manipulate various cloth assets. In this map, you will see how to create cloth that simulates being affected by wind, how to adjust properties that can mimic various types of materials (silk or leather for example), cloth collision settings to prevent cloth from intersecting objects or itself as well as an example of a playable character wearing cloth. 
 
@@ -52,15 +52,15 @@ The **Cloth** map inside the Content Examples project shows how to create and ma
 
 ![](Cloth_Topic.png)(w:640)
 
-Listed below are the examples provided inside the **Cloth** map:
+Listed below are the examples provided inside the Cloth map:
 
 | Example | What is Demonstrated |
 | ------- | ------------ |
 | **1.1 Apex Cloth and WindDirectionalSource** | **SkeletalMesh** with **Apex Cloth Asset** affected by a **WindDirectionalSource** (e.g. a cloth asset that simulates a flag blowing in the wind). |
-| **1.2 Apex Cloth Properties (imported)** | **Apex Cloth Assets** with different properties are set to simulate different types of cloth (e.g. Burlap, Silk, Rubber, and Heavy Leather materials are simulated). |
-| **1.3 Apex Cloth Collision** | Two examples using **Capsule Rigid Bodies** for collision which helps prevent cloth from intersecting with parts of the **SkeletalMesh** that are not simulating cloth. |
+| **1.2 Apex Cloth Properties (imported)** | Apex Cloth Assets with different properties are set to simulate different types of cloth (e.g. Burlap, Silk, Rubber, and Heavy Leather materials are simulated). |
+| **1.3 Apex Cloth Collision** | Two examples using **Capsule Rigid Bodies** for collision which helps prevent cloth from intersecting with parts of the SkeletalMesh that are not simulating cloth. |
 | **1.4 BackStop** | Another method of cloth collision where the distance cloth can move is determined by the bone its weighted to (better for performance). |
-| **1.5 Latch to Nearest** | Lets you define **Vertices** that drive the movement of other **Vertices** and is used to create "thicker" or more complex geometry (e.g. heavier cloth or a pocket attached to cloth). |
+| **1.5 Latch to Nearest** | Lets you define **Vertices** that drive the movement of other Vertices and is used to create "thicker" or more complex geometry (e.g. heavier cloth or a pocket attached to cloth). |
 | **1.6 Self Collision** | Demonstrates an example where **Self Collision** is enabled and disabled which prevents the cloth from intersecting itself during movement. |
 | **1.7 Playable Character with Cloth** | Playable character that demonstrates cloth attached to a character (e.g. playable Owen who wears a cloth coat and tie that reacts to movement). |
 
