UEDOC-2107

Converted gifs into YouTube videos

diff --git a/Engine/Documentation/Source/Engine/Animation/Tools/VertexAnimationTool/VertexAnimationTools.INT.udn b/Engine/Documentation/Source/Engine/Animation/Tools/VertexAnimationTool/VertexAnimationTools.INT.udn
--- a/Engine/Documentation/Source/Engine/Animation/Tools/VertexAnimationTool/VertexAnimationTools.INT.udn
+++ b/Engine/Documentation/Source/Engine/Animation/Tools/VertexAnimationTool/VertexAnimationTools.INT.udn
@@ -69,7 +69,17 @@ This tool has only be tested using **3Ds Max 2015**.
 While it might be able to work with other 3Ds Max versions, none of them have been tested so use other version of 3Ds Max at your own risk.
 To install the script just drag it from its location in **4.8\Engine\Extras\3dsMaxScripts** right into the 3Ds Max viewport and the script will launch itself.
 
-![](VAT_Max_Install.gif)(convert:false)
+[OBJECT:EmbeddedVideo]
+[PARAMLITERAL:width]
+640
+[/PARAMLITERAL]
+[PARAMLITERAL:height]
+360
+[/PARAMLITERAL]
+[PARAMLITERAL:videoid]
+koPxpL0CgA8
+[/PARAMLITERAL]
+[/OBJECT]
 
 [region:tip]
 If you find that you are using this script a lot you can always add the script to one of the tool bars or quad menus. If you are unfamiliar how to do this the [**Autodesk site has a very detailed walk through**](http://knowledge.autodesk.com/support/3ds-max/learn-explore/caas/CloudHelp/cloudhelp/2015/ENU/3DSMax/files/GUID-A2CF8BAA-7B52-40A8-8C40-803B1AB5FC05-htm.html) that will explain how you go about doing this.
@@ -84,7 +94,17 @@ Since UE4 uses Centimeters for its default unit of measurement we will need to e
 1. Next click on **System Unit Setup** and change the **System Unit Setup** from Inches to **Centimeters** and then press the **OK** button.
 1. Finally change the **Display Unit Scale** to **Generic Units** and then press the **OK** button.
 
-	![](VAT_3DsMax_Unit_Setup.gif)(convert:false)
+	[OBJECT:EmbeddedVideo]
+	[PARAMLITERAL:width]
+	640
+	[/PARAMLITERAL]
+	[PARAMLITERAL:height]
+	360
+	[/PARAMLITERAL]
+	[PARAMLITERAL:videoid]
+	lZ6d6Q902qw
+	[/PARAMLITERAL]
+	[/OBJECT]
 
 [region:note]
 It is extremely important that you **DO NOT** skip this step. Skipping this step could result in rendering errors when the content is imported into UE4 because UE4 and 3Ds Max's units are not the same.
