UEDOC-2089

gif has been converted into a YouTube Video

diff --git a/Engine/Documentation/Source/Engine/Rendering/Materials/Functions/Reference/PivotPainter/PivotPainter.INT.udn b/Engine/Documentation/Source/Engine/Rendering/Materials/Functions/Reference/PivotPainter/PivotPainter.INT.udn
--- a/Engine/Documentation/Source/Engine/Rendering/Materials/Functions/Reference/PivotPainter/PivotPainter.INT.udn
+++ b/Engine/Documentation/Source/Engine/Rendering/Materials/Functions/Reference/PivotPainter/PivotPainter.INT.udn
@@ -8,7 +8,17 @@ Description:Functions designed to allow utilization of the Pivot Painter MAXScri
 
 
 [REGION:header_img]
-[![splash_demo.jpg](Engine/Content/Tools/PivotPainter/splash_demo.jpg)](Engine/Content/Tools/PivotPainter/PivotPainterFoliageExamples.mp4)
+[OBJECT:EmbeddedVideo]
+[PARAMLITERAL:width]
+640
+[/PARAMLITERAL]
+[PARAMLITERAL:height]
+360
+[/PARAMLITERAL]
+[PARAMLITERAL:videoid]
+G3vUj7c5gUM
+[/PARAMLITERAL]
+[/OBJECT]
 [/REGION]
 
 Pivot Painter material functions allow you to tap into the [Pivot Painter MAXScript](Engine/Content/Tools/PivotPainter), which stores rotation information within the vertices of a mesh. This is a great way to handle dynamic motion on Static Meshes.
