#UE4doc #docs Completed first draft of 'Installing Unreal Engine' revision.

diff --git a/Engine/Documentation/Source/GettingStarted/Installation/InstallingUnrealEngine.INT.udn b/Engine/Documentation/Source/GettingStarted/Installation/InstallingUnrealEngine.INT.udn
--- a/Engine/Documentation/Source/GettingStarted/Installation/InstallingUnrealEngine.INT.udn
+++ b/Engine/Documentation/Source/GettingStarted/Installation/InstallingUnrealEngine.INT.udn
@@ -1,8 +1,184 @@
 Availability:Public
 Title:Installing Unreal Engine
-Crumbs: %ROOT%, GettingStarted
 Description:Steps for installing Unreal Engine.
+Crumbs: %ROOT%, GettingStarted
+Prereq: GettingStarted/RecommendedSpecifications
+Prereq: Support/Builds
+Related: GettingStarted
+Related: Engine/Content/QuickStart
+Related: Programming/QuickStart
+Related: Engine/QuickStart
+Version: 4.8
+
+[TOC (start:2 end:3)]
+
+----
+
+We make it easy for subscribers to download, install, and run Unreal Engine 4 (UE4) from our launcher.
+
+![](InstallationGoal.png)
+
+_By the end of this tutorial, you'll be able to run the Unreal Editor from the **Epic Games Launcher** (pictured above)._
+
+The Epic Games Launcher enables our subscribers to conveniently download and install updates for UE4.
+
+## Creating an Epic Games Account
+
+[REGION:warning]
+If you already have an Epic Games account, skip ahead to learn about [Downloading and Running the Installer](#Bookmark1).
+[/REGION]
+
+In order to run the Epic Games Launcher, you need to have a valid Epic Games account.
+
+1.	**Navigate** to [UnrealEngine.com](https://www.unrealengine.com/what-is-unreal-engine-4).
+
+1.	**Click** the **Get Unreal** button (pictured below).
+
+	![](UnrealWebsite.png)
+
+1.	**Fill-out** our **'Join the Community'** Form and **click** the **Sign Up** button (pictured below).
+	
+	![](SignupScreen.png)
+
+	[REGION:note]
+	Make sure to read through our [terms of service](http://epicgames.com/tos) prior to signing up.
+	[/REGION]
+
+1.	**Read** our [Unreal Engine End User License Agreement](https://www.unrealengine.com/eula) (EULA) prior to **clicking** the **Accept** button (pictured below).
+
+	![](EULA.png)
+
+	[REGION:note]
+	Refer to our [EULA FAQ](https://www.unrealengine.com/faq#legal) if you have additional questions about the EULA's terms and conditions.
+	[/REGION]
+
+## Downloading and Running the Installer
+
+(#Bookmark1)
+
+With a valid Epic Games account, you're able to download our Epic Games Launcher set-up program (installer).
+
+%Globals:osselect_linux%
+
+[OBJECT:ToggleButtonContent]
+	[PARAMLITERAL:category]
+	OS
+	[/PARAMLITERAL]
+	[PARAMLITERAL:id]
+	windows
+	[/PARAMLITERAL]
+	[PARAMLITERAL:active]
+	active_button_content
+	[/PARAMLITERAL]
+	[PARAM:content]
+	1.	**Click** the **Download** button to get our installer (pictured below).
+	
+		![](Download_Screen.png)
+	
+	1.	Once the download is complete, **Run** the installer (pictured below).
+	
+		![](RunInstaller.png)
+	
+	1.	**Click** the **Install** button when the **Epic Games Launcher Setup** dialog box appears (pictured below).
+	
+		![](RunEpicLauncherSetup.png)
+	[/PARAM]
+[/OBJECT]
+
+[OBJECT:ToggleButtonContent]
+	[PARAMLITERAL:category]
+	OS
+	[/PARAMLITERAL]
+	[PARAMLITERAL:id]
+	mac
+	[/PARAMLITERAL]
+	[PARAMLITERAL:active]
+	active_button_content
+	[/PARAMLITERAL]
+	[PARAM:content]
+	1.	**Click** the **Download** button to get our installer (pictured below).
+	
+		![](Download_Screen.png)
+	
+	1.	**Double-click** the `EpicGamesLauncher.dmg` icon in **Finder** to open a new **Finder Window**.
+
+	1.	**Drag** the **Epic Games Launcher** icon into the **Applications** folder (pictured below).
+	
+		![](InstallEpicLaucherMac.png)	
+	[/PARAM]
+[/OBJECT]
 
+[OBJECT:ToggleButtonContent]
+	[PARAMLITERAL:category]
+	OS
+	[/PARAMLITERAL]
+	[PARAMLITERAL:id]
+	linux
+	[/PARAMLITERAL]
+	[PARAMLITERAL:active]
+	active_button_content
+	[/PARAMLITERAL]
+	[PARAM:content]
+	We currently don't provide a binary installer of Unreal Engine for Linux.
+	[REGION:tip]
+	Linux users; however, can compile Unreal Engine from source by referring to our [](GettingStarted/DownloadingUnrealEngine) documentation.
+	[/REGION]
+	[/PARAM]
+[/OBJECT]
+
+## Signing into the Epic Launcher
+
+After the set-up program installs the Epic Games Launcher on your machine, you're ready to sign-in (see below).
+
+![](EpicSignIn.png)
+
+## Installing Unreal Engine
+
+Now that you've signed into the Epic Games Launcher, you're ready to install UE4.
+
+[REGION:warning]
+[REGION:largetitle]
+Space Requirements
+[/REGION]
+Unreal Engine uses approximately eight gigabytes of disk space for every version of the engine that is installed on a machine. Make sure you have plenty of disk space prior to installing UE4.
+[/REGION]
+
+1.	**Click** the **Unreal Engine** tab inside the Epic Games Launcher Loading Screen (pictured below).
+
+	![](LauncherLoadScreen.png)
+
+1.	**Click** the **Install Engine** button to **download and install** the latest version of Unreal Engine (pictured below).
+
+	![](InstallEngine.png)
+
+1.	**Read** our [Unreal Engine End User License Agreement](https://www.unrealengine.com/eula) (EULA) prior to **clicking** the **Accept** button (pictured below)
+
+	![](AcceptEULALauncher.png)
+
+1.	Depending on your system specifications, **downloading and installing** Unreal Engine will take between 10-40 minutes.
+	
+	![](DownloadingEngine.png)
+
+	_Example progress indicator (pictured above)._
+
+## Launching Unreal Engine
+
+After the Epic Games Launcher successfully downloads and installs Unreal Engine, you're ready to **click** the **Launch** button (see below).
+
+![](LaunchUE4.png)
+
+## Getting Started with Unreal Engine
+
+Learn how to use Unreal Engine by referring to the [](GettingStarted) documentation!
+
+[REGION:tip]
+If you're an artist, programmer, or designer looking to quickly get started with UE4, refer to the following tutorials:
+*	[](Engine/Content/QuickStart)
+*	[](Programming/QuickStart)
+*	[](Engine/QuickStart)
+[/REGION]
+
+<!--Deprecated on 21 July 2015 due to outdated screenshots and outdated procedures
 %Globals:OSSelect%
 
 ## Download the Installer
@@ -50,7 +226,25 @@ You can download the installer from the [UnrealEngine.com](http://unrealengine.c
 	[/PARAM]
 [/OBJECT]
 
-<!--
+## Sign-In and Download
+
+Once you have the Unreal Engine Launcher running, you must sign in and download the current version of the engine.
+
+1. Sign in with your account credentials.
+
+	[REGION:fullwidth]
+	![](ue4_installer_login.png)
+	[/REGION]
+
+1. The current version of Unreal Engine 4 will begin downloading:
+
+	![](ue4_installer_download.png)
+
+1. Once the download is completed, you can launch Unreal Engine using the **Launch** button:
+
+	![](launch.png)
+-->
+<!--Previously deprecated procedures
 [EXCERPT:Windows]
 1. Run the installer you downloaded and choose a location to install Unreal Engine to:
 
@@ -75,22 +269,4 @@ You can download the installer from the [UnrealEngine.com](http://unrealengine.c
 1. Once Finder finishes copying the files, Unreal Engine will be available in the **Launchpad**. Run it from there. 
 
 [/EXCERPT:MacOS]
--->
-
-## Sign-In and Download
-
-Once you have the Unreal Engine Launcher running, you must sign in and download the current version of the engine.
-
-1. Sign in with your account credentials.
-
-	[REGION:fullwidth]
-	![](ue4_installer_login.png)
-	[/REGION]
-
-1. The current version of Unreal Engine 4 will begin downloading:
-
-	![](ue4_installer_download.png)
-
-1. Once the download is completed, you can launch Unreal Engine using the **Launch** button:
-
-	![](launch.png)
+-->
\ No newline at end of file
