TypoCorrection

diff --git a/Engine/Documentation/Source/GettingStarted/DownloadingUnrealEngine/DownloadingUnrealEngine4.INT.udn b/Engine/Documentation/Source/GettingStarted/DownloadingUnrealEngine/DownloadingUnrealEngine4.INT.udn
--- a/Engine/Documentation/Source/GettingStarted/DownloadingUnrealEngine/DownloadingUnrealEngine4.INT.udn
+++ b/Engine/Documentation/Source/GettingStarted/DownloadingUnrealEngine/DownloadingUnrealEngine4.INT.udn
@@ -1,7 +1,7 @@
-Availability:Public
+﻿Availability:Public
 Title:Downloading Unreal Engine Source Code
-Description:Steps for connecting to the source code repository and downloading the latest build of Unreal Engine.
 Crumbs: %ROOT%, Programming, Programming/Development
+Description:Steps for connecting to the source code repository and downloading the latest build of Unreal Engine.
 Prereq: GettingStarted/RecommendedSpecifications
 Prereq: Support/Builds
 Related: Programming/Development/BuildingUnrealEngine
