UEDOC-2105

Converting gif to a YouTube video.
Replaced the Blueguy within the video
Added a note to the Documentation

diff --git a/Engine/Documentation/Source/Engine/Animation/AnimHowTo/AimOffset/AimOffset.INT.udn b/Engine/Documentation/Source/Engine/Animation/AnimHowTo/AimOffset/AimOffset.INT.udn
--- a/Engine/Documentation/Source/Engine/Animation/AnimHowTo/AimOffset/AimOffset.INT.udn
+++ b/Engine/Documentation/Source/Engine/Animation/AnimHowTo/AimOffset/AimOffset.INT.udn
@@ -9,7 +9,17 @@ Related: Engine\Animation\AnimHowTo\AdditiveAnimations
 
 In this guide, we will create an [**Aim Offset**](Engine/Animation/AimOffset), which is an asset that stores a blendable series of poses to help a character aim a weapon. We will take an animation sequence and chop it up into usable poses for an Aim Offset, get the player's pitch/yaw location of the mouse and use that to determine which blended pose to use so that the character moves and aims in the location of the mouse (with some restrictions) as seen in the example below. 
 
-![](AimGIF.gif)(convert:false)
+[OBJECT:EmbeddedVideo]
+[PARAMLITERAL:width]
+640
+[/PARAMLITERAL]
+[PARAMLITERAL:height]
+360
+[/PARAMLITERAL]
+[PARAMLITERAL:videoid]
+uyeMVzNn3_s
+[/PARAMLITERAL]
+[/OBJECT]
 
 [TOC(start:2 end:2)]
 
@@ -377,7 +387,17 @@ In this final section, we will change the default idle pose and update Character
 
 	This will remove the yellow warning sign on the compile button in the upper left corner of the window since the project does not have an Action Mapping for Crouch by default, we are going to use the **C** button to crouch (you can use any Key Event you wish for this). 
 
-1.	**Compile** and **Save** then close the Blueprint and click **Play** from the Main Editor Toolbar to play the game. 
+1.	**Compile** and **Save** then close the Blueprint 
+
+	[REGION:note]
+	In Engine Version 4.8 and higher, you will need to remove the **ThirdPersonCharacter** from your level. This can be done by locating the **ThirdPersonCharacter** within the **Scene Outliner** of the **Editor Viewport**, clicking on the **ThirdPersonCharacter**, and pressing delete. 
+
+	This needs to be done due to the **ThirdPersonCharacter** being automatically possesed within the level. 
+
+	Be sure to **Save** the project once removing the **ThirdPersonCharacter** from the level.
+	[/REGION:note]
+
+1.	Click **Play** from the Main Editor Toolbar to play the game. 
 
 When you play in the editor, the character will now aim while stationary and will react to mouse movement, pointing in the direction of your mouse. There are some restrictions in place that prevent the character from turning and facing backwards while the mouse is behind them as well and while the character is jogging, moving the mouse also turns the character's torso in the direction you are pointing. 
 
