UEDOC-1806

Replacing all gifs with YouTube Videos

diff --git a/Engine/Documentation/Source/Engine/Animation/Tools/VertexAnimationTool/VAT_KF_Meshes/VAT_KF_Meshes.INT.udn b/Engine/Documentation/Source/Engine/Animation/Tools/VertexAnimationTool/VAT_KF_Meshes/VAT_KF_Meshes.INT.udn
--- a/Engine/Documentation/Source/Engine/Animation/Tools/VertexAnimationTool/VAT_KF_Meshes/VAT_KF_Meshes.INT.udn
+++ b/Engine/Documentation/Source/Engine/Animation/Tools/VertexAnimationTool/VAT_KF_Meshes/VAT_KF_Meshes.INT.udn
@@ -31,12 +31,32 @@ Position the Teapot copy so that the orange bounding box of the Melt modifier do
 Once the Teapot copy is in position release the mouse and the shift key and the Clone Options will be displayed. 
 Make sure the **Object** is set to **Copy** and then set the **Number of Copies** to **Four** and press the **Ok** button to create the copies.
 
-	![](VAT_KFM_Teapot_Copy.gif)(convert:false)
+	[OBJECT:EmbeddedVideo]
+	[PARAMLITERAL:width]
+	640
+	[/PARAMLITERAL]
+	[PARAMLITERAL:height]
+	360
+	[/PARAMLITERAL]
+	[PARAMLITERAL:videoid]
+	H9vBL17HuWU
+	[/PARAMLITERAL]
+	[/OBJECT]
 
 1. Then select the 2nd Teapot mesh and set the **Melt Amount** in the Melt modifier to 25. 
 Then select the next Teapot in line and set the **Melt Amount** of the Melt modifier to 50 and the next one to 75, and finally set the last Teapot **Melt Amount** to 100.
 
-	![](VAT_KFM_Teapot_Melt_Amount.gif)(convert:false)
+	[OBJECT:EmbeddedVideo]
+	[PARAMLITERAL:width]
+	640
+	[/PARAMLITERAL]
+	[PARAMLITERAL:height]
+	360
+	[/PARAMLITERAL]
+	[PARAMLITERAL:videoid]
+	Ltbpe9C8ENw
+	[/PARAMLITERAL]
+	[/OBJECT]
 
 1. With all the Teapots now having some type of deformation on them, select all the Teapots and then with the mouse Right-Click  in the 3Ds Max viewport and select **Convert To:** > **Convert to Editable Poly**.
 
@@ -48,7 +68,17 @@ Then select the next Teapot in line and set the **Melt Amount** of the Melt modi
 Once the meshes are in the center of the world, deselect the Teapot meshes and then select them in the order that the animation is supposed to play. 
 For this example we want to select Teapot001 first  then select Teapot002 and so on till all five are selected.
 
-	![](VAT_KFM_Teapot_Export_Setup.gif)(convert:false)
+	[OBJECT:EmbeddedVideo]
+	[PARAMLITERAL:width]
+	640
+	[/PARAMLITERAL]
+	[PARAMLITERAL:height]
+	360
+	[/PARAMLITERAL]
+	[PARAMLITERAL:videoid]
+	sE8RBIVCBgA
+	[/PARAMLITERAL]
+	[/OBJECT]
 
 	[region:note]
 	The Sequence Painter Tool requires that you select the meshes in the order that they should be played in the animation. Failing to do this will result in the animation that gets played in UE4 not syncing up to the bone driven animation that it is supposed to replace.
@@ -96,11 +126,31 @@ Set the value of one of the Constant Parameter node and plug that into the **Bas
 1. Next use the mouse and Right - Click in the Material Graph to expose the search menu and then search for **MS_SequencePainter_SequenceFlipbook** Material Function using MS as the search term.
 When you find the MS_SequencePainter_SequenceFlipbook Material Function, click on it with the Left Mouse button to add it to the Material Graph.
 
-	![](VAT_KFM_Adding_MF_Flipbook.gif)(convert:false)
+	[OBJECT:EmbeddedVideo]
+	[PARAMLITERAL:width]
+	640
+	[/PARAMLITERAL]
+	[PARAMLITERAL:height]
+	360
+	[/PARAMLITERAL]
+	[PARAMLITERAL:videoid]
+	juDHJS0K7b4
+	[/PARAMLITERAL]
+	[/OBJECT]
 
 1. With the correct Material Function now in place, connect the output of the **MS_SequencePainter_SequenceFlipbook**,  into the World Position Offset input on the Main Material node and then connect the remaining Scalar parameter to the **Number Of Frames** input on the MS_SequencePainter_SequenceFlipbook Material Function node.
 
-	![](VAT_KFM_Linking_Up_MF.gif)(convert:false)
+	[OBJECT:EmbeddedVideo]
+	[PARAMLITERAL:width]
+	640
+	[/PARAMLITERAL]
+	[PARAMLITERAL:height]
+	360
+	[/PARAMLITERAL]
+	[PARAMLITERAL:videoid]
+	5VGa5sdRvfk
+	[/PARAMLITERAL]
+	[/OBJECT]
 
 1. Now that everything has been hooked up press the **Apply** button to compile the Material and then press the **Save** button. 
 After that is completed you can then close down the Material Graph Editor.
@@ -110,8 +160,17 @@ After that is completed you can then close down the Material Graph Editor.
 1. With the Material created, compiled, and saved drag the Teapot mesh from the Content Browser into the level and then drag the Material that was just created for it from the Content Browser on to the mesh that was placed in the level. 
 Once this has been completed you will now see the mesh "playing" the animation.
 
-	![](VAT_KFM_Final_Effect.gif)(convert:false)
-
+	[OBJECT:EmbeddedVideo]
+	[PARAMLITERAL:width]
+	640
+	[/PARAMLITERAL]
+	[PARAMLITERAL:height]
+	360
+	[/PARAMLITERAL]
+	[PARAMLITERAL:videoid]
+	TcyZIZp-0u8
+	[/PARAMLITERAL]
+	[/OBJECT]
 
 
 
