#UE4doc #docs:
Removed Overview header
Changed font settings for multiple terms
Changed a picture due to the Class Viewer not being the default way to access classes now. Also changed text preceding the picture to reflect this change.
Commented out multiple lines and a picture as they are no longer relevant

diff --git a/Engine/Documentation/Source/Resources/ContentExamples/Decals/Decals.INT.udn b/Engine/Documentation/Source/Resources/ContentExamples/Decals/Decals.INT.udn
--- a/Engine/Documentation/Source/Resources/ContentExamples/Decals/Decals.INT.udn
+++ b/Engine/Documentation/Source/Resources/ContentExamples/Decals/Decals.INT.udn
@@ -37,7 +37,7 @@ Description:An overview of the Decals example level
 [/VAR]
 
 
-## Overview
+
 
 [REGION:fullwidth]
 
@@ -45,13 +45,17 @@ Description:An overview of the Decals example level
 
 Welcome to the **Decals** example level! This level demonstrates ways in which Decals can be utilized in your levels. 
 
-To create a Decal Actor, simply drag  and drop them from the **Class Viewer** panel.
+To create a Deferred Decal Actor, drag and drop them from the **Modes** panel.
 
 ![](DecalActorList.png)
 
-You then only need to apply a Decal Material, which is really just a regular Material that has had its **Material Domain** property set to _Deferred Decal_.
+[REGION:note]
+This is only accessible from the **Modes** panel if you're in Placement mode.
+[/REGION]
+
+<!--You then only need to apply a Decal Material, which is really just a regular Material that has had its **Material Domain** property set to _Deferred Decal_.
 
-![](DeferredDecalSetting.png)
+![](DeferredDecalSetting.png) - Commented out due to no longer being relevant, as decal actors are automatically Deferred Decals-->
 
 [INCLUDE:Resources\ContentExamples\ExampleLevelOverview#ExampleLevel]
 
