#UE4doc #docs:
Removed Overview header
Changed font settings for multiple terms
Replaced picture as it was out of date
Commented out a line and reworded a line as it was out of date.

diff --git a/Engine/Documentation/Source/Resources/ContentExamples/Decals/1_1/1_1.INT.udn b/Engine/Documentation/Source/Resources/ContentExamples/Decals/1_1/1_1.INT.udn
--- a/Engine/Documentation/Source/Resources/ContentExamples/Decals/1_1/1_1.INT.udn
+++ b/Engine/Documentation/Source/Resources/ContentExamples/Decals/1_1/1_1.INT.udn
@@ -55,12 +55,12 @@ Related: Engine/Rendering/Materials/MaterialInstances
 
 ![%Resources/ContentExamples/Decals/1_1:title%](Resources\ContentExamples\Decals\Decals_1.png)
 
-## Overview
 
-Decals are Materials that are projected onto meshes in your level, including Static Meshes and Skeletal Meshes. These meshes can have a **Mobility** setting of Static or Movable and the Decal will still project on them. Many Decals can be rendered at once without a large performance decrease. Performance decreases with larger screen space size and higher shader instruction count.
 
+Decals are Materials that are projected onto meshes in your level, including Static Meshes and Skeletal Meshes. These meshes can have a Mobility setting of Static or Movable and the Decal will still project on them. Many Decals can be rendered at once without a large performance decrease. Performance decreases with larger screen space size and higher shader instruction count.
 
-You can add the Decal Actor directly from the Classes Viewer. You then need to make a new Material and change its **Material Domain** property to _Deferred Decal_. You can then apply the Material to the Decal Actor.
+
+You can add the Deferred Decal Actor directly from the **All Classes** tab of the **Modes** panel. <!--You then need to make a new Material and change its **Material Domain** property to _Deferred Decal_.--> You can then apply a Material to the Decal Actor.
 
 ![](DecalActorList.png)
 
