#UE4doc #docs:
Reverted some changes as they were incorrect

diff --git a/Engine/Documentation/Source/Resources/ContentExamples/Decals/1_1/1_1.INT.udn b/Engine/Documentation/Source/Resources/ContentExamples/Decals/1_1/1_1.INT.udn
--- a/Engine/Documentation/Source/Resources/ContentExamples/Decals/1_1/1_1.INT.udn
+++ b/Engine/Documentation/Source/Resources/ContentExamples/Decals/1_1/1_1.INT.udn
@@ -60,7 +60,7 @@ Related: Engine/Rendering/Materials/MaterialInstances
 Decals are Materials that are projected onto meshes in your level, including Static Meshes and Skeletal Meshes. These meshes can have a Mobility setting of Static or Movable and the Decal will still project on them. Many Decals can be rendered at once without a large performance decrease. Performance decreases with larger screen space size and higher shader instruction count.
 
 
-You can add the Deferred Decal Actor directly from the **All Classes** tab of the **Modes** panel. <!--You then need to make a new Material and change its **Material Domain** property to _Deferred Decal_.--> You can then apply a Material to the Decal Actor.
+You can add the Deferred Decal Actor directly from the **All Classes** tab of the **Modes** panel. You then need to make a new Material and change the **Material Domain** property to **Deferred Decal** in the Material's **Details** panel. You can then apply the Material to the Decal Actor.
 
 ![](DecalActorList.png)
 
diff --git a/Engine/Documentation/Source/Resources/ContentExamples/Decals/Decals.INT.udn b/Engine/Documentation/Source/Resources/ContentExamples/Decals/Decals.INT.udn
--- a/Engine/Documentation/Source/Resources/ContentExamples/Decals/Decals.INT.udn
+++ b/Engine/Documentation/Source/Resources/ContentExamples/Decals/Decals.INT.udn
@@ -53,9 +53,9 @@ To create a Deferred Decal Actor, drag and drop them from the **Modes** panel.
 This is only accessible from the **Modes** panel if you're in Placement mode.
 [/REGION]
 
-<!--You then only need to apply a Decal Material, which is really just a regular Material that has had its **Material Domain** property set to _Deferred Decal_.
+You then only need to apply a Decal Material, which is really just a regular Material that has had its **Material Domain** property set to **Deferred Decal** in the Material's Details panel.
 
-![](DeferredDecalSetting.png) - Commented out due to no longer being relevant, as decal actors are automatically Deferred Decals-->
+![](DeferredDecalSetting.png)
 
 [INCLUDE:Resources\ContentExamples\ExampleLevelOverview#ExampleLevel]
 
