#UE4doc #docs:
Removed Overview header

diff --git a/Engine/Documentation/Source/Resources/ContentExamples/Decals/1_2/1_2.INT.udn b/Engine/Documentation/Source/Resources/ContentExamples/Decals/1_2/1_2.INT.udn
--- a/Engine/Documentation/Source/Resources/ContentExamples/Decals/1_2/1_2.INT.udn
+++ b/Engine/Documentation/Source/Resources/ContentExamples/Decals/1_2/1_2.INT.udn
@@ -54,7 +54,7 @@ Related: Engine/Rendering/Materials/MaterialInstances
 
 ![%Resources/ContentExamples/Decals/1_1:title%](Resources\ContentExamples\Decals\Decals_2.png)
 
-## Overview
+
 
 This example demonstrates the available Blend Modes that Decals can use. Different Blend Modes change the way the Decal affects the surface. This is set using the **Decal Blend Mode** property found within a Material.
 
