#UE4doc #docs:
Removed Overview header

diff --git a/Engine/Documentation/Source/Resources/ContentExamples/Decals/1_3/1_3.INT.udn b/Engine/Documentation/Source/Resources/ContentExamples/Decals/1_3/1_3.INT.udn
--- a/Engine/Documentation/Source/Resources/ContentExamples/Decals/1_3/1_3.INT.udn
+++ b/Engine/Documentation/Source/Resources/ContentExamples/Decals/1_3/1_3.INT.udn
@@ -54,9 +54,9 @@ Related: Engine/Rendering/Materials/MaterialInstances
 
 ![%Resources/ContentExamples/Decals/1_1:title%](Resources\ContentExamples\Decals\Decals_3.png)
 
-## Overview
 
-In this example, we show how Decals can make use of both opacity masks and normal maps. The decal Material uses a **Decal Blend Mode** of _Translucent_, meaning that it can accept a value for transparency. By feeding in a texture - in this case, the Unreal Logo - we control the shape of the Decal. However, we also have the power to add a normal map, which gives the decal more of a solid feel as opposed to looking like paint on the projection surface. 
+
+In this example, we show how Decals can make use of both opacity masks and normal maps. The decal Material uses a **Decal Blend Mode** of **Translucent**, meaning that it can accept a value for transparency. By feeding in a texture - in this case, the Unreal Logo - we control the shape of the Decal. However, we also have the power to add a normal map, which gives the decal more of a solid feel as opposed to looking like paint on the projection surface. 
 
 This is a useful technique for creating certain kinds of details and decorations in your levels.
 
