#UE4doc #docs:
Removed Overview header

diff --git a/Engine/Documentation/Source/Resources/ContentExamples/Decals/1_4/1_4.INT.udn b/Engine/Documentation/Source/Resources/ContentExamples/Decals/1_4/1_4.INT.udn
--- a/Engine/Documentation/Source/Resources/ContentExamples/Decals/1_4/1_4.INT.udn
+++ b/Engine/Documentation/Source/Resources/ContentExamples/Decals/1_4/1_4.INT.udn
@@ -53,7 +53,7 @@ Related: Engine/Rendering/Materials/MaterialInstances
 
 ![%Resources/ContentExamples/Decals/1_1:title%](Resources\ContentExamples\Decals\Decals_4.png)
 
-## Overview
+
 
 
 In this example, you can see how Decals project over multiple surfaces. Notice that as surfaces are oriented more parallel to the projection direction, there is some expected stretching. Also, keep in mind that decals will project on to both static and dynamic meshes, including Skeletal Meshes.
