#UE4doc #docs:
Removed Overview header
Changed font settings for multiple terms

diff --git a/Engine/Documentation/Source/Resources/ContentExamples/Decals/1_5/1_5.INT.udn b/Engine/Documentation/Source/Resources/ContentExamples/Decals/1_5/1_5.INT.udn
--- a/Engine/Documentation/Source/Resources/ContentExamples/Decals/1_5/1_5.INT.udn
+++ b/Engine/Documentation/Source/Resources/ContentExamples/Decals/1_5/1_5.INT.udn
@@ -54,14 +54,14 @@ Related: Engine/Rendering/Materials/MaterialInstances
 
 ![%Resources/ContentExamples/Decals/1_1:title%](Resources\ContentExamples\Decals\Decals_5.png)
 
-## Overview
 
 
-This example shows how you can turn off Decals on specific meshes with the **Receive Decals** property under the Rendering category. Note that you will have to click the expansion arrow for the category to see this property.
+
+This example shows how you can turn off Decals on specific meshes with the **Receives Decals** property under the **Rendering** category. Note that you will have to click the expansion arrow for the category to see this property.
 
 ![](ReceivesDecals.png)
 
-You can view which meshes are not receiving decals with the **Decal Mask** View Mode found under Buffer Visualization.
+You can view which meshes are not receiving decals with the **Decal Mask** View Mode found under **Buffer Visualization**.
 
 ![](DecalMaskViewMode.png)
 
