#UE4doc #docs:
Fixed grammar problem w/ something being plural that shouldn't
Removed unnecessary italics

diff --git a/Engine/Documentation/Source/Resources/ContentExamples/Decals/1_6/1_6.INT.udn b/Engine/Documentation/Source/Resources/ContentExamples/Decals/1_6/1_6.INT.udn
--- a/Engine/Documentation/Source/Resources/ContentExamples/Decals/1_6/1_6.INT.udn
+++ b/Engine/Documentation/Source/Resources/ContentExamples/Decals/1_6/1_6.INT.udn
@@ -56,7 +56,7 @@ Related: Engine/Rendering/Materials/MaterialInstances
 
 ![%Resources/ContentExamples/Decals/1_1:title%](Resources\ContentExamples\Decals\Decals_6.png)
 
-This examples shows how Decals can be assigned a sort priority, allowing certain Decals to project on top of others. You can control the order in which Decals are drawn when they are layered on each other using the **Sort Order** property. The _higher_ the sort order value will bring that decal to the forefront.
+This example shows how Decals can be assigned a sort priority, allowing certain Decals to project on top of others. You can control the order in which Decals are drawn when they are layered on each other using the **Sort Order** property. The higher the sort order value will bring that decal to the forefront.
 
 ![](DecalSortOrder.png)
 
