#UE4 Docs:Commenting out the Steam VR section

diff --git a/Engine/Documentation/Source/Platforms/Platforms.INT.udn b/Engine/Documentation/Source/Platforms/Platforms.INT.udn
--- a/Engine/Documentation/Source/Platforms/Platforms.INT.udn
+++ b/Engine/Documentation/Source/Platforms/Platforms.INT.udn
@@ -89,8 +89,17 @@ platforms
 			* [](Platforms/Oculus/BestPractices "%Platforms/Oculus/BestPractices:description%")
 			* [](Platforms/Oculus/QuickStart "%Platforms/Oculus/QuickStart:description%")
 		* [](Platforms/GearVR "%Platforms/GearVR:Description%")
-			* [](Platforms/GearVR/QuickStart "%Platforms/GearVR:Description%")
-			* [](Platforms/GearVR/BestPractices "%Platforms/GearVR:Description%")
+			* [](Platforms/GearVR/Prerequisites "%Platforms/GearVR/Prerequisites:Description%")
+			* [](Platforms/GearVR/QuickStart "%Platforms/GearVR/QuickStart:Description%")
+			* [](Platforms/GearVR/BestPractices "%Platforms/GearVR/BestPractices:Description%")
+			* [](Platforms/GearVR/Debugging "%Platforms/GearVR/Debugging:Description%")
+
+<!--
+		* [](Platforms/SteamVR "%Platforms/SteamVR:Description%")
+			* [](Platforms/SteamVR/BestPractices "%Platforms/GearVR:Description%")
+			* [](Platforms/SteamVR/QuickStart "%Platforms/GearVR:Description%")
+			* [](Platforms/SteamVR/Tps&Tricks "%Platforms/SteamVR/Tps&Tricks:description%")
+-->
 	[/PARAM]
 [/OBJECT]
 
@@ -136,4 +145,4 @@ platforms
 		* [](Engine/Deployment/UnrealFrontend "%Engine/Deployment/UnrealFrontend:description%")
 		* [](Platforms/DeviceProfiles "%Platforms/DeviceProfiles:description%")
 	[/PARAM]
-[/OBJECT]
+[/OBJECT]
\ No newline at end of file
