#UE4doc #docs:
Removed Overview header
Changed font settings for multiple terms

diff --git a/Engine/Documentation/Source/Resources/ContentExamples/Destructables/Destructables.INT.udn b/Engine/Documentation/Source/Resources/ContentExamples/Destructables/Destructables.INT.udn
--- a/Engine/Documentation/Source/Resources/ContentExamples/Destructables/Destructables.INT.udn
+++ b/Engine/Documentation/Source/Resources/ContentExamples/Destructables/Destructables.INT.udn
@@ -46,23 +46,23 @@ Related:
 [TOC(start:2 end:2)]
 
 
-## Overview
 
-The Content Examples provided inside the **Destructibles** map showcase how destructible meshes can be created and used in Unreal Engine 4. Automatically generated and custom imported destructibles, fracturing controls with **Support** and **Destruction Depth** levels as well as how to fire events upon destruction like spawning particles or playing sounds are demonstrated in this map. 
+
+The Content Examples provided inside the **Destructibles** map showcase how destructible meshes can be created and used in **Unreal Engine 4**. Automatically generated and custom imported destructibles, fracturing controls with **Support** and **Destruction Depth** levels as well as how to fire events upon destruction like spawning particles or playing sounds are demonstrated in this map. 
 
 
 ## Destructibles Map
 
 ![](Destructible_Topic.png)(w:640)
 
-Listed below are the examples provided inside the **Destructibles** map:
+Listed below are the examples provided inside the Destructibles map:
 
 | Example | What is Demonstrated |
 | ------- | ------------ |
 | **1.1 Voronoi** | Automatically generated destructible mesh. |
 | **1.2 Custom Destruction (Imported FBX)** | Demonstrates an imported custom FBX destructible mesh. |
 | **1.3 Support** | Custom fracture a portion of a mesh without fracturing the enitre mesh. |
-| **1.4 Destruction Depth** | The complexity of the fracture as defined by **Destruction Depth**. |
+| **1.4 Destruction Depth** | The complexity of the fracture as defined by Destruction Depth. |
 | **1.5 WorldSupport** |  How to attach a destructible mesh to a piece of static world geometry.  |
 | **1.6 Destruction Particle Effects** | How to spawn a particle effect when a destructible mesh gets destroyed. |
 | **1.7 Destruction Sounds** | Playing sounds attached to a destrucbtible mesh. |
