#UE4doc #docs:
Removed Overview header
Changed font settings for multiple terms

diff --git a/Engine/Documentation/Source/Resources/ContentExamples/DynamicSceneShadows/DynamicSceneShadows.INT.udn b/Engine/Documentation/Source/Resources/ContentExamples/DynamicSceneShadows/DynamicSceneShadows.INT.udn
--- a/Engine/Documentation/Source/Resources/ContentExamples/DynamicSceneShadows/DynamicSceneShadows.INT.udn
+++ b/Engine/Documentation/Source/Resources/ContentExamples/DynamicSceneShadows/DynamicSceneShadows.INT.udn
@@ -58,13 +58,13 @@ Related: Engine/Rendering/LightingAndShadows/Shadows
 
 [TOC(start:2 end:2)]
 
-## Overview
+
 The **Dynamic Scene Shadows** Content Example map illustrates how to create dynamic shadowing effects with baked lighting through the use of **Cascading Shadow Mapping** which is a shadowing method specifically for use with **Directional Light Actors**.
 They are superior to baked shadows in that they allow for shadowing of dynamic objects, such as trees with leaves that sway in the wind creating a much more realistic scene.
 
 ## How Cascading Shadow Maps Work
 
-The **Cascading Shadow Map** method divides the view frustum into a series of distance-based shadow cascades, each of which with steadily lower resolution as you move farther from the camera. 
+The Cascading Shadow Map method divides the view frustum into a series of distance-based shadow cascades, each of which with steadily lower resolution as you move farther from the camera. 
 This means that shadows nearest the player's view have the highest resolution while those farthest away have the lowest. 
 Beyond the range of the **Dynamic Shadow Distance** property, the system blends back into static baked shadows. In this way, you can take a scene with baked lighting and smoothly transition it to static baked shadows as you move away from objects.
 
@@ -78,15 +78,15 @@ In the image above, notice that the dynamic shadows (nearer the camera) fade awa
 [/REGION]
 ## Cascaded Shadow Map Properties
 
-In the **Details** panel for a **Directional Light Actor**, you will find the **Cascaded Shadow Map** category. Within are the following properties.
+In the **Details** panel for a Directional Light Actor, you will find the **Cascaded Shadow Map** category. Within are the following properties.
 
 | **Property** | **Description** |
 |--------|---|
-| **Dynamic Shadow Distance (MovableLight and StationaryLight)** | This controls the distance from the camera within which you will see cascading shadow maps. Beyond this, you will only see precalculated shadow maps. A setting of 0 effectively disables the dynamic shadow. Note that there is a setting for **Movable Lights** and **Stationary Lights**.  |
+| **Dynamic Shadow Distance (MovableLight and StationaryLight)** | This controls the distance from the camera within which you will see cascading shadow maps. Beyond this, you will only see precalculated shadow maps. A setting of 0 effectively disables the dynamic shadow. Note that there is a setting for Movable Lights and Stationary Lights.  |
 | **Num Dynamic Shadow Cascades** | This controls how many sections of cascading into which to divide the dynamic shadow distance. More levels provide better shadow resolution at distance, but at a greater performance cost.  |
-| **Cascade Distribution Exponent** | This controls how near the camera transitions between cascades can occur. This value is applied as a fraction of the **Dynamic Shadow Distance**. A value of 1 will cause the transitions to take place at distances proportional to their resolution. Values greater than 1 will pull the transition closer to the camera. |
+| **Cascade Distribution Exponent** | This controls how near the camera transitions between cascades can occur. This value is applied as a fraction of the Dynamic Shadow Distance. A value of 1 will cause the transitions to take place at distances proportional to their resolution. Values greater than 1 will pull the transition closer to the camera. |
 | **Cascade Transition Fraction** | This controls the proportion of the fade from one cascade to the next. As values approach 0, the transition tightens. At  0, there is no transition, leading to hard edges between cascades. |
 | **Shadow Distance Fadeout Fraction** | This softens the transition into cascaded shadow maps as a whole, with 0 being a hard transition and 1 being completely softened. A value of 0.1 is generally good. There is no performance cost to this value. |
-| **Use Inset Shadows for Movable Objects** | This option allows for inset shadows to be used on movable objects lit by stationary lights. This is useful when movable objects move beyond the range of the **Dynamic Shadow Distance**, particularly when the **Dynamic Shadow Distance** is set to a low value. When **Dynamic Shadow Distance** is very large (currently > 8000), this option is forced off. |
+| **Use Inset Shadows for Movable Objects** | This option allows for inset shadows to be used on movable objects lit by stationary lights. This is useful when movable objects move beyond the range of the Dynamic Shadow Distance, particularly when the Dynamic Shadow Distance is set to a low value. When Dynamic Shadow Distance is very large (currently > 8000), this option is forced off. |
 
 [/REGION]
