#UE4doc #docs:
Removed Overview header
Changed font settings for multiple terms

diff --git a/Engine/Documentation/Source/Resources/ContentExamples/EffectsGallery/EffectsGallery.INT.udn b/Engine/Documentation/Source/Resources/ContentExamples/EffectsGallery/EffectsGallery.INT.udn
--- a/Engine/Documentation/Source/Resources/ContentExamples/EffectsGallery/EffectsGallery.INT.udn
+++ b/Engine/Documentation/Source/Resources/ContentExamples/EffectsGallery/EffectsGallery.INT.udn
@@ -48,7 +48,7 @@ Related: Engine/Rendering/ParticleSystems/Reference/Modules/VectorField
 [TOC(start:2 end:2)]
 
 
-## Overview
+
 The **Effects** map contains several visual effects and examples of Particle Systems that can help add increased realism or assist in communicating gameplay elements like hazards, a fire or explosion for example, to players in your levels. 
 
 This map shows how to create moving particles, particles that have collision, particles that can spawn other particles, particles that can be attached to a player and several others. Each example is setup in a way so that it can be deconstructed so that you can alter or create your own effect based on the example. 
@@ -70,7 +70,7 @@ Listed below are the examples provided inside the **Effects** map:
 | Example | What is Demonstrated |
 | ------- | ------------ |
 | **[1.1 CPU and GPU particles](Resources\ContentExamples\EffectsGallery\1_A)** | The difference between CPU and GPU particles along with pros/cons are demonstrated with two simple particle fountains. |
-| **[1.2 GPU particles with Velocity Cone module](Resources\ContentExamples\EffectsGallery\1_B)** | The use of the **Velocity Cone** module which outputs particles in a cone shape. This example also shows how a **Static Mesh** can be added to a particle through the **Mesh TypeData** module (e.g. two emitters output sparks while another outputs a **Static Mesh**). |
+| **[1.2 GPU particles with Velocity Cone module](Resources\ContentExamples\EffectsGallery\1_B)** | The use of the **Velocity Cone** module which outputs particles in a cone shape. This example also shows how a **Static Mesh** can be added to a particle through the **Mesh TypeData** module (e.g. two emitters output sparks while another outputs a Static Mesh). |
 | **[1.3 GPU particles with Emitter Initial Location](Resources\ContentExamples\EffectsGallery\1_C)** | Three spheres travel around a central point while leaving trails behind them. This example also shows how to spawn a particle from the location of other particles with the **Emitter Initial Location** module.  |
 | **[1.4 GPU particles with Point Gravity](Resources\ContentExamples\EffectsGallery\1_D)** | Shows how to specify a point in space where all particles inside an emitter will be drawn to through the **Point Gravity** module. |
 | **[1.5 GPU particles with collision](Resources\ContentExamples\EffectsGallery\1_E)** | Demonstrates particles that collide with the environment and have collision (e.g. a downpour of rain bouncing off objects and falling to the ground). |
