#UE4 Docs: Making this doc public and added a new section that talks about the differance between dev and shipping builds.

diff --git a/Engine/Documentation/Source/Platforms/Android/ReducingAPKSize/ReducingAPKSizes.INT.udn b/Engine/Documentation/Source/Platforms/Android/ReducingAPKSize/ReducingAPKSizes.INT.udn
--- a/Engine/Documentation/Source/Platforms/Android/ReducingAPKSize/ReducingAPKSizes.INT.udn
+++ b/Engine/Documentation/Source/Platforms/Android/ReducingAPKSize/ReducingAPKSizes.INT.udn
@@ -1,4 +1,4 @@
-﻿Availability:Docs
+﻿Availability:Public
 Title: Reducing APK Package size 
 Crumbs:%ROOT%, Platforms, Platforms/Android
 Description:How to reduce the size of your packaged Android game.
@@ -45,7 +45,7 @@ Please note that the Create compressed cooked packages option is **enabled** by
 1. Look for the option labeled **Create compressed cooked packages** and enable it if not already enabled.
 
 	![](T_Enable_Comp_Packages.png)
-
+	
 If you have not packaged your game with the Create compressed cooked packages checkmark box enabled you should see a huge difference in size when your project is re-packaged with it enabled.
 In fact it is not uncommon for some projects APK packages to be reduced in size by up to 50% when create compressed cooked packages is enabled.
 
@@ -77,6 +77,12 @@ Once the biggest assets have been established you can then examine the assets in
 
 ![](T_Cooked_Content.png)
 
+## Development VS Shipping Build Size
+
+When trying to figure out the final size your project will be keep in mind that the Development version of your project will be slightly larger than your Shipping build size.
+On the Medieval Match example game the difference in size between a Development and Shipping build was around **14 percent**.
+However since each project has different requirements, the saving between the two different builds types for your project could be more or less than a 14 percent savings.
+
 ## Per-Device Texture LOD's
 With the release of Unreal Engine 4.8 you can now specify which size Texture should be used on which devices. 
 You can read more about this in the [Texture Guidelines for Mobile Platforms](Platforms/Mobile/Textures) section.
