#UE4 Docs: Fixing some spelling mistakes to make the doc more clear on what to do.

diff --git a/Engine/Documentation/Source/Platforms/GearVR/GlobalMenu/3/GlobalMenu_HUDBPSetup.INT.udn b/Engine/Documentation/Source/Platforms/GearVR/GlobalMenu/3/GlobalMenu_HUDBPSetup.INT.udn
--- a/Engine/Documentation/Source/Platforms/GearVR/GlobalMenu/3/GlobalMenu_HUDBPSetup.INT.udn
+++ b/Engine/Documentation/Source/Platforms/GearVR/GlobalMenu/3/GlobalMenu_HUDBPSetup.INT.udn
@@ -83,8 +83,8 @@ With the **GearVR_GlobalMenu** HUD now assigned as the game HUD, we can now star
 
 	![](T_Final_Graph.png)
 
-1. Now you can package your project up and deploy it to your Smartphone to use with the Gear VR HMD. 
-	To see the menu in action, launch your project and then hold down the back button that is on the side of Gear VR HMD.
+1. Now you can package your project up and deploy it to your Smartphone to use with the Samsung Gear VR HMD. 
+	To see the menu in action, launch your project and then hold down the back button that is on the side of Samsung Gear VR HMD.
 	As you do you will see a line being drawn across the screen and then when it finishes the Global Menu will be displayed.
 
 
diff --git a/Engine/Documentation/Source/Platforms/GearVR/GlobalMenu/GearVRGlobalMenu.INT.udn b/Engine/Documentation/Source/Platforms/GearVR/GlobalMenu/GearVRGlobalMenu.INT.udn
--- a/Engine/Documentation/Source/Platforms/GearVR/GlobalMenu/GearVRGlobalMenu.INT.udn
+++ b/Engine/Documentation/Source/Platforms/GearVR/GlobalMenu/GearVRGlobalMenu.INT.udn
@@ -14,7 +14,7 @@ In the following Quick Start we will cover everything you need to know so the us
 
 [region:note]
 Please note that the following Quick Start will work on any project that is setup for Gear VR. 
-If you have no done so already please make sure to check out the Gear [VR Quick Start](Platforms\GearVR\QuickStart) for detailed information on how to get a project up and running with Gear VR.
+If you have not done so already please make sure to check out the Gear [VR Quick Start](Platforms\GearVR\QuickStart) for detailed information on how to get a project up and running with Gear VR.
 [/region]
 
 [OBJECT:TopicButtonList]
