#UE4doc #docs:
Bolded a term

diff --git a/Engine/Documentation/Source/Resources/ContentExamples/EffectsGallery/1_A/1_A.INT.udn b/Engine/Documentation/Source/Resources/ContentExamples/EffectsGallery/1_A/1_A.INT.udn
--- a/Engine/Documentation/Source/Resources/ContentExamples/EffectsGallery/1_A/1_A.INT.udn
+++ b/Engine/Documentation/Source/Resources/ContentExamples/EffectsGallery/1_A/1_A.INT.udn
@@ -54,7 +54,7 @@ Related: Engine/Rendering/ParticleSystems/Reference/Modules/VectorField
 
 
 
-One of the particle types available in Unreal Engine 4 is **GPU Sprites**. These are particles that are first spawned on the CPU, but then processed and calculated entirely by the graphics card. The benefit to this is that since the GPU is handling the calculation, many thousands more particles can be processed at once, allowing for much denser and more detailed particle systems.
+One of the particle types available in **Unreal Engine 4** is **GPU Sprites**. These are particles that are first spawned on the CPU, but then processed and calculated entirely by the graphics card. The benefit to this is that since the GPU is handling the calculation, many thousands more particles can be processed at once, allowing for much denser and more detailed particle systems.
 
 This first effect shows off two very simple particle fountains, one created via standard CPU particle sprites and the other via GPU sprites. Notice that the fountain on the right - which uses GPU particles - is outputting significantly more particles than the other. 
 
