#UE4doc #docs: ##Overview removed as according to the instructional guide.

diff --git a/Engine/Documentation/Source/Engine/Rendering/Materials/Functions/Reference/VectorOps/VectorOps.INT.udn b/Engine/Documentation/Source/Engine/Rendering/Materials/Functions/Reference/VectorOps/VectorOps.INT.udn
--- a/Engine/Documentation/Source/Engine/Rendering/Materials/Functions/Reference/VectorOps/VectorOps.INT.udn
+++ b/Engine/Documentation/Source/Engine/Rendering/Materials/Functions/Reference/VectorOps/VectorOps.INT.udn
@@ -5,8 +5,6 @@ Description:Functions build to handle vector mathematics, such as calculating a
 
 [TOC (start:2 end:3)]
 
-## Overview
-
 The VectorOps category contains special functions for applying various vector-based math equations.
 
 ## VectorOps Functions
