#UE4doc #docs: fixed grammar error: use as your both your local axis -> use as both your local axis

diff --git a/Engine/Documentation/Source/Engine/Rendering/Materials/Functions/Reference/WorldPositionOffset/WorldPositionOffset.INT.udn b/Engine/Documentation/Source/Engine/Rendering/Materials/Functions/Reference/WorldPositionOffset/WorldPositionOffset.INT.udn
--- a/Engine/Documentation/Source/Engine/Rendering/Materials/Functions/Reference/WorldPositionOffset/WorldPositionOffset.INT.udn
+++ b/Engine/Documentation/Source/Engine/Rendering/Materials/Functions/Reference/WorldPositionOffset/WorldPositionOffset.INT.udn
@@ -97,7 +97,7 @@ This shader node only supports uniform scaling. Also, the rotate about axis axis
 | Item | Description |
 | -------- | -------- |
 |[REGION:tablesection]Inputs[/REGION] ||
-| **Pivot Axis/Pos (Vector3)** | Enter a number to use as your both your local axis line axis and position. If you would like to lock off your model, enter on its local z axis near its pivot enter (0,0,1). If you would like to lock of the top of the model, enter your model's height in as (0,0,model height).|
+| **Pivot Axis/Pos (Vector3)** | Enter a number to use as both your local axis line axis and position. If you would like to lock off your model, enter on its local z axis near its pivot enter (0,0,1). If you would like to lock of the top of the model, enter your model's height in as (0,0,model height).|
 |[REGION:tablesection]Outputs[/REGION] ||
 | **Pivot** | This can be used as the pivot point in the rotate about axis node.|
 
