#UE4doc #docs: Changed title of How To to match instructional guide.

diff --git a/Engine/Documentation/Source/Engine/Rendering/Materials/HowTo/Fresnel/Fresnel.INT.udn b/Engine/Documentation/Source/Engine/Rendering/Materials/HowTo/Fresnel/Fresnel.INT.udn
--- a/Engine/Documentation/Source/Engine/Rendering/Materials/HowTo/Fresnel/Fresnel.INT.udn
+++ b/Engine/Documentation/Source/Engine/Rendering/Materials/HowTo/Fresnel/Fresnel.INT.udn
@@ -1,5 +1,5 @@
 Availability: Public
-Title:Material - How To Use Fresnel in your Materials
+Title:Material - Using Fresnel in your Materials
 Crumbs: %ROOT%, Engine, Engine/Rendering/Materials
 Description:A How To Guide for using the Fresnel Material node. 
 Related: Engine/Rendering/Materials/Functions/Reference/VectorOps
