#UE4doc #docs: Made grammar changes: the the -> the, can not -> cannot. Changed title of How To to match instructional guide.

diff --git a/Engine/Documentation/Source/Engine/Rendering/Materials/HowTo/Making_Parameters/Making_Parameters.INT.udn b/Engine/Documentation/Source/Engine/Rendering/Materials/HowTo/Making_Parameters/Making_Parameters.INT.udn
--- a/Engine/Documentation/Source/Engine/Rendering/Materials/HowTo/Making_Parameters/Making_Parameters.INT.udn
+++ b/Engine/Documentation/Source/Engine/Rendering/Materials/HowTo/Making_Parameters/Making_Parameters.INT.udn
@@ -1,5 +1,5 @@
 Availability: Public
-Title:Material Editor - How To make Material Parameters
+Title:Material Editor - Making Material Parameters
 Crumbs: %ROOT%, Engine, Engine/Rendering/Materials, Engine/Rendering/Materials/HowTo
 Description:A How To Guide for Setting Up and using Material Parameters
 Related: Engine/Rendering/Materials/ExpressionReference/Parameters
diff --git a/Engine/Documentation/Source/Engine/Rendering/Materials/HowTo/Subsurface_Profile/Subsurface_Profile.INT.udn b/Engine/Documentation/Source/Engine/Rendering/Materials/HowTo/Subsurface_Profile/Subsurface_Profile.INT.udn
--- a/Engine/Documentation/Source/Engine/Rendering/Materials/HowTo/Subsurface_Profile/Subsurface_Profile.INT.udn
+++ b/Engine/Documentation/Source/Engine/Rendering/Materials/HowTo/Subsurface_Profile/Subsurface_Profile.INT.udn
@@ -1,5 +1,5 @@
 Availability:Public
-Title: Material Editor - How To use a Subsurface Profile in your Materials 
+Title: Material Editor - Using a Subsurface Profile in your Materials 
 Crumbs: %ROOT%, Engine, Engine/Rendering/Materials
 Description: Setting up Materials to use Subsurface Profiles as they offer a better way to render skin and wax like surfaces.
 Related: Engine/Rendering/Materials/LightingModels/SubSurfaceProfile
@@ -48,7 +48,7 @@ Creating Subsurface Profiles is something that can be accomplished in the follow
 
 	![](SSP_Creating_Naming_SSP.png)
 
-## Creating a Material the uses a Subsurface Profile
+## Creating a Material that uses a Subsurface Profile
 
 Now that the Material and Subsurface Profile have been created, it is now time to fill them in with data. 
 In the following steps, we will be creating a Material and a Material Instance that can make use of a Subsurface Profile.
@@ -199,7 +199,7 @@ Combining a Material Instance with a Subsurface profile is a great way to get th
  control over your results for a number of reasons.
 
 * A Material Instance allows you to tweak values in real time which makes seeing your results much quicker. 
-* You can use the **Opacity** input  to dampen the Subsurface contribution to the surface. However it is recommended that you keep this value at 1 and adjust the parameters inside the the Subsurface Profile. Only adjust this value if you find that you can not get the results you are looking for when adjusting the Subsurface Profile.
+* You can use the **Opacity** input  to dampen the Subsurface contribution to the surface. However it is recommended that you keep this value at 1 and adjust the parameters inside the Subsurface Profile. Only adjust this value if you find that you cannot get the results you are looking for when adjusting the Subsurface Profile.
 * The **Value** slider of each color picker controls the influence of the Subsurface effect. The closer that this value is set toward white, the more prominent the effect will be. While the closer the value is set to black, the less visible the effect will be.
 
 
