#UE4doc #docs: ##Overview removed as according to the instructional guide.

diff --git a/Engine/Documentation/Source/Engine/UI/OutputLog/OutputLog.INT.udn b/Engine/Documentation/Source/Engine/UI/OutputLog/OutputLog.INT.udn
--- a/Engine/Documentation/Source/Engine/UI/OutputLog/OutputLog.INT.udn
+++ b/Engine/Documentation/Source/Engine/UI/OutputLog/OutputLog.INT.udn
@@ -22,5 +22,3 @@ Description:Overview and usage of the Output Log.
 	[/PARAM]
 [/OBJECT]
 [/VAR]
-
-## Overview
