UEDOC-97

Image has been replaced and gif has been converted into a youtube embedded video

diff --git a/Engine/Documentation/Source/Engine/Animation/Sequences/Curves/Curves.INT.udn b/Engine/Documentation/Source/Engine/Animation/Sequences/Curves/Curves.INT.udn
--- a/Engine/Documentation/Source/Engine/Animation/Sequences/Curves/Curves.INT.udn
+++ b/Engine/Documentation/Source/Engine/Animation/Sequences/Curves/Curves.INT.udn
@@ -18,7 +18,17 @@ In this example, we will show how to set up a simple Material Curve that causes
 
 We will start with a simple animation in which our character is lunging forward with a hammer. We will then set up a simple Material curve that will cause his face to blush heavily while he does this. To begin, let us look at the animation:
 
-![](Engine/Animation/RootMotion\HammerSlam.gif)(convert:false)
+[OBJECT:EmbeddedVideo]
+[PARAMLITERAL:width]
+640
+[/PARAMLITERAL]
+[PARAMLITERAL:height]
+360
+[/PARAMLITERAL]
+[PARAMLITERAL:videoid]
+TKAcNubFOH8
+[/PARAMLITERAL]
+[/OBJECT]
 
 We can open that Animation Sequence in Persona and look at it by **double-click** it in the **Content Browser**.
 
@@ -29,7 +39,7 @@ We can open that Animation Sequence in Persona and look at it by **double-click*
 The character's head has a simple Material with a parameter named _Blush_. Here is a look at the Material Graph:
 
 [REGION:fullwidth]
-![](BlushMaterial.png)
+![](BlushMaterial.jpg)
 [/REGION]
 
 [REGION:note]
