#UE4doc #docs:
Changed font settings for multiple terms

diff --git a/Engine/Documentation/Source/Resources/ContentExamples/EffectsGallery/1_B/1_B.INT.udn b/Engine/Documentation/Source/Resources/ContentExamples/EffectsGallery/1_B/1_B.INT.udn
--- a/Engine/Documentation/Source/Resources/ContentExamples/EffectsGallery/1_B/1_B.INT.udn
+++ b/Engine/Documentation/Source/Resources/ContentExamples/EffectsGallery/1_B/1_B.INT.udn
@@ -56,7 +56,7 @@ Related: Engine/Rendering/ParticleSystems/Reference/Modules/VectorField
 [INCLUDE:Resources\ContentExamples\ExampleLevelOverview#PlaySim]
 
 
-This next example demonstrates the use of the Velocity Cone module, which outputs particles in a cone shape. The Particle System actually contains 3 Emitters:
+This next example demonstrates the use of the **Velocity Cone** module, which outputs particles in a cone shape. The Particle System actually contains 3 Emitters:
 
 * 2 Emitters that output cones of sparks - one that fires upward and another that fires downward.
 * 1 Emitter that outputs a single Static Mesh by way of a Mesh TypeData module.
@@ -65,7 +65,7 @@ This next example demonstrates the use of the Velocity Cone module, which output
 ![](VelocityConeEmitters.png)
 [/REGION]
 
-The first two Emitters are GPU Sprite emitters. The most important modules in the effect are the _Velocity Cone_ module, which handles the shape of the emission. There is also a _Constant Acceleration_ module on each one which provides a sense of gravity, as well as a _Collision (Scene Depth)_ module that uses Z-Buffer information to allow collision simulation.
+The first two Emitters are GPU Sprite emitters. The most important modules in the effect are the Velocity Cone module, which handles the shape of the emission. There is also a Constant Acceleration module on each one which provides a sense of gravity, as well as a Collision (Scene Depth) module that uses Z-Buffer information to allow collision simulation.
 
 The third Emitter uses a Mesh TypeData module to emit an instance of a Static Mesh into the Particle System, in this case a simple sphere. The Emitter's Lifetime is set to 0, which effectively causes the particle to live forever.
 
@@ -75,7 +75,7 @@ For more information on Velocity modules, please see the [Velocity Modules docum
 
 ### Blueprints and Particles
 
-An interesting side note regarding this example is that it is more than just a Particle System. In the level, the effect has been added by way of a Blueprint (_BP_sway_cube_). The Blueprint simply has the Particle System as a Component, and uses a simple node graph that defines a looping rotation animation. This causes the effect to slowly spin the Blueprint Actor. This is an excellent example of how Blueprints can be utilized to add a bit more life into your particle-based effects. Note that this is one of many possible ways to achieve a rotation effect.
+An interesting side note regarding this example is that it is more than just a Particle System. In the level, the effect has been added by way of a Blueprint (BP_sway_cube). The Blueprint simply has the Particle System as a Component, and uses a simple node graph that defines a looping rotation animation. This causes the effect to slowly spin the Blueprint Actor. This is an excellent example of how Blueprints can be utilized to add a bit more life into your particle-based effects. Note that this is one of many possible ways to achieve a rotation effect.
 
 [REGION:fullwidth]
 ![](SwayCube.png)
