#UE4doc #docs - Adjusted bold lettering and bolded links

diff --git a/Engine/Documentation/Source/Engine/Landscape/Editing/PaintMode/PaintMode.INT.udn b/Engine/Documentation/Source/Engine/Landscape/Editing/PaintMode/PaintMode.INT.udn
--- a/Engine/Documentation/Source/Engine/Landscape/Editing/PaintMode/PaintMode.INT.udn
+++ b/Engine/Documentation/Source/Engine/Landscape/Editing/PaintMode/PaintMode.INT.udn
@@ -27,7 +27,7 @@ Description:Guide to the Landscape tool's Paint mode.
 
 The tools in **Paint** mode enable you to modify the appearance of your Landscape by selectively applying Material layers to parts of your Landscape.
 
-For more information about Landscape Materials, see [](Engine/Landscape\Materials).
+For more information about Landscape Materials, see [**Landscape Materials**](Engine/Landscape\Materials).
 
 
 ## Painting Tools
@@ -49,7 +49,7 @@ The Painting Tools have some common options:
 | **Negative Mask** | When checked, and when **Use Region as Mask** is also checked, the region selection acts as a mask, but the active area is comprised of the unselected region. |
 
 
-For more information about Landscape Material layers, see [Layers](#Layers), later on this page.
+For more information about Landscape Material layers, see [**Layers**](#Layers), later on this page.
 
 
 ### Paint
@@ -121,7 +121,7 @@ The layers defined in the Landscape Material automatically populate the list of
 
 ![Landscape_Target.png](Landscape_Target.png)(convert:false)
 
-Whichever layer is selected is the one you can apply to the Landscape with the Painting Tools, according to the tools' options and settings, and to the [brush](Engine\Landscape\Editing\Brushes) you are using.
+Whichever layer is selected is the one you can apply to the Landscape with the Painting Tools, according to the tools' options and settings, and to the [**brush**](Engine\Landscape\Editing\Brushes) you are using.
 
 ![Landscape_Layers.jpg](Landscape_Layers.jpg)(w:900 h:389)
 
@@ -129,7 +129,7 @@ Whichever layer is selected is the one you can apply to the Landscape with the P
 Many of the Painting Tools are similar to the Sculpting Tools, and you use them similarly, but to manipulate the application of layers instead of the heightmap.
 [/REGION]
 
-You create layers in the Material itself. For more information about layers and Landscape Materials, see [](Engine\Landscape\Materials).
+You create layers in the Material itself. For more information about layers and Landscape Materials, see [**Landscape Materials**](Engine\Landscape\Materials).
 
 
 ### Layer Info Objects
@@ -177,8 +177,8 @@ The primary purpose of layer info objects is to act as a unique key for painted
 
 | Option | Description |
 | ------ | ----------- |
-| **Phys Material** | The [Physical Material](Engine/Physics/PhysicalMaterials) (if any) assigned to areas of Landscape where this layer is dominant. |
-| **Hardness** | The value used by the [Erosion](Engine/Landscape/Editing/SculptMode/#Erosion) tool. |
+| **Phys Material** | The [**Physical Material**](Engine/Physics/PhysicalMaterials) (if any) assigned to areas of Landscape where this layer is dominant. |
+| **Hardness** | The value used by the [**Erosion**](Engine/Landscape/Editing/SculptMode/#Erosion) tool. |
 
 ### Orphaned Layers
 
