UEDOC-2131

Converted all gifs into YouTube embeded videos

diff --git a/Engine/Documentation/Source/Engine/Rendering/Materials/HowTo/Making_Parameters/Making_Parameters.INT.udn b/Engine/Documentation/Source/Engine/Rendering/Materials/HowTo/Making_Parameters/Making_Parameters.INT.udn
--- a/Engine/Documentation/Source/Engine/Rendering/Materials/HowTo/Making_Parameters/Making_Parameters.INT.udn
+++ b/Engine/Documentation/Source/Engine/Rendering/Materials/HowTo/Making_Parameters/Making_Parameters.INT.udn
@@ -29,7 +29,17 @@ You can create Material Parameters in the following ways inside the Material Edi
 
 The simplest way to make a Material Parameter out of a Material Expression is to **Right-Click** on an existing Material Expression node and then choose the **Convert to Parameter** option from the menu if available. Once converted, you can then name the Material Parameter and give it a default value. To make things better organized, you can also give your Material Parameter a **Group**. A **Group** will group Parameters with the same group inside of Material Instances, making it easier to see Parameters that go together. 
 
-![](Making_Parameters_Right_Click_Convert.gif)(convert:false)
+[OBJECT:EmbeddedVideo]
+[PARAMLITERAL:width]
+640
+[/PARAMLITERAL]
+[PARAMLITERAL:height]
+360
+[/PARAMLITERAL]
+[PARAMLITERAL:videoid]
+lbCuDpVSfUA
+[/PARAMLITERAL]
+[/OBJECT]
 
 [region:note]
 Not all Material Expressions can be converted into Material Parameters. The best way to see if a Material Expression can be converted into a Material Parameter is to **Right-Click** on it and check to see if the **Convert to Parameter** option is available. If you do not see the **Convert to Parameter** option, then the Material Expression cannot be converted.
@@ -42,7 +52,17 @@ To do this, first search for the word **Parameters** in the Palette search menu.
 Then select the Material Parameter you want to use with the **Left Mouse Button** and with the **Left Mouse Button** still held down, Drag the Material Parameter you want into the Material Graph, 
 releasing the **Left Mouse Button** when your mouse pointer is over the Material Graph to add the Material Expression node in the Material Graph.
 
-![](Making_Parameters_Adding_From_Palette.gif)(convert:false)
+[OBJECT:EmbeddedVideo]
+[PARAMLITERAL:width]
+640
+[/PARAMLITERAL]
+[PARAMLITERAL:height]
+360
+[/PARAMLITERAL]
+[PARAMLITERAL:videoid]
+8ajYkU0Qmj4
+[/PARAMLITERAL]
+[/OBJECT]
 
 ### Adding Parameters from the Right-Click Menu
 
@@ -50,7 +70,17 @@ You can also add Material Parameter Expression nodes to the Material Graph by **
 Once the search menu has been displayed, search for parameters using the word **Parameters** as the search term.
 Once the Material Parameter node you want to use has been located, select the Material Parameter Expression node from the search menu with the **Left Mouse Button** and then press the **Enter** key to place the Material Parameter node in the Material Graph.
 
-![](Making_Parameters_Right_Click_Add.gif)(convert:false)
+[OBJECT:EmbeddedVideo]
+[PARAMLITERAL:width]
+640
+[/PARAMLITERAL]
+[PARAMLITERAL:height]
+360
+[/PARAMLITERAL]
+[PARAMLITERAL:videoid]
+fyVkiBqwyWw
+[/PARAMLITERAL]
+[/OBJECT]
 
 ## Using Parameters
 
@@ -62,7 +92,17 @@ Where the real power of Material Parameters comes into play is when you use them
 	* **  Vector Parameter ** x 1
 	* ** Scalar Node** x 2
 	
-	![](Making_Parameters_Adding_Nodes.gif)(convert:false)
+	[OBJECT:EmbeddedVideo]
+	[PARAMLITERAL:width]
+	640
+	[/PARAMLITERAL]
+	[PARAMLITERAL:height]
+	360
+	[/PARAMLITERAL]
+	[PARAMLITERAL:videoid]
+	A0jBZjF7iKc
+	[/PARAMLITERAL]
+	[/OBJECT]
 
 	[region:tip]
 	You can add either Material Parameter nodes or Material Expression nodes. If you do use Material Expression nodes, make sure that you convert them to Material Parameter nodes via the **Right-Click** menu after you have selected them.
@@ -74,19 +114,59 @@ Where the real power of Material Parameters comes into play is when you use them
 	* ** Metallic** 0.0
 	* ** Roughness** 0.5
 
-	![](Making_Parameters_Naming_Parameters.gif)(convert:false)
+	[OBJECT:EmbeddedVideo]
+	[PARAMLITERAL:width]
+	640
+	[/PARAMLITERAL]
+	[PARAMLITERAL:height]
+	360
+	[/PARAMLITERAL]
+	[PARAMLITERAL:videoid]
+	StKVHekJT14
+	[/PARAMLITERAL]
+	[/OBJECT]
 
 1. Once everything is named, connect all of the Parameter nodes into their respective inputs and make sure to compile and save your Material. 
 
-	![](Making_Parameters_Connecting_Compile.gif)(convert:false)
+	[OBJECT:EmbeddedVideo]
+	[PARAMLITERAL:width]
+	640
+	[/PARAMLITERAL]
+	[PARAMLITERAL:height]
+	360
+	[/PARAMLITERAL]
+	[PARAMLITERAL:videoid]
+	Fwh23pol9CI
+	[/PARAMLITERAL]
+	[/OBJECT]
 
 1. Now that the Material is compiled and saved, close the Material and find it in the **Content Browser**. Once you find it, **Right-Click** on it with the mouse and select the **Create Material Instance** option to create a new Material Instance from the Material. 
 
-	![](Making_Parameters_Create_Material_Instance.gif)(convert:false)
+	[OBJECT:EmbeddedVideo]
+	[PARAMLITERAL:width]
+	640
+	[/PARAMLITERAL]
+	[PARAMLITERAL:height]
+	360
+	[/PARAMLITERAL]
+	[PARAMLITERAL:videoid]
+	N17UEEqVEx4
+	[/PARAMLITERAL]
+	[/OBJECT]
 
 1. Open the newly created Material Instance by **Double-Clicking** on it with the **Left Mouse Button** in the **Content Browser**. Once the Material Instance is open you will be able to edit, in real time, the parameters that were set up and named inside of the Material like **Base_Color**, **Metallic**, and **Roughness**.
 
-	![](Making_Parameters_Making_Instance.gif)(convert:false)
+	[OBJECT:EmbeddedVideo]
+	[PARAMLITERAL:width]
+	640
+	[/PARAMLITERAL]
+	[PARAMLITERAL:height]
+	360
+	[/PARAMLITERAL]
+	[PARAMLITERAL:videoid]
+	tnmNiek09Tw
+	[/PARAMLITERAL]
+	[/OBJECT]
 
 ## Conclusion 
 
