UEDOC-2132

gifs have been replaced with YouTube videos

diff --git a/Engine/Documentation/Source/Engine/Rendering/Materials/HowTo/Subsurface_Profile/Subsurface_Profile.INT.udn b/Engine/Documentation/Source/Engine/Rendering/Materials/HowTo/Subsurface_Profile/Subsurface_Profile.INT.udn
--- a/Engine/Documentation/Source/Engine/Rendering/Materials/HowTo/Subsurface_Profile/Subsurface_Profile.INT.udn
+++ b/Engine/Documentation/Source/Engine/Rendering/Materials/HowTo/Subsurface_Profile/Subsurface_Profile.INT.udn
@@ -27,7 +27,17 @@ When completed, your **Content Browser** should look like this.
 	![](SSP_Open_New_Material.png)
 
 1. Now with the Material open, we need to change the Materials **Shading Model** from **Default Lit** to **Subsurface Profile**. 
-	![](SSP_Enable_Subsurface_Profile.gif)(convert:false)
+	[OBJECT:EmbeddedVideo]
+	[PARAMLITERAL:width]
+	640
+	[/PARAMLITERAL]
+	[PARAMLITERAL:height]
+	360
+	[/PARAMLITERAL]
+	[PARAMLITERAL:videoid]
+	QxtmqE7B_w8
+	[/PARAMLITERAL]
+	[/OBJECT]
 
 1. At this time, you can assign a **Subsurface Profile** by inputting one into the **Subsurface Profile** section in the **Details** panel.
 	![](SSP_Apply_SSP.png)
@@ -79,14 +89,34 @@ To wire a node to the main Material node, use the **Left Mouse Button** to click
 With the **Left Mouse Button** still held down, drag to the **Right** and the connection wire will appear. Keep dragging to the **Right** with the  **Left Mouse Button** still held until you come to the input you want to use on the Main Material node. 
 Once your mouse is over the input you want to use, release the **Left Mouse Button** to connect. 
 
-	![](SSP_Wire_Material_Up.gif)(convert:false)
+	[OBJECT:EmbeddedVideo]
+	[PARAMLITERAL:width]
+	640
+	[/PARAMLITERAL]
+	[PARAMLITERAL:height]
+	360
+	[/PARAMLITERAL]
+	[PARAMLITERAL:videoid]
+	osuR7VXYBWc
+	[/PARAMLITERAL]
+	[/OBJECT]
 
 1. Now that everything is wired up, we can now add our Subsurface Profile to the Material. 
  To do this, first find the **Subsurface Profile** input in the **Material** section of the **Details** panel. 
  Then use the **Left Mouse Button** to click in the area that says **None**. 
  When you do, you will see the Subsurface Profile that was created in the previous steps, select this profile by clicking on it. 
 	
-	![](SSP_Set_Subsurface_Profile.gif)(convert:false)
+	[OBJECT:EmbeddedVideo]
+	[PARAMLITERAL:width]
+	640
+	[/PARAMLITERAL]
+	[PARAMLITERAL:height]
+	360
+	[/PARAMLITERAL]
+	[PARAMLITERAL:videoid]
+	YBRE4qwtfX0
+	[/PARAMLITERAL]
+	[/OBJECT]
 	[region:note]
 	If you do not provide a Subsurface profile, a default one will be used. The default one that is used is a close approximation to Caucasian skin. 
 	[/region]
@@ -122,7 +152,17 @@ When completed, your **Content Browser** should look something like this.
 1. Now that the Material Instance has been created, it is now time to test out our new Material. 
 	To do this, we first need to create a new blank level to work in by going to the main menu and under the **File** option select **New Level**. 
 	When prompted what type of level to select, select the **Empty Level**. 
-	![](Engine\Rendering\Materials\HowTo\Subsurface_Scattering\SSM_Create_New_Level.gif)(convert:false)
+	[OBJECT:EmbeddedVideo]
+	[PARAMLITERAL:width]
+	640
+	[/PARAMLITERAL]
+	[PARAMLITERAL:height]
+	360
+	[/PARAMLITERAL]
+	[PARAMLITERAL:videoid]
+	bRuKiOW0i-w
+	[/PARAMLITERAL]
+	[/OBJECT]
 
 1. With the new level created, we need to now add some lights so that we can test how our Subsurface Material will look. The first light that we are going to add is a **Point Light**. To add a Point Light to the scene, we first need to find the Point Light which is in the **Modes** menu under the **Lights** option. Once the Point Light is located, you can add one to the screen by first selecting it with the **Left Mouse Button**  and then with the **Left Mouse Button** still held down dragging the Point Light into the scene and releasing the **Left Mouse Button** to add it. Once added to the scene, select the light by clicking on it with the **Left Mouse Button**. Then enter the following values.
 	| Property Name| Value|
@@ -158,7 +198,17 @@ When completed, your **Content Browser** should look something like this.
 1. Now that the Static Mesh has been placed in the level, we can now apply our Subsurface Material to it. 
 	To do this, first find the Subsurface Material in the **Content Browser**. 
 	Once the Material has been located, Drag it from the **Content Browser** on to the mesh to apply it. 
-	![](SSP_Apply_Material_Instance.gif)(convert:false)
+	[OBJECT:EmbeddedVideo]
+	[PARAMLITERAL:width]
+	640
+	[/PARAMLITERAL]
+	[PARAMLITERAL:height]
+	360
+	[/PARAMLITERAL]
+	[PARAMLITERAL:videoid]
+	XU50LZ1lxMA
+	[/PARAMLITERAL]
+	[/OBJECT]
 
 ## Adjusting a Subsurface Profile
 
@@ -180,7 +230,17 @@ When completed, your **Content Browser** should look something like this.
 	You should avoid using a vivid color here if you want to get a more complex shading variation over area you see the scattering.
 
 1. Keep in mind that you can adjust the parameters of the Subsurface Profile  in real time as demonstrated below. 
-	![](SSP_Adjusting_SSP.gif)(convert:false)
+	[OBJECT:EmbeddedVideo]
+	[PARAMLITERAL:width]
+	640
+	[/PARAMLITERAL]
+	[PARAMLITERAL:height]
+	360
+	[/PARAMLITERAL]
+	[PARAMLITERAL:videoid]
+	u3gFHZiKxCY
+	[/PARAMLITERAL]
+	[/OBJECT]
 
 ## Tips & Tricks
 Because of how Subsurface Profiles render, there are a few Tips & Tricks to getting the most out of using them. 
