#UE4doc #docs - adjusted bolding

diff --git a/Engine/Documentation/Source/Engine/Paper2D/HowTo/Paper2DHowTo.INT.udn b/Engine/Documentation/Source/Engine/Paper2D/HowTo/Paper2DHowTo.INT.udn
--- a/Engine/Documentation/Source/Engine/Paper2D/HowTo/Paper2DHowTo.INT.udn
+++ b/Engine/Documentation/Source/Engine/Paper2D/HowTo/Paper2DHowTo.INT.udn
@@ -7,7 +7,7 @@ Related: Engine\Paper2D\Flipbooks
 
 
 ## Overview
-The **Paper 2D - How To** pages provide short step-by-step examples of ways to accomplish certain tasks within Paper 2D in Unreal Engine 4. While there may be other ways to accomplish these tasks, these samples should get you started and give you ideas on how you may want to approach your objectives. 
+The **Paper 2D - How-to** pages provide short step-by-step examples of ways to accomplish certain tasks within Paper 2D in Unreal Engine 4. While there may be other ways to accomplish these tasks, these samples should get you started and give you ideas on how you may want to approach your objectives. 
 
 Refer to the table below for more information and click each link to view its contents. 
 
