#UE4doc #docs:
Removed italics from one term

diff --git a/Engine/Documentation/Source/Resources/ContentExamples/EffectsGallery/1_C/1_C.INT.udn b/Engine/Documentation/Source/Resources/ContentExamples/EffectsGallery/1_C/1_C.INT.udn
--- a/Engine/Documentation/Source/Resources/ContentExamples/EffectsGallery/1_C/1_C.INT.udn
+++ b/Engine/Documentation/Source/Resources/ContentExamples/EffectsGallery/1_C/1_C.INT.udn
@@ -56,7 +56,7 @@ Related: Engine/Rendering/ParticleSystems/Reference/Modules/VectorField
 This example shows 3 spheres traveling about a central point while leaving trails behind them. The effect requires that particles be spawned from the location of other particles, which is handled by way of the Emitter Initial Location module. This spawns particles at the location of the particles from a named Emitter. This means that the effect uses two Emitters:
 
 * **Lead_sphere:** This Emitter creates 3 spheres that fly around the center of the effect due to a Point Attractor module.
-* **Trails:** This Emitter spawns trails that are born at the location of the particles in _Lead_sphere_.
+* **Trails:** This Emitter spawns trails that are born at the location of the particles in Lead_sphere.
 
 
 ![](image016.png)
