#UE4doc #docs:
Changed font settings for a couple terms

diff --git a/Engine/Documentation/Source/Resources/ContentExamples/EffectsGallery/1_E/1_E.INT.udn b/Engine/Documentation/Source/Resources/ContentExamples/EffectsGallery/1_E/1_E.INT.udn
--- a/Engine/Documentation/Source/Resources/ContentExamples/EffectsGallery/1_E/1_E.INT.udn
+++ b/Engine/Documentation/Source/Resources/ContentExamples/EffectsGallery/1_E/1_E.INT.udn
@@ -53,8 +53,8 @@ Related: Engine/Rendering/ParticleSystems/Reference/Modules/VectorField
 [/REGION]
 
 
-Another advantage of the GPU sprites, besides being able to spawn large numbers of particles, is that it is possible for them to cheaply collide with the environment. This is handled by way of a _Collision (Scene Depth)_ module, which allows a GPU particle system to collide with geometry based on Z-Depth. 
+Another advantage of the GPU sprites, besides being able to spawn large numbers of particles, is that it is possible for them to cheaply collide with the environment. This is handled by way of a **Collision (Scene Depth)** module, which allows a GPU particle system to collide with geometry based on Z-Depth. 
 
-The nice thing about the module is that you do not have to worry about any collision surfaces for meshes in the scene. Since the module uses the Z-Depth buffer from the GPU, you automatically collide with any visible geometry. The _Collision (Scene Depth)_ module provides you control over resilience (bounciness) and friction.
+The nice thing about the module is that you do not have to worry about any collision surfaces for meshes in the scene. Since the module uses the Z-Depth buffer from the GPU, you automatically collide with any visible geometry. The Collision (Scene Depth) module provides you control over resilience (bounciness) and friction.
 
  ![](image024.png)
