#UE4doc #docs:
Removed italics from multiple terms

diff --git a/Engine/Documentation/Source/Resources/ContentExamples/EffectsGallery/1_F/1_F.INT.udn b/Engine/Documentation/Source/Resources/ContentExamples/EffectsGallery/1_F/1_F.INT.udn
--- a/Engine/Documentation/Source/Resources/ContentExamples/EffectsGallery/1_F/1_F.INT.udn
+++ b/Engine/Documentation/Source/Resources/ContentExamples/EffectsGallery/1_F/1_F.INT.udn
@@ -57,9 +57,9 @@ Related: Engine/Rendering/ParticleSystems/Reference/Modules/VectorField
 
 Rarely do particles alone handle the job of a special effect. In most cases, you will need to combine particles, lights, sounds, and other elements to achieve your goal. This example shows how a [Blueprint](Engine/Blueprints) can be created in order to combine a Particle System with a light and sound effect. The effect is a simple Blueprint Class deriving from Actor. It contains the following Components:
 
-* The _P_spark_burst_ Particle System.
-* A _Sound Component_ that plays the _sparks_ sound effect.
-* A _Light Component_ to supply a dynamic flash of light.
+* The P_spark_burst Particle System.
+* A Sound Component that plays the _sparks_ sound effect.
+* A Light Component to supply a dynamic flash of light.
 
 
 [REGION:fullwidth]
