#UE4doc #docs:
Changed font settings for one term

diff --git a/Engine/Documentation/Source/Resources/ContentExamples/EffectsGallery/1_G/1_G.INT.udn b/Engine/Documentation/Source/Resources/ContentExamples/EffectsGallery/1_G/1_G.INT.udn
--- a/Engine/Documentation/Source/Resources/ContentExamples/EffectsGallery/1_G/1_G.INT.udn
+++ b/Engine/Documentation/Source/Resources/ContentExamples/EffectsGallery/1_G/1_G.INT.udn
@@ -59,7 +59,7 @@ This effect is achieved by way of the Bone/Socket Location module, which allows
 
 ![](image034.png)
  
-Then, in the Emitter Actor properties, you add the instance parameter with the name you specified before in the location module, select 'PSPT_Actor' from the list, and select the Actor (Skeletal Mesh in this case) you want to emit particles from, and you will immediately see particles spawning at the bones/sockets you specified in the location module. Done!
+Then, in the Emitter Actor properties, you add the instance parameter with the name you specified before in the location module, select **PSPT_Actor** from the list, and select the Actor (Skeletal Mesh in this case) you want to emit particles from, and you will immediately see particles spawning at the bones/sockets you specified in the location module. Done!
 
 ![](image036.png)
 
