#UE4doc #docs:
Bolded a couple terms

diff --git a/Engine/Documentation/Source/Resources/ContentExamples/EffectsGallery/1_H/1_H.INT.udn b/Engine/Documentation/Source/Resources/ContentExamples/EffectsGallery/1_H/1_H.INT.udn
--- a/Engine/Documentation/Source/Resources/ContentExamples/EffectsGallery/1_H/1_H.INT.udn
+++ b/Engine/Documentation/Source/Resources/ContentExamples/EffectsGallery/1_H/1_H.INT.udn
@@ -63,7 +63,7 @@ Secondly, the lights that are intended to affect the particles must have **Cast
  
 ![](image042.png)
 
-And finally, the settings for the shadow and self-shadow behavior are located in the Material used in this particle effect, under the two Translucency groups in the material properties. This image shows these properties from the Material Editor.
+And finally, the settings for the shadow and self-shadow behavior are located in the Material used in this particle effect, under the two **Translucency** groups in the material properties. This image shows these properties from the **Material Editor**.
 
 ![](image044.png)
 
