#UE4doc #docs:
Changed font settings for a couple terms

diff --git a/Engine/Documentation/Source/Resources/ContentExamples/EffectsGallery/2_A/2_A.INT.udn b/Engine/Documentation/Source/Resources/ContentExamples/EffectsGallery/2_A/2_A.INT.udn
--- a/Engine/Documentation/Source/Resources/ContentExamples/EffectsGallery/2_A/2_A.INT.udn
+++ b/Engine/Documentation/Source/Resources/ContentExamples/EffectsGallery/2_A/2_A.INT.udn
@@ -53,12 +53,12 @@ Related: Engine/Rendering/ParticleSystems/Reference/Modules/VectorField
 [/REGION]
 
 
-One way to control the movement of GPU sprites is through the use of Vector Fields, which are 3D containers with a grid of velocity vectors that can be used to determine the velocity or acceleration of the sprites. This effect shows a Local Vector Field, meaning that it is a volume of vectors that exists entirely within the particle system. 
+One way to control the movement of GPU sprites is through the use of **Vector Fields**, which are 3D containers with a grid of velocity vectors that can be used to determine the velocity or acceleration of the sprites. This effect shows a Local Vector Field, meaning that it is a volume of vectors that exists entirely within the particle system. 
 
 These Vector Fields can be created within 3D packages. At Epic, we create ours using the fluid simulation system within Autodesk Maya. This means we can create particle motion that is based on a high-end fluidic simuation.
 
 
 ![](image048.png)
 
-Inside the Vector Field group, you will find various modules, including the _Local Vector Field_, as well as modules to control the field's rotation, rotation rate, and scale over life.
+Inside the Vector Field group, you will find various modules, including the Local Vector Field, as well as modules to control the field's rotation, rotation rate, and scale over life.
 
