#UE4doc #docs:
Bolded a term

diff --git a/Engine/Documentation/Source/Resources/ContentExamples/EffectsGallery/2_B/2_B.INT.udn b/Engine/Documentation/Source/Resources/ContentExamples/EffectsGallery/2_B/2_B.INT.udn
--- a/Engine/Documentation/Source/Resources/ContentExamples/EffectsGallery/2_B/2_B.INT.udn
+++ b/Engine/Documentation/Source/Resources/ContentExamples/EffectsGallery/2_B/2_B.INT.udn
@@ -54,7 +54,7 @@ Related: Engine/Rendering/ParticleSystems/Reference/Modules/VectorField
 [/REGION]
 
 
-A Global Vector Field is similar to the Local Vector Field in the previous example, but instead of being a module inside the emitter, this is an Actor that you can place in your level, positioning and editing it as desired. Just as with the Local Vector Field, GPU particles are only affected by the forces of the vector field when they are inside the field's volume. 
+A **Global Vector Field** is similar to the Local Vector Field in the previous example, but instead of being a module inside the emitter, this is an Actor that you can place in your level, positioning and editing it as desired. Just as with the Local Vector Field, GPU particles are only affected by the forces of the vector field when they are inside the field's volume. 
 
 To set it up, you simply have to place a Vector Field Volume in your level and add a Global Vector Field module to your emitter.  That is it!
 
