#UE4doc #docs:
Changed font settings for multiple terms

diff --git a/Engine/Documentation/Source/Resources/ContentExamples/EffectsGallery/2_C/2_C.INT.udn b/Engine/Documentation/Source/Resources/ContentExamples/EffectsGallery/2_C/2_C.INT.udn
--- a/Engine/Documentation/Source/Resources/ContentExamples/EffectsGallery/2_C/2_C.INT.udn
+++ b/Engine/Documentation/Source/Resources/ContentExamples/EffectsGallery/2_C/2_C.INT.udn
@@ -53,14 +53,14 @@ Related: Engine/Rendering/ParticleSystems/Reference/Modules/VectorField
 ![](image054.png)
 [/REGION]
 
-One very useful feature of the UE4 Material system is the [Black Body Material Expression node](Engine/Rendering/Materials\ExpressionReference\Utility#BlackBody), which can be found in the Material Editor. Beware, we now have to dive into SCIENCE!
+One very useful feature of the UE4 Material system is the [Black Body Material Expression node](Engine/Rendering/Materials\ExpressionReference\Utility#BlackBody), which can be found in the **Material Editor**. Beware, we now have to dive into SCIENCE!
 
 Black-body is a type of radiation that specific types of objects (or bodies) emit at a certain temperature. One of the visible characteristics of this radiation is the spectrum of color of the light that is emitted. 
 
 In simple terms, you give the node a real-world temperature in Kelvin, and you get a real-world color corresponding to that temperature. This solves that age-old problem of "what color should this explosion/fire be?" Using the Black Body material node, you can guarantee a realistic color to use on effects. To make the most of the Material Expression, it will help to have an understanding of [Black-body radiation.](http://en.wikipedia.org/wiki/Black-body_radiation)
 
 
-In this specific example, we use a Dynamic Parameter module within our Particle System. This allows us to create an arbitrarily named parameter that can then be used to drive some aspect of a Material. So we have a custom Dynamic Parameter named _Temperature_, whose value runs from 4500 to 0 throughout the lifespan of the particle. 
+In this specific example, we use a Dynamic Parameter module within our Particle System. This allows us to create an arbitrarily named parameter that can then be used to drive some aspect of a Material. So we have a custom Dynamic Parameter named Temperature, whose value runs from 4500 to 0 throughout the lifespan of the particle. 
 
  
  ![](image058.png)
