#UE4doc #docs:
Bolded a couple terms

diff --git a/Engine/Documentation/Source/Resources/ContentExamples/EffectsGallery/2_D/2_D.INT.udn b/Engine/Documentation/Source/Resources/ContentExamples/EffectsGallery/2_D/2_D.INT.udn
--- a/Engine/Documentation/Source/Resources/ContentExamples/EffectsGallery/2_D/2_D.INT.udn
+++ b/Engine/Documentation/Source/Resources/ContentExamples/EffectsGallery/2_D/2_D.INT.udn
@@ -52,12 +52,12 @@ Related: Engine/Rendering/ParticleSystems/Reference/Modules/VectorField
 ![](image060.png)
 [/REGION]
 
-One of the more useful features for complex particle effects is the Event Generator module. This module allows you to specify various types of Events that can occur during the lifespan of a particle, and use those Events to spawn an entirely different particle effect. Available Events include particle spawn, death, collision, burst, or even a Level Blueprint generated Event!
+One of the more useful features for complex particle effects is the **Event Generator** module. This module allows you to specify various types of Events that can occur during the lifespan of a particle, and use those Events to spawn an entirely different particle effect. Available Events include particle spawn, death, collision, burst, or even a Level Blueprint generated Event.
 
 In this case, we use the death of particles to create other particles. To make fireworks, we choose the type of event from the drop-down list and give it a name.
  
 ![](image062.png)
 
-You then use another module (in the second emitter) called EventReceiver Spawn that enables you to spawn particles based on the previously generated Event:
+You then use another module (in the second emitter) called **EventReceiver Spawn** that enables you to spawn particles based on the previously generated Event:
  
 ![](image064.png)
