#UE4doc #docs:
Changed font settings for a couple terms

diff --git a/Engine/Documentation/Source/Resources/ContentExamples/EffectsGallery/2_E/2_E.INT.udn b/Engine/Documentation/Source/Resources/ContentExamples/EffectsGallery/2_E/2_E.INT.udn
--- a/Engine/Documentation/Source/Resources/ContentExamples/EffectsGallery/2_E/2_E.INT.udn
+++ b/Engine/Documentation/Source/Resources/ContentExamples/EffectsGallery/2_E/2_E.INT.udn
@@ -53,10 +53,10 @@ Related: Engine/Rendering/ParticleSystems/Reference/Modules/VectorField
 ![](image066.png)
 [/REGION]
 
-The Ribbon TypeData allows you to create trails, which are very useful in all kinds of projectile-based effects. To make them work, you generally need at least 2 particle Emitters:
+The **Ribbon TypeData** allows you to create trails, which are very useful in all kinds of projectile-based effects. To make them work, you generally need at least 2 particle Emitters:
 
 * One particle Emitter that emits the source particles (from which the ribbons will spawn) - note that these cannot be GPU particles.
-* One particle Emitter that has a Ribbon TypeData module and a Source module (located in the _Trail_ group) to define which Emitter will be used as the source.
+* One particle Emitter that has a Ribbon TypeData module and a Source module (located in the **Trail** group) to define which Emitter will be used as the source.
 
 
 ![](image068.png)
