#UE4doc #docs:
Bolded a couple terms

diff --git a/Engine/Documentation/Source/Resources/ContentExamples/EffectsGallery/2_G/2_G.INT.udn b/Engine/Documentation/Source/Resources/ContentExamples/EffectsGallery/2_G/2_G.INT.udn
--- a/Engine/Documentation/Source/Resources/ContentExamples/EffectsGallery/2_G/2_G.INT.udn
+++ b/Engine/Documentation/Source/Resources/ContentExamples/EffectsGallery/2_G/2_G.INT.udn
@@ -53,7 +53,7 @@ Related: Engine/Rendering/ParticleSystems/Reference/Modules/VectorField
 ![](image074.png)
 [/REGION]
 
-Another fantastic and highly useful particle feature in the Unreal Engine is the ability for CPU particles to emit light! By adding a Light module to your Emitter, particles from the Emitter can also spawn a dynamic light. You can use the module's properties to control the intensity of the light, the radius, etc.
+Another fantastic and highly useful particle feature in the **Unreal Engine** is the ability for CPU particles to emit light. By adding a **Light** module to your Emitter, particles from the Emitter can also spawn a dynamic light. You can use the module's properties to control the intensity of the light, the radius, etc.
 
 ![](image076.png)
 
