#UE4doc #docs:
Bolded a few terms

diff --git a/Engine/Documentation/Source/Resources/ContentExamples/EffectsGallery/2_H/2_H.INT.udn b/Engine/Documentation/Source/Resources/ContentExamples/EffectsGallery/2_H/2_H.INT.udn
--- a/Engine/Documentation/Source/Resources/ContentExamples/EffectsGallery/2_H/2_H.INT.udn
+++ b/Engine/Documentation/Source/Resources/ContentExamples/EffectsGallery/2_H/2_H.INT.udn
@@ -53,7 +53,7 @@ Related: Engine/Rendering/ParticleSystems/Reference/Modules/VectorField
 ![](image078.png)
 [/REGION]
 
-The Beam TypeData allows you to create an effect that connects two points - you can either specify the Source and the Target point, or have them as Actors, or other Emitters. It is used often as a bullet trail effect in shooter type games. In this case, we use it as an electric arc, with the help of the Noise module, inside the Beam group. 
+The **Beam TypeData** allows you to create an effect that connects two points - you can either specify the Source and the Target point, or have them as Actors, or other Emitters. It is used often as a bullet trail effect in shooter type games. In this case, we use it as an electric arc, with the help of the **Noise** module, inside the **Beam** group. 
  
  
  ![](image080.png)
