#UE4doc #docs:
Changed font settings for multiple terms

diff --git a/Engine/Documentation/Source/Resources/ContentExamples/ExampleLevelOverview/ExampleLevelOverview.INT.udn b/Engine/Documentation/Source/Resources/ContentExamples/ExampleLevelOverview/ExampleLevelOverview.INT.udn
--- a/Engine/Documentation/Source/Resources/ContentExamples/ExampleLevelOverview/ExampleLevelOverview.INT.udn
+++ b/Engine/Documentation/Source/Resources/ContentExamples/ExampleLevelOverview/ExampleLevelOverview.INT.udn
@@ -48,7 +48,7 @@ Description:An overview of how the Example levels included with Unreal Engine 4
 ![](ExampleMap.png)
 [/REGION]
 
-The Example Levels for Unreal Engine 4 have been built to help you quickly get up to speed with a variety of different techniques and artistic applications for building your own games. This document overviews the process of using these examples in the most effective manner along with their corresponding documentation.
+The Example Levels for **Unreal Engine 4** (UE4) have been built to help you quickly get up to speed with a variety of different techniques and artistic applications for building your own games. This document overviews the process of using these examples in the most effective manner along with their corresponding documentation.
 
 ## Level Layout
 
@@ -63,7 +63,7 @@ Some example levels will have multiple rooms, open or closed ceilings, and any n
 
 ## Using Example Levels 
 
-Because each of the example levels is broken up into various individual examples, it is very easy to select a single effect and open it up for more information. This is often going to be a great way to help you learn more about various tools and techniques available within Unreal Engine 4. 
+Because each of the example levels is broken up into various individual examples, it is very easy to select a single effect and open it up for more information. This is often going to be a great way to help you learn more about various tools and techniques available within UE4. 
 
 
 [REGION:imagetable]
@@ -74,7 +74,7 @@ Because each of the example levels is broken up into various individual examples
 [/REGION]
 
 
-The way in which you open up a specific example may vary depending on the type of asset. In many cases, you can **right-click** it and choose **Find in Content Browser**. This will automatically select that asset in the **Content Browser**. **Double-clicking** the asset will then open it up in its respective editor. In the example shown above, we have opened up a Particle System within Cascade.
+The way in which you open up a specific example may vary depending on the type of asset. In many cases, you can **Right-click** it and choose **Find in Content Browser**. This will automatically select that asset in the **Content Browser**. **Double-clicking** the asset will then open it up in its respective editor. In the example shown above, we have opened up a Particle System within **Cascade**.
 
 ## Cross-Referencing with Documentation
 
@@ -83,7 +83,7 @@ While each of the example levels stand alone as a great way to take apart indivi
 <!--TODO - ADD PIC OF EFFECTS GALLERY CROSS REFERENCE-->
 
 
-As you move about the level, you will see that each station contains a callout number to help you cross-reference each of the assets that have been set up. One of these stations is labeled **2.3**, and shows a Particle System being colored with a Material that uses the Blackbody node. Within the documentation page, you will see a section with the same label, which will contain relevant information for this effect, including how it was set up and any important features or considerations.
+As you move about the level, you will see that each station contains a callout number to help you cross-reference each of the assets that have been set up. One of these stations is labeled 2.3, and shows a Particle System being colored with a Material that uses the Blackbody node. Within the documentation page, you will see a section with the same label, which will contain relevant information for this effect, including how it was set up and any important features or considerations.
 
 
 
