UEDOC-2136

converted gif into a Youtube embeded video

diff --git a/Engine/Documentation/Source/Engine/UMG/QuickStart/4/UMGQS_4.INT.udn b/Engine/Documentation/Source/Engine/UMG/QuickStart/4/UMGQS_4.INT.udn
--- a/Engine/Documentation/Source/Engine/UMG/QuickStart/4/UMGQS_4.INT.udn
+++ b/Engine/Documentation/Source/Engine/UMG/QuickStart/4/UMGQS_4.INT.udn
@@ -108,7 +108,17 @@ Time to get our hands dirty and script the functionality of our Main Menu. At th
 
 1.	Click the **Play** Button from the Main Toolbar to play in the editor. 
 
-	[![](MenuImage.png)(w:720)](ExampleMainMenu.mp4)
+	[OBJECT:EmbeddedVideo]
+	[PARAMLITERAL:width]
+	640
+	[/PARAMLITERAL]
+	[PARAMLITERAL:height]
+	360
+	[/PARAMLITERAL]
+	[PARAMLITERAL:videoid]
+	mEVUvk32CRI
+	[/PARAMLITERAL]
+	[/OBJECT]
 
 	_Click the image above to play the video._
 
