UEDOC-2137

converted gifs into Youtube videos

diff --git a/Engine/Documentation/Source/Engine/UMG/QuickStart/5/UMGQS_5.INT.udn b/Engine/Documentation/Source/Engine/UMG/QuickStart/5/UMGQS_5.INT.udn
--- a/Engine/Documentation/Source/Engine/UMG/QuickStart/5/UMGQS_5.INT.udn
+++ b/Engine/Documentation/Source/Engine/UMG/QuickStart/5/UMGQS_5.INT.udn
@@ -170,7 +170,17 @@ With the visuals set, next we will provide the scripted functionality for our Pa
 
 1.	Click the **Compile** and **Save** Buttons, then the **Play** button to play in the editor. 
 
-	[![](Engine\UMG\QuickStart\4\MenuImage.png)(w:720)](PauseMovie.mp4)
+	[OBJECT:EmbeddedVideo]
+	[PARAMLITERAL:width]
+	640
+	[/PARAMLITERAL]
+	[PARAMLITERAL:height]
+	360
+	[/PARAMLITERAL]
+	[PARAMLITERAL:videoid]
+	RALk7pUzRdw
+	[/PARAMLITERAL]
+	[/OBJECT]
 
 	Once in-game, if you press the **M** Key (or whichever key you assigned), you will see that the game is paused and the Pause Menu appears. From here, you can resume gameplay with the **Resume** button or quit back to the Main Menu with the **Quit** Button. 
 	
