TypoCorrection

diff --git a/Engine/Documentation/Source/Engine/Paper2D/HowTo/Paper2DHowTo.INT.udn b/Engine/Documentation/Source/Engine/Paper2D/HowTo/Paper2DHowTo.INT.udn
--- a/Engine/Documentation/Source/Engine/Paper2D/HowTo/Paper2DHowTo.INT.udn
+++ b/Engine/Documentation/Source/Engine/Paper2D/HowTo/Paper2DHowTo.INT.udn
@@ -1,12 +1,10 @@
-Availability:Public
+﻿Availability:Public
 Title:Paper 2D - How To...
 Crumbs: %ROOT%, Engine, Engine/Paper2D
 Description:The Paper 2D How To pages provides several short step-by-step guides for working with Paper 2D. 
 Related: Engine\Paper2D\Sprites
 Related: Engine\Paper2D\Flipbooks
 
-
-## Overview
 The **Paper 2D - How-to** pages provide short step-by-step examples of ways to accomplish certain tasks within Paper 2D in Unreal Engine 4. While there may be other ways to accomplish these tasks, these samples should get you started and give you ideas on how you may want to approach your objectives. 
 
 Refer to the table below for more information and click each link to view its contents. 
