#UE4 DOCS: submitting stat page updates

diff --git a/Engine/Documentation/Source/Engine/Performance/CookingStats/CookingStats.INT.udn b/Engine/Documentation/Source/Engine/Performance/CookingStats/CookingStats.INT.udn
new file mode 100644
--- /dev/null
+++ b/Engine/Documentation/Source/Engine/Performance/CookingStats/CookingStats.INT.udn
@@ -0,0 +1,134 @@
+Availability:Public
+Title: Cooking Statistics
+Crumbs: %ROOT%, Engine
+Description:How to identify performance problems and fix them.
+
+When you cook your project, statistics are always created. You can find them saved into:
+
+	<projectdirectory>\Saved\Stats\<datetime>Stats.csv
+
+The **Stats.csv** file contains a number of lines each with a set of comma separated values. The first value is a unique key which describes that line. Which, is followed by a tag that can have a value associated with it, using a '=' to separate. Example:
+
+	DDCTransactionId2154470643B07308C1979DA468F3EA48_0,StartTime=2015.07.14-17.09.19
+
+The key is the first parameter, so in this example, this is the first parameter:
+
+	DDCTransactionId2154470643B07308C1979DA468F3EA48_0
+
+And the first tag is:
+
+	StartTime=2015.07.14-17.09.19
+
+The first tag can then be split into a tag name and value pair
+
+	TagName: StartTime
+	Value: 2015.07.14-17.09.19
+
+These keys and tags can help you solve issues in the cooking pipeline.
+
+
+## Debugging Cooking Performance Issues
+
+Most time variations are caused by the **Derived Data Cache** being invalidated
+and needing to rebuild textures for platforms.
+
+If repeated cooks are still slow, this can indicate there is a large amount of data being cooked.  Consider using iterative cooking (experimental) or if cooking for testing / development purposes, only cook what is required.
+
+
+
+## Derived Data Cache (DDC) Stats
+
+DDC stats are recognizable in the stats.CSV file as their key will begin with DDCTransactionId.
+
+DDC stats are in a hierarchical format to reflect how the DDC searches for content in the cache.
+
+Review the 3 entries for DCC stats in, (tag lines edited for clarity)
+
+	DDCTransactionId2154470643B07308C1979DA468F3EA48_86
+	 GetCachedData
+	 CacheKey=MATSM_984afa3849534faa8491540b994bd32e_SF_PS4_15493__BC5N_DEV_SL_NoDBufF7646F9__526E4F887A3482CE276E9E40861788C0F9774343
+	 StartTime=2015.07.14-17.09.30
+	DDCTransactionId2154470643B07308C1979DA468F3EA48_87
+	 CachedDataProbablyExists
+	 Parent=DDCTransactionId2154470643B07308C1979DA468F3EA48_86
+	 StartTime=2015.07.14-17.09.30
+	DDCTransactionId2154470643B07308C1979DA468F3EA48_88 
+	 GetCachedData
+	 Parent=DDCTransactionId2154470643B07308C1979DA468F3EA48_86
+	 StartTime=2015.07.14-17.09.30
+
+The second two transactions are linked to the first through the use of the
+"Parent" tag.
+The parent tag specifies the first transaction as the source for the DDC
+request.
+The first transaction also specifies the tag "GetCachedData" which is the
+entry point on the transaction.  We can also see the CacheKey which was
+specified.
+Each transaction has a start time and a duration. The duration is in
+milliseconds, all duration's are inclusive time.
+The exclusive time for a parent transaction can be determined by
+subtracting the child transactions duration.
+
+## Save Package Stats
+
+Save package is a large part of cooking, and each save package stats key starts with "SavePackageTransactionId":
+
+	SavePackageTransactionId830D62FB4EE90050FDB33BBDAA57A9A7_124
+	 Filename=D:/Builds/UE4/QAGame/Saved/Cooked/PS4/Engine/Plugins/Runtime/LeapMotionController/Content/LM_SampleDistortion.uasset
+	 FlushAsyncLoading=0.019673ms  FlushStreamingFunc=0.007272ms
+	 BlockTillAllRequestsFinished=0.006843ms
+	 ResetLoadersForSaveUnMarkAllObjects=0.158671ms
+	 TagPackageExports=0.222396ms  BeginCacheForCookedPlatformData=0.086393ms
+	 SerializeImports=0.120610ms  GatherLocalizableTextData=0.005558ms
+	 MarkNames=0.007700ms  SerializeSummary=0.016678ms
+	 SerializeNames=0.027802ms  SerializeGatherableTextData=0.004277ms
+	 BuildExportMap=0.017535ms  SortExportsNonSeekfree=0.008553ms
+	 SortExportsSeekfreeConstructor=0.019245ms
+	 SortExportsSeekfreeInner=0.327609ms  SortExportsSeekfree=0.005133ms
+	 BuildDependencyMap=0.010692ms  SerializeImportMap=0.005562ms
+	 SerializeExportMap=0.004705ms  SerializeDependencyMap=0.004277ms
+	 SerializeStringAssetReferenceMap=0.004277ms
+	 SaveThumbNailsAssetRegistryDataWorldLevelInfo=0.022668ms
+	 SerializeExports=0.109062ms  SerializeBulkData=0.015821ms
+	 AsyncWrite=4.576687ms  UnaccountedTime=0.307079ms
+	 RunGuid=RunID7F4E40464DE38EF1DECAF5A385FBDB1E
+
+The tags in save package are broken up into sections of the save.
+Debugging issues in save package requires specific knowledge of the
+different sections of save package.
+
+## Cooking Package Stats
+
+Cooking Package Stats let you know how the cooker is processing the
+packages (from the cooker side).  They give you an idea of how long each
+package takes to process.
+
+Cooking package stats key is the standard package name.
+
+	../../../Engine/Content/EditorMaterials/Cascade/CASC_Cross.uasset
+	 AllOfLoadPackage_Duration=0.115048ms  AllOfLoadPackage_HitCount=5
+	 ResolveRedirectors_Duration=0.001281ms  ResolveRedirectors_HitCount=5
+	 CacheForCookedPlatformData_Duration=5.137812ms
+	 CacheForCookedPlatformData_HitCount=1
+	 UnsolicitedMarkup_Duration=0.065863ms  UnsolicitedMarkup_HitCount=1
+	 SavePackageCacheForCookedPlatformData_Duration=0.146270ms
+	 SavePackageCacheForCookedPlatformData_HitCount=2
+	 SaveCookedPackage_Duration=1.683380ms  SaveCookedPackage_HitCount=1
+	 RunGuid=RunID2062A4FA4DBD3833D4A026A94B8970BB
+
+Each stat has the hitcount and also the duration recorded.
+For example:
+
+* **AllOfLoadPackage** is a single stat, but has two tags associated with it:
+	* **AllOfLoadPackage_Duration** is the time taken to run all of load package.
+	* **AllOfLoadPackage_HitCount** is the number of times load package was hit.
+
+Different sections of the cooker can be hit multiple times due to the
+cooker trying to optimize the cooking process.
+The most common situation is that the cooker loads a package and the
+content is missing from DDC so needs to be regenerated (slow process),
+In this case the cooker will try and opportunistically cook other packages
+which have their content ready and come back to this package.
+This doesn't necessarily mean that the package was loaded 5 times, the
+cooker will early out if the package is already loaded and try and keep the
+package in memory if it isn't saved yet.
\ No newline at end of file
diff --git a/Engine/Documentation/Source/Engine/Performance/StatCommands/StatCommands.INT.udn b/Engine/Documentation/Source/Engine/Performance/StatCommands/StatCommands.INT.udn
--- a/Engine/Documentation/Source/Engine/Performance/StatCommands/StatCommands.INT.udn
+++ b/Engine/Documentation/Source/Engine/Performance/StatCommands/StatCommands.INT.udn
@@ -1,4 +1,4 @@
-Availability:Doc
+Availability:Public
 Title: Stat Commands
 Crumbs: %ROOT%, Engine
 Description:Console commands specific to displaying game statistics.
