#UE4doc #docs - updated images

diff --git a/Engine/Documentation/Source/Engine/Paper2D/HowTo/TopDown/TopDown.INT.udn b/Engine/Documentation/Source/Engine/Paper2D/HowTo/TopDown/TopDown.INT.udn
--- a/Engine/Documentation/Source/Engine/Paper2D/HowTo/TopDown/TopDown.INT.udn
+++ b/Engine/Documentation/Source/Engine/Paper2D/HowTo/TopDown/TopDown.INT.udn
@@ -35,7 +35,7 @@ Here we will set up our project for this tutorial, you can import your own asset
 	![](TopDown3.png)
 
 	[REGION:note]
-	For more information on how to set up assets for importing, see the [](Engine\Paper2D\Importing) documentation. 
+	For more information on how to set up assets for importing, see the [**Paper 2D Import Options**](Engine\Paper2D\Importing) documentation. 
 	[/REGION]
 
 1.	**Right-click** on the **SoldierSprites** asset and select **create Flipbooks**. 
@@ -138,7 +138,7 @@ The character is setup, however there is no movement script applied to it which
 There is no need to reinvent the wheel here, since this template comes with a character with scripted functionality for movement, we can copy that into our Paper Character to provide movement for it. 
 
 [REGION:note]
-For more detailed information on setting up character movement and setting up character movement from a blank project, see the [](Gameplay/HowTo/CharacterMovement/Blueprints) documentation. 
+For more detailed information on setting up character movement and setting up character movement from a blank project, see the [**Setting Up Character Movement in Blueprints**](Gameplay/HowTo/CharacterMovement/Blueprints) documentation. 
 [/REGION]
 
 1.	Inside the **ThirdPersonBP** folder in the **Content Browser**, open the **Blueprints** folder then the **ThirdPersonCharacter** Blueprint. 
