#UE4doc #docs - bolded a couple of links.

diff --git a/Engine/Documentation/Source/Engine/Paper2D/Sprites/Paper2DSprites.INT.udn b/Engine/Documentation/Source/Engine/Paper2D/Sprites/Paper2DSprites.INT.udn
--- a/Engine/Documentation/Source/Engine/Paper2D/Sprites/Paper2DSprites.INT.udn
+++ b/Engine/Documentation/Source/Engine/Paper2D/Sprites/Paper2DSprites.INT.udn
@@ -55,14 +55,13 @@ Related: Engine/Content/Types/Textures
 
 [TOC (start:2 end:3)]
 
-## Overview
 
 [EXCERPT:intro]
 A **Sprite** in **Paper 2D** is a Texture Mapped Planar Mesh and associated Material that can be rendered in the world, created entirely within Unreal Engine 4 (UE4). In simpler terms, it's a quick and easy way to draw 2D images in UE4. 
 
 Sprites can also be edited inside UE4 in the **Sprite Editor** which offers four modes: **View** for previewing the sprite and general statistics, **Edit Source Region** which displays the full source texture and allows you to set the area within the source texture that composes the individual sprite, **Edit Collision** displays and allows editing of the sprite collision shapes, and **Edit RenderGeom** which displays and allows editing of the sprite render geometry. 
 
-Once you have a collection of Sprite assets, you can then begin animating them by creating a **Flipbook** which sequentially will play through the Sprites added to the Flipbook. See the [Flipbook](Engine\Paper2D\Flipbooks) documentation for more information. 
+Once you have a collection of Sprite assets, you can then begin animating them by creating a **Flipbook** which sequentially will play through the Sprites added to the Flipbook. See the [**Flipbook**](Engine\Paper2D\Flipbooks) documentation for more information. 
 [/EXCERPT]
 
 ## Creating Sprites
@@ -98,7 +97,7 @@ Sprites are created in the **Content Browser**, like other assets. They can be c
 	![](sprite_new_details_texture.png)
 
 [REGION:note]
-See the [](Engine\Paper2D\Sprites\Editor) documentation for information on working with Sprites in the **Sprite Editor**.
+See the [**Sprite Editor Reference**](Engine\Paper2D\Sprites\Editor) documentation for information on working with Sprites in the **Sprite Editor**.
 [/REGION]
 
 
@@ -133,7 +132,7 @@ See the [](Engine\Paper2D\Sprites\Editor) documentation for information on worki
 **Paper 2D** includes an importer for JSON formatted sprite sheet descriptions, which can be exported from tools like **Adobe Flash CS6** or **Texture Packer**. It will import any referenced textures and create sprite assets for each sprite. The importer also assumes that all of the sprites are frames of an animation, so it will always create a Flipbook in addition to the individual sprites. You can delete the generated flipbook if it is not needed.
 
 [REGION:note]
-See [](Engine\Paper2D\Importing) for more information on Importing Options. 
+See [**Paper 2D Import Options**](Engine\Paper2D\Importing) for more information on Importing Options. 
 [/REGION]
 
 
