#UE4doc #docs - removed overview header and bolded links.

diff --git a/Engine/Documentation/Source/Engine/Paper2D/Sprites/Collision/SpriteCollision.INT.udn b/Engine/Documentation/Source/Engine/Paper2D/Sprites/Collision/SpriteCollision.INT.udn
--- a/Engine/Documentation/Source/Engine/Paper2D/Sprites/Collision/SpriteCollision.INT.udn
+++ b/Engine/Documentation/Source/Engine/Paper2D/Sprites/Collision/SpriteCollision.INT.udn
@@ -11,7 +11,6 @@ Related: Engine/Content/Types/Textures
 [TOC(start:2)]
 
 
-## Overview
 
 Just like other types of geometry, such as Static Meshes, Skeletal Meshes, etc., **Sprites** can define shapes that are used to calculate collisions with other geometry in the world.
 
@@ -25,7 +24,7 @@ different settings:
 | Collision Type | Description |
 | -------------- | ----------- |
 | **None** | No collision geometry is generated; use this for purely decorative sprites. |
-| **Use 2D Physics** | This is an *experimental* option to generate collision geometry for use with Box2D. See [2D Collision](#2dcollision) for more information on the current limitations. |
+| **Use 2D Physics** | This is an *experimental* option to generate collision geometry for use with Box2D. See [**2D Collision**](#2dcollision) for more information on the current limitations. |
 | **Use 3D Physics** | Collision geometry will be generated for use with PhysX. The 2D collision geometry defined in the sprite will be extruded using the **Collision Thickness** units deep in the perpendicular axis to make 3D collision geometry. |
 
 ## Collision Geometry
@@ -35,13 +34,13 @@ geometry. The following settings are available:
 
 | Type | Description |
 | ---- | ----------- |
-| **Source Bounding Box** | Uses the bounding box of the [Source Region](Engine/Paper2D/Sprites/RenderGeometry) for the Sprite. [REGION:none]![](collision_box.png)(w:192)[/REGION]  |
+| **Source Bounding Box** | Uses the bounding box of the [**Source Region**](Engine/Paper2D/Sprites/RenderGeometry) for the Sprite. [REGION:none]![](collision_box.png)(w:192)[/REGION]  |
 | **Tight Bounding Box** | Generates a tightened bounding box that excludes any fully transparent areas of the Sprite. This results in better collision in most cases. [REGION:none]![](collision_tight.png)(w:192)[/REGION] |
 | **Shrink Wwrapped** | (*Experimental*) Generates complex geoemetry the conforms to the opaque areas of the Sprite. This results in the most realistic collision, but the added geometry could result in a performance hit. [REGION:none]![](collision_shrink.png)(w:192)[/REGION] |
 | **Fully Custom** | Enables you to specify the geometry used interactively in the viewport. [REGION:none]![](collision_custom.png)(w:192)[/REGION] |
 | **Diced** | Split up into smaller squares, including only no-empty ones in the final geometry. |
 
-The Edit Collision mode displays the collision geometry and lets you adjust it in the viewport, automatically setting the mode to fully custom.
+The **Edit Collision** mode displays the collision geometry and lets you adjust it in the viewport, automatically setting the mode to fully custom.
 
 The toolbar buttons can be used to add additional polygons or snap vertices to the pixel grid. New vertices can be added by selecting an edge and **Shift+clicking** and selected vertices can be deleted by pressing **Delete**.
 
