UEDOC-2141

Converted gifs into YouTube Videos

diff --git a/Engine/Documentation/Source/Engine/Paper2D/HowTo/Physics/SpritePhysics.INT.udn b/Engine/Documentation/Source/Engine/Paper2D/HowTo/Physics/SpritePhysics.INT.udn
--- a/Engine/Documentation/Source/Engine/Paper2D/HowTo/Physics/SpritePhysics.INT.udn
+++ b/Engine/Documentation/Source/Engine/Paper2D/HowTo/Physics/SpritePhysics.INT.udn
@@ -13,7 +13,17 @@ _For this example we are using a 2D Side Scroller Game, however you can apply th
 
 When using physics in any 2D style game, typically you'll want to apply some constraints to your sprite when physics are applied to prevent the sprite from moving or rotating in such a way that the card becomes visible to the player, or even worse, falls out of your game world as seen in the example below. 
 
-![](GIF_1.gif)(convert:false)
+[OBJECT:EmbeddedVideo]
+[PARAMLITERAL:width]
+640
+[/PARAMLITERAL]
+[PARAMLITERAL:height]
+360
+[/PARAMLITERAL]
+[PARAMLITERAL:videoid]
+xjsflokQvHc
+[/PARAMLITERAL]
+[/OBJECT]
 
 Above, we have applied physics to a sprite but have not applied any constraints to those physics so our object that falls and applies physics actually tumbles outside of our game world and off screen which is not what we want. We want our physics object to react to the world and player interaction, however we want to keep it inside our level and prevent it from rotating on certain axis so that the sprite is always fully visible. 
 
@@ -29,7 +39,17 @@ For our sprite, we can apply constraints from the **Details** panel in addition
 
 	While this will apply physics and lock it to only the Y axis, there are additional issues that may occur with just this setting. 
 
-	![](GIF_2.gif)(convert:false)
+	[OBJECT:EmbeddedVideo]
+	[PARAMLITERAL:width]
+	640
+	[/PARAMLITERAL]
+	[PARAMLITERAL:height]
+	360
+	[/PARAMLITERAL]
+	[PARAMLITERAL:videoid]
+	pfEhG2MOBo0
+	[/PARAMLITERAL]
+	[/OBJECT]
 
 	Our sprite is now locked to the Y axis, however you can see that it can still rotate freely which may not be what we want. 
 
@@ -39,7 +59,17 @@ For our sprite, we can apply constraints from the **Details** panel in addition
 
 	Now when physics is applied to the sprite its position is locked to the Y and its rotation is locked at the X. 
 
-	![](GIF_3.gif)(convert:false)
+	[OBJECT:EmbeddedVideo]
+	[PARAMLITERAL:width]
+	640
+	[/PARAMLITERAL]
+	[PARAMLITERAL:height]
+	360
+	[/PARAMLITERAL]
+	[PARAMLITERAL:videoid]
+	CRVDaw0VxEs
+	[/PARAMLITERAL]
+	[/OBJECT]
 
 	You can also use the **Mode** option to constrain movement along a specified axis (in this case **XZPlane** to achieve the same effect). 
 
