#UE4doc #docs - adjusting bold lettering and removed overview header.

diff --git a/Engine/Documentation/Source/Engine/Paper2D/Sprites/Creation/CreatingSprites.INT.udn b/Engine/Documentation/Source/Engine/Paper2D/Sprites/Creation/CreatingSprites.INT.udn
--- a/Engine/Documentation/Source/Engine/Paper2D/Sprites/Creation/CreatingSprites.INT.udn
+++ b/Engine/Documentation/Source/Engine/Paper2D/Sprites/Creation/CreatingSprites.INT.udn
@@ -6,17 +6,16 @@ Availability:Docs
 [TOC(start:2)]
 
 
-## Overview
 
-Sprites are created in the Content Browser, like other assets. They can be created from scratch as blank assets, generated from other existing assets, or created using imported data.
+**Sprites** are created in the Content Browser, like other assets. They can be created from scratch as blank assets, generated from other existing assets, or created using imported data.
 
 ## Blank Sprite Asset
 
 **To create a new blank Sprite asset:**
 
-1. Click the ![New](button_cb_new.png)(h:24) button in the Content Browser
+1. Click the ![New](button_cb_new.png)(h:24) button in the **Content Browser**.
 
-1. Under Miscellaneous, select Sprite.
+1. Under _Miscellaneous_, select **Sprite**.
 	
 	![](sprite_new_menu.png)
 
@@ -24,7 +23,7 @@ Sprites are created in the Content Browser, like other assets. They can be creat
 
 	![](sprite_new_name.png)
 
-1. Double-click the new Sprite asset to open it in the Sprite Editor
+1. **Double-click** the new Sprite asset to open it in the **Sprite Editor**
 
 1. In the **Details** panel, you can assign a Texture to the Sprite asset using the **Source Texture** property.
 	
@@ -34,7 +33,7 @@ Sprites are created in the Content Browser, like other assets. They can be creat
 
 **To create a sprite from an existing Texture asset:**
 
-1. In the Content Browser, right-click on the Texture asset and choose **Create Sprite**.
+1. In the Content Browser, **right-click** on the Texture asset and choose **Create Sprite**.
 
 	![](sprite_create_menu.png)
 
@@ -45,7 +44,7 @@ Sprites are created in the Content Browser, like other assets. They can be creat
 
 ## Series of Sprites from Sprite Sheet Texture
 
-1. In the Content browser, right-click on the sprite sheet Texture and choose **Extract Sprites**.
+1. In the Content browser, **right-click** on the sprite sheet Texture and choose **Extract Sprites**.
 
 	![](sprite_extract_menu.png)
 
@@ -55,7 +54,7 @@ Sprites are created in the Content Browser, like other assets. They can be creat
 
 
 [REGION:tip]
-Select all of the new Sprite assets, Right-click on them, and choose **Create Flipbooks** to generate a Flipbook asset from the new Sprites.
+Select all of the new Sprite assets, **Right-click** on them, and choose **Create Flipbooks** to generate a Flipbook asset from the new Sprites.
 
 ![](sprite_create_flipbook.gif)(convert:false)
 [/REGION]
