UEDOC-2142

Converted gif into YouTube video

diff --git a/Engine/Documentation/Source/Engine/Paper2D/HowTo/Sockets/SpriteSockets.INT.udn b/Engine/Documentation/Source/Engine/Paper2D/HowTo/Sockets/SpriteSockets.INT.udn
--- a/Engine/Documentation/Source/Engine/Paper2D/HowTo/Sockets/SpriteSockets.INT.udn
+++ b/Engine/Documentation/Source/Engine/Paper2D/HowTo/Sockets/SpriteSockets.INT.udn
@@ -83,7 +83,17 @@ Now that we have a socket attached to each frame of our Flipbook, we can tell ou
 
 1.	Once you **Compile** and play in the editor, the Actor Class spawned will spawn at the socket location and become attached. 
 
-	![](Example_GIF.gif)(convert:false)
+	[OBJECT:EmbeddedVideo]
+	[PARAMLITERAL:width]
+	640
+	[/PARAMLITERAL]
+	[PARAMLITERAL:height]
+	360
+	[/PARAMLITERAL]
+	[PARAMLITERAL:videoid]
+	_UwQz20VxYc
+	[/PARAMLITERAL]
+	[/OBJECT]
 
 You could add a socket to a sprite and attach any number of different things in that socket, whether or not the sprite is part of a Flipbook and is animated as well as what goes in the socket is entirely up to you. 
 
