#UE4doc #docs - Updated image colors, adjusted bold links and removed overview header.

diff --git a/Engine/Documentation/Source/Engine/Paper2D/Sprites/Editor/SpriteEditor.INT.udn b/Engine/Documentation/Source/Engine/Paper2D/Sprites/Editor/SpriteEditor.INT.udn
--- a/Engine/Documentation/Source/Engine/Paper2D/Sprites/Editor/SpriteEditor.INT.udn
+++ b/Engine/Documentation/Source/Engine/Paper2D/Sprites/Editor/SpriteEditor.INT.udn
@@ -10,7 +10,6 @@ Related: Engine/Content/Types/Textures
 
 [TOC(start:2 end:2)]
 
-## Overview
 
 The Paper 2D **Sprite Editor** enables setup and editing of individual Sprite assets.
 
@@ -97,7 +96,7 @@ This shows all of the editable properties of the opened Sprite asset.
 
 ### Sprite List
 
-This shows all of the Sprites that share a common source texture. **Double-clicking** on another Sprite will switch the current asset being edited, which can be useful when adjusting a bunch of Sprites in a row. Another approach for bulk editing is to use the [Property Matrix](GettingStarted/HowTo/DetailsPanel#propertymatrix) context menu option on several selected assets in the **Content Browser**.
+This shows all of the Sprites that share a common source texture. **Double-clicking** on another Sprite will switch the current asset being edited, which can be useful when adjusting a bunch of Sprites in a row. Another approach for bulk editing is to use the [**Property Matrix**](GettingStarted/HowTo/DetailsPanel#propertymatrix) context menu option on several selected assets in the **Content Browser**.
 
 ![](sprite_list.png)
 
@@ -105,9 +104,9 @@ This shows all of the Sprites that share a common source texture. **Double-click
 
 Editing Sprites inside the **Sprite Editor** is done through either the properties inside the **Details** panel or by switching to one of the editing modes from the **Mode Switching Toolbar**. Additional documentation on each of the editing modes can be found on the following pages: 
 
-* For Sprite Source Region editing, see the [Edit Source Region](Engine\Paper2D\Sprites\RenderGeometry#definingthesourceregion) documentation. 
-* For Sprite Collision editing, see the [Edit Collision](Engine\Paper2D\Sprites\Collision) documentation. 
-* For Sprite Render Geometry editing, see the [Edit RenderGeom](Engine\Paper2D\Sprites\RenderGeometry) documentation. 
+* For Sprite Source Region editing, see the [**Edit Source Region**](Engine\Paper2D\Sprites\RenderGeometry#definingthesourceregion) documentation. 
+* For Sprite Collision editing, see the [**Edit Collision**](Engine\Paper2D\Sprites\Collision) documentation. 
+* For Sprite Render Geometry editing, see the [**Edit RenderGeom**](Engine\Paper2D\Sprites\RenderGeometry) documentation. 
 
 
 
