#UE4doc #docs - Updated images and bolded links.

diff --git a/Engine/Documentation/Source/Engine/Paper2D/TileMaps/Paper2DTilesMaps.INT.udn b/Engine/Documentation/Source/Engine/Paper2D/TileMaps/Paper2DTilesMaps.INT.udn
--- a/Engine/Documentation/Source/Engine/Paper2D/TileMaps/Paper2DTilesMaps.INT.udn
+++ b/Engine/Documentation/Source/Engine/Paper2D/TileMaps/Paper2DTilesMaps.INT.udn
@@ -14,7 +14,6 @@ The features on this page are currently in the **experimental stage** and as suc
 
 [TOC (start:2 end:2)]
 
-## Overview
 
 **Tile Sets** and **Tile Maps** within **Paper 2D**, provide a quick and easy way to layout the framework or "general layout" of your 2D levels. 
 By creating and using a Tile Set (a collection of tiles pulled from a texture) with a Tile map (a 2D grid with a defined width and height in tiles), you can select various tiles to "paint" onto the Tile Map which can be used for your level layout. 
@@ -89,11 +88,11 @@ Inside the **Viewport**, **Left-clicking** on your Tile Sheet will display a whi
 
 Above, the default **Tile Width** and **Tile Height** (in pixels) of 32 is being used represented by the white selection box. Unfortunately our tiles are larger than the default values and we want our Tile Width and Height to encompass our entire tile (represented by the yellow box outline). 
 
-Inside the **Tile Editor** preview window the currently selected Tile and its associated Tile Number (see below) are displayed. Just below it, inside the **Details** panel, you can adjust the Tile Width and Tile Height so that it properly encompasses your entire tile. 
+Inside the **Tile Editor** preview window the currently selected Tile and its associated Tile Number (see below) are displayed. Just below it, inside the **Details** panel, you can adjust the _Tile Width_ and _Tile Height_ so that it properly encompasses your entire tile. 
 
 ![](CurrentlySelectedTile.png)
 
-In our example, we've increased the Tile Width and Tile Height to 64 to match our tile size. 
+In our example, we've increased the _Tile Width_ and _Tile Height_ to 64 to match our tile size. 
 
 [REGION:lightbox]
 [![](UpdatedTileSheet.png)(w:800)](UpdatedTileSheet.png)
@@ -103,7 +102,7 @@ _Click for full view._
 
 In the Viewport, the white tile selection box now fully highlights an individual tile on the Tile Sheet (you may need to experiment with different values depending on how your Tile Sheet is set up). The Tile Editor window also updates to show a single tile from our Tile Sheet that is selected. 
 
-In addition to setting Tile Width and Tile Height, you can use the following settings to further adjust your Tile Set from the **Details** panel. 
+In addition to setting _Tile Width_ and _Tile Height_, you can use the following settings to further adjust your Tile Set from the **Details** panel. 
 
 
 | Setting | Description |
@@ -136,7 +135,7 @@ Also inside the **Details** panel, you can specify per tile settings (such as ap
 	![](Collision_3.png)
 
 	[REGION:note]
-	For information on sprite collision types, see the [](Engine\Paper2D\Sprites\Collision) documentation. 
+	For information on sprite collision types, see the [**Sprite Collision**](Engine\Paper2D\Sprites\Collision) documentation. 
 	[/REGION]
 	
 This will set this tile up for collision based on the specified geometry type and shape used. For the collision to actually take effect in-game, when creating your **Tile Map** you will need to set the **Sprite Collision Domain** (see below for information on creating/using Tile Maps). 
