UEDOC-2126

Images have been updated

diff --git a/Engine/Documentation/Source/Engine/Rendering/Materials/HowTo/DetailTexturing/DetailTexturing.INT.udn b/Engine/Documentation/Source/Engine/Rendering/Materials/HowTo/DetailTexturing/DetailTexturing.INT.udn
--- a/Engine/Documentation/Source/Engine/Rendering/Materials/HowTo/DetailTexturing/DetailTexturing.INT.udn
+++ b/Engine/Documentation/Source/Engine/Rendering/Materials/HowTo/DetailTexturing/DetailTexturing.INT.udn
@@ -141,25 +141,26 @@ If for whatever reason you cannot use the **Detail Texturing** Material Function
 
 1. Now that everything has been cleaned up and converted, we can start to lay down the required Material Expression nodes. For this part, we are going to need the following type and amount of nodes. You can find each of the following Material Expression nodes by searching for them in the **Palette**. 
 	* Texture Coordinate x 1
-	* Multiply x 2
+	* Clamp x 1
+	* Constant x 2
 	* Add x 2
-	* Constant  x 2
+	* Multiply x 3
 
 	When completed, the Material Graph should look like this. 
 
-	![](DT_Required_Nodes.png)
+	![](DT_RequiredNodes.jpg)
 
 1. Once all the required Material Expressions nodes have been added, it is now time to begin hooking everything up. The first section that is going to be hooked up is the Base Color portion. When completed, the Material Graph should look like this. 
 	
-	![](DT_Base_Color_Connections.png)
+	![](DT_BaseColor_Connections.jpg)
 
 1. Next the Normal Map needs to be hooked up in the same manner that the Base Color was set up. When completed, the Material Graph should look like this. 
 
-	![](DT_Normal_Map_Connections.png)
+	![](DT_NormalMap_Connections.jpg)
 
 1. With the Base Color and Normal map now all connected together, the Material is now ready to be compiled, saved, and used. 
 
-	![](DT_Completed_Material.png)
+	![](DT_Complete_Material.jpg)
 
 
 ## Detail Texturing Tip & Tricks
