#UE4 DOCS: removed physics references in the volumes actor page, as they can no longer alter gravity.

diff --git a/Engine/Documentation/Source/Engine/Actors/Volumes/Volumes.INT.udn b/Engine/Documentation/Source/Engine/Actors/Volumes/Volumes.INT.udn
--- a/Engine/Documentation/Source/Engine/Actors/Volumes/Volumes.INT.udn
+++ b/Engine/Documentation/Source/Engine/Actors/Volumes/Volumes.INT.udn
@@ -31,7 +31,6 @@ Description:Guide to placing and working with Volumes to augment the behavior an
 **Volumes** are 3-dimensional areas within your levels, each of which serve a specific purpose. Some of these purposes include:
 
 * Causing damage to the Player.
-* Changing the laws of physics, perhaps allowing Players to float or to reduce gravity within the space of the volume.
 * Blocking the Player from entering the volume, acting as a collision surface.
 * Changing the way a level is calculated for lighting or visibility.
 
