UEDOC-2144

Gif has been replaced with a Youtube embedded video

diff --git a/Engine/Documentation/Source/Engine/Paper2D/HowTo/TopDown/TopDown.INT.udn b/Engine/Documentation/Source/Engine/Paper2D/HowTo/TopDown/TopDown.INT.udn
--- a/Engine/Documentation/Source/Engine/Paper2D/HowTo/TopDown/TopDown.INT.udn
+++ b/Engine/Documentation/Source/Engine/Paper2D/HowTo/TopDown/TopDown.INT.udn
@@ -165,7 +165,17 @@ For more detailed information on setting up character movement and setting up ch
 
 1.	From the Main Tool bar, click the **Play** button to play in the editor. 
 
-	![](TopDownExample.gif)(convert:false)
+	[OBJECT:EmbeddedVideo]
+	[PARAMLITERAL:width]
+	640
+	[/PARAMLITERAL]
+	[PARAMLITERAL:height]
+	360
+	[/PARAMLITERAL]
+	[PARAMLITERAL:videoid]
+	92LFHoHWfOk
+	[/PARAMLITERAL]
+	[/OBJECT]
 
 	Although the 2D character is idle, you can now move around the default map using the **WASD** keys. 
 
