#UE4doc #docs Edited documentation based on peer review by MinaW and LaurenR.

diff --git a/Engine/Documentation/Source/GettingStarted/DownloadingUnrealEngine/DownloadZip/DownloadZip.INT.udn b/Engine/Documentation/Source/GettingStarted/DownloadingUnrealEngine/DownloadZip/DownloadZip.INT.udn
--- a/Engine/Documentation/Source/GettingStarted/DownloadingUnrealEngine/DownloadZip/DownloadZip.INT.udn
+++ b/Engine/Documentation/Source/GettingStarted/DownloadingUnrealEngine/DownloadZip/DownloadZip.INT.udn
@@ -16,7 +16,7 @@ If you prefer to download our source code as a Zip file from GitHub, this sectio
 	active_button_content
 	[/PARAMLITERAL]
 	[PARAM:content]
-	1.	**Navigate** to our GitHub page.
+	1.	Navigate to our GitHub page.
 	1.	**Click** the **Download ZIP** button (pictured below).
 		![](DownloadZipButton.png)
 		[REGION:warning]
@@ -29,17 +29,17 @@ If you prefer to download our source code as a Zip file from GitHub, this sectio
 
 		Third party zip utilities may not require the above mentioned steps prior to decompression.
 		[/REGION]
-	1.	**Decompress** 'UnrealEngine-(branch_name).zip' onto your hard disk.
-	1.	If you already haven't done so, **install** Visual Studio 2013. Although MFC libraries are included in Visual Studio's default installation settings, make sure that MFC libraries are included for ATL support.
+	1.	Decompress 'UnrealEngine-(branch_name).zip' onto your hard disk.
+	1.	If you already haven't done so, install Visual Studio 2013. Although MFC libraries are included in Visual Studio's default installation settings, make sure that MFC libraries are included for ATL support.
 		[REGION:note]
 		Developers build UE4 with all desktop editions of Visual Studio, including Microsoft's free version ([Visual Studio Community 2013](https://www.visualstudio.com/products/visual-studio-community-vs)).
 		[/REGION]
-	1.	**Open** your source folder in Explorer and **run Setup.bat**.
+	1.	Open your source folder in Explorer and run **Setup.bat**.
 		[REGION:note]
-		Our Setup batch file downloads binary content for the engine, installs prerequesites, and sets up Unreal file associations. Downloading engine binaries may take some time because they currently take up three to four gigabytes of memory. Subsequent checkouts take less time to download.
+		Our Setup batch file downloads binary content for the engine, installs prerequesites, and sets up Unreal file associations. Downloading engine binaries may take some time because they currently take up three to four gigabytes of disk space. Subsequent checkouts take less time to download.
 		[/REGION]
 		[REGION:warning]
-		A warning from SmartScreen may appear in Windows 8. **Click 'More Info'**, then **'Run Anyway'** to continue running our Setup batch file.
+		A warning from SmartScreen may appear in Windows 8. **Click 'More Info'**, then **click 'Run Anyway'** to continue running our Setup batch file.
 		[/REGION]
 	[/PARAM]
 [/OBJECT]
@@ -55,14 +55,15 @@ If you prefer to download our source code as a Zip file from GitHub, this sectio
 	active_button_content
 	[/PARAMLITERAL]
 	[PARAM:content]
-	1.	**Navigate** to our GitHub page.
+	1.	Navigate to our GitHub page.
 	1.	**Click** the **Download ZIP** button (pictured below).
 		![](DownloadZipButton.png)
-	1.	**Decompress** 'UnrealEngine-(branch_name).zip' onto your hard disk.
-	1.	If you haven't already done so, **install** the latest version of [Xcode](https://itunes.apple.com/us/app/xcode/id497799835).
-	1.	**Open** your source folder in Finder and **run Setup.command**.
+	1.	Decompress 'UnrealEngine-(branch_name).zip' onto your hard disk.
+	1.	If you haven't already done so, install the latest version of [Xcode](https://itunes.apple.com/us/app/xcode/id497799835).
+	1.	Open your source folder in Finder and run **Setup.command**.
 		[REGION:note]
-		**Close** the Terminal window after Setup.command downloads binary content for the engine, installs prequisites, and sets up Unreal file associations.
+		*	Our Setup batch file downloads binary content for the engine, installs prerequesites, and sets up Unreal file associations. Downloading engine binaries may take some time because they currently take up three to four gigabytes of disk space. Subsequent checkouts take less time to download.  
+		*	Close the Terminal window after **Setup.command** downloads binary content for the engine, installs prequisites, and sets up Unreal file associations.
 		[/REGION]
 	[/PARAM]
 [/OBJECT]
@@ -83,16 +84,17 @@ If you prefer to download our source code as a Zip file from GitHub, this sectio
 	At this time, we only provide development support for Ubuntu 15.04.
 	[/REGION]
 
-	1.	**Navigate** to our GitHub page.
+	1.	Navigate to our GitHub page.
 	1.	**Click** the **Download ZIP** button (pictured below).
 		![](DownloadZipButton.png)
-	1.	**Decompress** 'UnrealEngine-(branch_name).zip' onto your hard disk.
-	1.	Inside the root folder, **run Setup.sh** from the terminal to **download** the engine's binary content (example pictured below).
+	1.	Decompress 'UnrealEngine-(branch_name).zip' onto your hard disk.
+	1.	Inside the root folder, run **Setup.sh** from the terminal to download the engine's binary content (example pictured below).
 
 		![](RunSetupShellScript_Linux.png)
 
 		[REGION:note]
-		If you don't see a ** \*\*SUCCESS\*\* ** message after running this step, refer to the **'BuildThirdParty.log'** file for more information. 'BuildThirdParty.log' is located in the `[UE4Root]/Engine/Build/BatchFiles/Linux` directory.
+		*	Our Setup batch file downloads binary content for the engine, installs prerequesites, and sets up Unreal file associations. Downloading engine binaries may take some time because they currently take up three to four gigabytes of disk space. Subsequent checkouts take less time to download.  
+		*	If you don't see a ** \*\*SUCCESS\*\* ** message after running this step, refer to the **'BuildThirdParty.log'** file for more information. **'BuildThirdParty.log'** is located in the `[UE4Root]/Engine/Build/BatchFiles/Linux` directory.
 		[/REGION]
 	[/PARAM]
 [/OBJECT]
diff --git a/Engine/Documentation/Source/GettingStarted/DownloadingUnrealEngine/DownloadingUnrealEngine4.INT.udn b/Engine/Documentation/Source/GettingStarted/DownloadingUnrealEngine/DownloadingUnrealEngine4.INT.udn
--- a/Engine/Documentation/Source/GettingStarted/DownloadingUnrealEngine/DownloadingUnrealEngine4.INT.udn
+++ b/Engine/Documentation/Source/GettingStarted/DownloadingUnrealEngine/DownloadingUnrealEngine4.INT.udn
@@ -3,8 +3,8 @@ Title:Downloading Unreal Engine Source Code
 Crumbs: %ROOT%, Programming, Programming/Development
 Description:Steps for connecting to the source code repository and downloading the latest build of Unreal Engine.
 Prereq: GettingStarted/RecommendedSpecifications
-Prereq: Support/Builds
 Related: Programming/Development/BuildingUnrealEngine
+Related: Support/Builds
 Version: 4.8
 
 [TOC (start:2 end:3)]
@@ -21,9 +21,11 @@ This page shows subscribers how to download and build Unreal Engine from our sou
 
 ----
 
-![](WindowsMacLinux.png)
+Source code for Unreal Engine 4 (UE4) can be downloaded by our subscribers from the Unreal Engine GitHub repository. 
 
-Source code for Unreal Engine 4 (UE4) can be downloaded by our subscribers from the Unreal Engine GitHub repository. With access to our source code, subscribers can customize the engine and tools to meet their needs. Our subscribers are building the editor in Windows, Mac, Linux, and compiling games for both Android and iOS. Source code for all of our tools (such as Unreal Lightmass and Frontend) is included as well.
+<!-- The following is commented out per LaurenR's request during Peer Review on 23.07.15
+With access to our source code, subscribers can customize the engine and tools to meet their needs. Our subscribers are building the editor in Windows, Mac, Linux, and compiling games for both Android and iOS. Source code for all of our tools (such as Unreal Lightmass and Frontend) is included as well.
+-->
 
 [REGION:note]
 Before you can access the [repository](https://github.com/EpicGames/UnrealEngine), you must be an Unreal Engine subscriber with a GitHub account.  
@@ -43,7 +45,9 @@ Once you've accessed the repository, you'll notice that we've published UE4's so
 
 (#Bookmark1)
 
+[REGION:raw]
 ![](ReleaseBranch.png)
+[/REGION]
 
 The [release branch](https://github.com/EpicGames/UnrealEngine/tree/release) (pictured above) is extensively tested by our QA team and makes a great starting point for learning UE4 or making Unreal games. We work hard to stabilize this branch with the goal of releasing new builds every one to two months.
 
@@ -51,7 +55,9 @@ The [release branch](https://github.com/EpicGames/UnrealEngine/tree/release) (pi
 
 (#Bookmark2)
 
+[REGION:raw]
 ![](PromotedBranch.png)
+[/REGION]
 
 The [promoted branch](https://github.com/EpicGames/UnrealEngine/tree/promoted) (pictured above) is updated with builds for use by our designers and artists. We try to update this branch daily so that our subscribers have access to UE4's latest stable features. This branch is a great place to learn how feature development works in UE4.
 
@@ -63,11 +69,15 @@ The [promoted branch](https://github.com/EpicGames/UnrealEngine/tree/promoted) (
 Because the master branch contains cutting-edge technology, UE4 isn't guaranteed to compile from source.
 [/REGION]
 
+[REGION:raw]
 ![](MasterBranch.png)
+[/REGION]
 
 The [master branch](https://github.com/EpicGames/UnrealEngine/tree/master) (pictured above) tracks [live changes](https://github.com/EpicGames/UnrealEngine/commits/master) (pictured below) being made by our engine team. Our experienced subscribers typically work from this branch when they want to work in tandem with UE4's developers.
 
+[REGION:raw]
 ![](LiveChanges.png)
+[/REGION]
 
 [REGION:note]
 Other branches may appear from time-to-time as we work to stabilize new releases or hotfixes.
diff --git a/Engine/Documentation/Source/GettingStarted/DownloadingUnrealEngine/GitHub/GitHub.INT.udn b/Engine/Documentation/Source/GettingStarted/DownloadingUnrealEngine/GitHub/GitHub.INT.udn
--- a/Engine/Documentation/Source/GettingStarted/DownloadingUnrealEngine/GitHub/GitHub.INT.udn
+++ b/Engine/Documentation/Source/GettingStarted/DownloadingUnrealEngine/GitHub/GitHub.INT.udn
@@ -18,23 +18,23 @@ Forking and cloning our source code from GitHub is a great way to get started wi
 	active_button_content
 	[/PARAMLITERAL]
 	[PARAM:content]
-	1.	**Install** [GitHub for Windows](https://windows.github.com/).
+	1.	Install [GitHub for Windows](https://windows.github.com/).
 	1.	**Fork** and **Clone** our [repository](https://github.com/EpicGames/UnrealEngine) (pictured below).
 		![](UnrealRepo.png)
 		[REGION:tip]
 		*	Refer to GitHub's [Forking Projects Guide](https://guides.github.com/activities/forking/) to learn how to fork and clone our repository.
 		*	To use Git from the command line, refer to GitHub's articles on [Setting-up Git](https://help.github.com/articles/set-up-git/) and [Forking a Repository](https://help.github.com/articles/fork-a-repo/).
 		[/REGION]
-	1.	If you already haven't done so, **install** Visual Studio 2013. Although MFC libraries are included in Visual Studio's default installation settings, make sure that MFC libraries are included for ATL support.
+	1.	If you already haven't done so, install Visual Studio 2013. Although MFC libraries are included in Visual Studio's default installation settings, make sure that MFC libraries are included for ATL support.
 		[REGION:note]
 		Developers build UE4 with all desktop editions of Visual Studio, including Microsoft's free version ([Visual Studio Community 2013](https://www.visualstudio.com/products/visual-studio-community-vs)).
 		[/REGION]
-	1.	**Open** your source folder in Explorer and **run Setup.bat**.
+	1.	Open your source folder in Explorer and run **Setup.bat**.
 		[REGION:note]
-		Our Setup batch file downloads binary content for the engine, installs prerequesites, and sets up Unreal file associations. Downloading engine binaries may take some time because they currently take up three to four gigabytes of memory. Subsequent checkouts take less time to download.
+		Our Setup batch file downloads binary content for the engine, installs prerequesites, and sets up Unreal file associations. Downloading engine binaries may take some time because they currently take up three to four gigabytes of disk space. Subsequent checkouts generally take less time to download.
 		[/REGION]
 		[REGION:warning]
-		A warning from SmartScreen may appear in Windows 8. **Click 'More Info'**, then **'Run Anyway'** to continue running our Setup batch file.
+		A warning from SmartScreen may appear in Windows 8. **Click 'More Info'**, then **click 'Run Anyway'** to continue running our Setup batch file.
 		[/REGION]
 	[/PARAM]
 [/OBJECT]
@@ -50,17 +50,18 @@ Forking and cloning our source code from GitHub is a great way to get started wi
 	active_button_content
 	[/PARAMLITERAL]
 	[PARAM:content]
-	1.	**Install** [GitHub for Mac](https://mac.github.com/).
+	1.	Install [GitHub for Mac](https://mac.github.com/).
 	1.	**Fork** and **Clone** our [repository](https://github.com/EpicGames/UnrealEngine) (pictured below).
 		![](UnrealRepo.png)
 		[REGION:tip]
 		*	Refer to GitHub's [Forking Projects Guide](https://guides.github.com/activities/forking/) to learn how to fork and clone our repository.
 		*	To use Git from the command line, refer to GitHub's articles on [Setting-up Git](https://help.github.com/articles/set-up-git/) and [Forking a Repository](https://help.github.com/articles/fork-a-repo/).
 		[/REGION]
-	1.	If you haven't already done so, **install** the latest version of [Xcode](https://itunes.apple.com/us/app/xcode/id497799835).
-	1.	**Open** your source folder in Finder and **run Setup.command**.
+	1.	If you haven't already done so, install the latest version of [Xcode](https://itunes.apple.com/us/app/xcode/id497799835).
+	1.	Open your source folder in Finder and run **Setup.command**.
 		[REGION:note]
-		**Close** the Terminal window after Setup.command downloads binary content for the engine, installs prequisites, and sets up Unreal file associations.
+		*	Our Setup command file downloads binary content for the engine, installs prerequesites, and sets up Unreal file associations. Downloading engine binaries may take some time because they currently take up three to four gigabytes of disk space. Subsequent checkouts generally take less time to download.  
+		*	Close the Terminal window after **Setup.command** downloads binary content for the engine, installs prequisites, and sets up Unreal file associations.
 		[/REGION]
 	[/PARAM]
 [/OBJECT]
@@ -81,15 +82,16 @@ Forking and cloning our source code from GitHub is a great way to get started wi
 	At this time, we only provide development support for Ubuntu 15.04.
 	[/REGION]
 
-	1.	**Set up** [Git](https://help.github.com/articles/set-up-git/).
+	1.	Set up [Git](https://help.github.com/articles/set-up-git/).
 	1.	**Fork** and **Clone** our [repository](https://help.github.com/articles/fork-a-repo/) (pictured below).
 		![](UnrealRepo.png)
-	1.	Inside the root folder, **run Setup.sh** from the terminal to **download** the engine's binary content (example pictured below).
+	1.	Inside the root folder, run **Setup.sh** from the terminal to download the engine's binary content (example pictured below).
 
 		![](RunSetupShellScript_Linux.png)
 
 		[REGION:note]
-		If you don't see a ** \*\*SUCCESS\*\* ** message after running this step, refer to the **'BuildThirdParty.log'** file for more information. 'BuildThirdParty.log' is located in the `[UE4Root]/Engine/Build/BatchFiles/Linux` directory.
+		*	Our Setup shell script downloads binary content for the engine, installs prerequesites, and sets up Unreal file associations. Downloading engine binaries may take some time because they currently take up three to four gigabytes of disk space. Subsequent checkouts generally take less time to download.  
+		*	If you don't see a ** \*\*SUCCESS\*\* ** message after running this step, refer to the **'BuildThirdParty.log'** file for more information. **'BuildThirdParty.log'** is located in the `[UE4Root]/Engine/Build/BatchFiles/Linux` directory.
 		[/REGION]
 	[/PARAM]
 [/OBJECT]
diff --git a/Engine/Documentation/Source/GettingStarted/Installation/InstallingUnrealEngine.INT.udn b/Engine/Documentation/Source/GettingStarted/Installation/InstallingUnrealEngine.INT.udn
--- a/Engine/Documentation/Source/GettingStarted/Installation/InstallingUnrealEngine.INT.udn
+++ b/Engine/Documentation/Source/GettingStarted/Installation/InstallingUnrealEngine.INT.udn
@@ -3,29 +3,36 @@ Title:Installing Unreal Engine
 Description:Steps for installing Unreal Engine.
 Crumbs: %ROOT%, GettingStarted
 Prereq: GettingStarted/RecommendedSpecifications
-Prereq: Support/Builds
 Related: GettingStarted
 Related: Engine/Content/QuickStart
 Related: Programming/QuickStart
 Related: Engine/QuickStart
+Related: Support/Builds
 Version: 4.8
 
-[TOC (start:2 end:3)]
+[REGION:warning]
+We currently don't provide a binary installer of Unreal Engine for Linux.  
+Linux users; however, can compile Unreal Engine from source by referring to our [](GettingStarted/DownloadingUnrealEngine) documentation.
+[/REGION]
 
-----
+By the end of this tutorial, you'll be able to:
 
-We make it easy for subscribers to download, install, and run Unreal Engine 4 (UE4) from our launcher.
+*	[Create an Epic Games account _(if you don't already have one)_](#Bookmark1)
+*	[Download and Run the Installer (Epic Games Launcher Set-Up Program)](#Bookmark2)
+*	[Sign into the Epic Games Launcher](#Bookmark3)
+*	[Install Unreal Engine](#Bookmark4)
+*	[Launch Unreal Engine](#Bookmark5)
 
 ![](InstallationGoal.png)
 
-_By the end of this tutorial, you'll be able to run the Unreal Editor from the **Epic Games Launcher** (pictured above)._
-
-The Epic Games Launcher enables our subscribers to conveniently download and install updates for UE4.
+_The Epic Games Launcher enables our subscribers to conveniently download and install updates for UE4 (example pictured above)._
 
 ## Creating an Epic Games Account
 
+(#Bookmark1)
+
 [REGION:warning]
-If you already have an Epic Games account, skip ahead to learn about [Downloading and Running the Installer](#Bookmark1).
+If you already have an Epic Games account, skip ahead to learn about [Downloading and Running the Installer](#Bookmark2).
 [/REGION]
 
 In order to run the Epic Games Launcher, you need to have a valid Epic Games account.
@@ -54,9 +61,9 @@ In order to run the Epic Games Launcher, you need to have a valid Epic Games acc
 
 ## Downloading and Running the Installer
 
-(#Bookmark1)
+(#Bookmark2)
 
-With a valid Epic Games account, you're able to download our Epic Games Launcher set-up program (installer).
+With a valid Epic Games account, you're able to download our installer (Epic Games Launcher Set-Up Program).
 
 %Globals:osselect_linux%
 
@@ -128,12 +135,16 @@ With a valid Epic Games account, you're able to download our Epic Games Launcher
 
 ## Signing into the Epic Launcher
 
-After the set-up program installs the Epic Games Launcher on your machine, you're ready to sign-in (see below).
+(#Bookmark3)
+
+After the set-up program installs the Epic Games Launcher on your machine, use your Epic Game account credentials to login (see below).
 
 ![](EpicSignIn.png)
 
 ## Installing Unreal Engine
 
+(#Bookmark4)
+
 Now that you've signed into the Epic Games Launcher, you're ready to install UE4.
 
 [REGION:warning]
@@ -147,7 +158,7 @@ Unreal Engine uses approximately eight gigabytes of disk space for every version
 
 	![](LauncherLoadScreen.png)
 
-1.	**Click** the **Install Engine** button to **download and install** the latest version of Unreal Engine (pictured below).
+1.	**Click** the **Install Engine** button to download and install the latest version of Unreal Engine (pictured below).
 
 	![](InstallEngine.png)
 
@@ -155,7 +166,7 @@ Unreal Engine uses approximately eight gigabytes of disk space for every version
 
 	![](AcceptEULALauncher.png)
 
-1.	Depending on your system specifications, **downloading and installing** Unreal Engine will take between 10-40 minutes.
+1.	Depending on your system specifications, downloading and installing Unreal Engine will take between 10-40 minutes.
 	
 	![](DownloadingEngine.png)
 
@@ -163,21 +174,35 @@ Unreal Engine uses approximately eight gigabytes of disk space for every version
 
 ## Launching Unreal Engine
 
+(#Bookmark5)
+
 After the Epic Games Launcher successfully downloads and installs Unreal Engine, you're ready to **click** the **Launch** button (see below).
 
 ![](LaunchUE4.png)
 
+Congratulations! You've learned how to:
+
+&#x2713;	[Create an Epic Games account _(if you didn't already have one)_](#Bookmark1)  
+&#x2713;	[Download and Run the Installer (Epic Games Launcher Set-Up Program)](#Bookmark2)  
+&#x2713;	[Sign into the Epic Games Launcher](#Bookmark3)  
+&#x2713;	[Install Unreal Engine](#Bookmark4)  
+&#x2713;	[Launch Unreal Engine](#Bookmark5)
+
+Now, you're ready to start using Unreal Engine.
+
 ## Getting Started with Unreal Engine
 
 Learn how to use Unreal Engine by referring to the [](GettingStarted) documentation!
 
-[REGION:tip]
-If you're an artist, programmer, or designer looking to quickly get started with UE4, refer to the following tutorials:
+If you're looking to quickly get started with UE4, refer to the following tutorials:
 *	[](Engine/Content/QuickStart)
 *	[](Programming/QuickStart)
 *	[](Engine/QuickStart)
-[/REGION]
 
+
+[REGION:tip]
+UE4's in-editor help features are a great way to get your questions answered.
+[/REGION]
 <!--Deprecated on 21 July 2015 due to outdated screenshots and outdated procedures
 %Globals:OSSelect%
 
diff --git a/Engine/Documentation/Source/Programming/Development/BuildingUnrealEngine/BuildingUnrealEngine.INT.udn b/Engine/Documentation/Source/Programming/Development/BuildingUnrealEngine/BuildingUnrealEngine.INT.udn
--- a/Engine/Documentation/Source/Programming/Development/BuildingUnrealEngine/BuildingUnrealEngine.INT.udn
+++ b/Engine/Documentation/Source/Programming/Development/BuildingUnrealEngine/BuildingUnrealEngine.INT.udn
@@ -9,6 +9,7 @@ Related: GettingStarted
 Related: Engine/Content/QuickStart
 Related: Programming/QuickStart
 Related: Engine/QuickStart
+Related: Support/Builds
 Version: 4.8
 
 %Globals:osselect_linux%
@@ -33,19 +34,19 @@ Version: 4.8
 		Make sure that Microsoft Visual Studio 2013 is installed prior to building Unreal Engine from source.
 		[/REGION]
 
-		1.	Inside the root directory, **run GenerateProjectFiles.bat** to set-up your project files.
+		1.	Inside the root directory, run **GenerateProjectFiles.bat** to set-up your project files.
 
 			[REGION:warning]
 			All project files are intermediate (`[UE4Root]\Engine\Intermediate\ProjectFiles`). You must generate project files each time you sync a new build to ensure they are up to date. If you delete your `Intermediate` folder, you must regenerate project files using the 'GenerateProjectFiles' batch file.
 			[/REGION]
 
-		1.	**Load** the project into Visual Studio by **double-clicking UE4.sln**.
+		1.	Load the project into Visual Studio by **double-clicking UE4.sln**.
 
 		1.	Set your solution configuration to **Development Editor** (pictured below).
 
 			![](configMenu_updated.png)
 
-		1.	Depending on your target-version of the Windows operating system (32-bit versus 64-bit), set your solution platform to either **Win32** or **Win64** (pictured below).
+		1.	Set your solution platform to **Win64** (pictured below).
 
 			![](platform_menu_updated.png)
 
@@ -61,12 +62,12 @@ Version: 4.8
 
 			![](SetUE4_StartPrj.png)
 
-		1.	**Right-click** the **UE4** project, **hover over 'Debug'** and **select 'Start New Instance'** to launch the compiled editor in debug mode (pictured below).
+		1.	**Right-click** the **UE4** project, **hover over 'Debug'** and **click 'Start New Instance'** to launch the editor (pictured below).
 
 			![](RunCompiledWindowsEditor.png)
 
 			[REGION:tip]
-			Alternatively, you can **press the F5 key** to start a new instance of the compiled editor in debug mode.
+			Alternatively, you can **press the F5 key** on your keyboard to start a new instance of the editor.
 			[/REGION]
 
 		1.	Congratulations! You've compiled and launched the engine from source (result pictured below).
@@ -88,7 +89,7 @@ Version: 4.8
 		Make sure that the latest version of XCode is installed prior to building Unreal Engine from Source.
 		[/REGION]
 		
-		1.	Inside the root directory, **run GenerateProjectFiles.command** to set-up your project files.
+		1.	Inside the root directory, run **GenerateProjectFiles.command** to set-up your project files.
 		
 		1.	Load the project into XCode by **double-clicking UE4.xcodeproj**.
 
@@ -127,11 +128,11 @@ Version: 4.8
 		At this time, we only provide development support for Ubuntu 15.04.
 		[/REGION]
 		
-		1.	Inside the root directory, **run GenerateProjectFiles.sh** from the terminal to set-up your project files (example pictured below). 
+		1.	Inside the root directory, run **GenerateProjectFiles.sh** from the terminal to set-up your project files (example pictured below). 
 
 			![](RunGenPrjFilesShellScript_Linux.png)
 
-		1.	To build the project, **run make** from the terminal (example pictured below).
+		1.	To build the project, run **make** from the terminal (example pictured below).
 
 			![](RunMakeFile_Linux.png)
 
@@ -155,11 +156,14 @@ Version: 4.8
 
 Learn how to use Unreal Engine by referring to the [](GettingStarted) documentation!
 
-[REGION:tip]
-If you're an artist, programmer, or designer looking to quickly get started with UE4, refer to the following tutorials:
+If you're looking to quickly get started with UE4, refer to the following tutorials:
 *	[](Engine/Content/QuickStart)
 *	[](Programming/QuickStart)
 *	[](Engine/QuickStart)
+
+
+[REGION:tip]
+UE4's in-editor help features are a great way to get your questions answered.
 [/REGION]
 
 [/EXCERPT:Main]
