UEDOC-898

Images have been updated

diff --git a/Engine/Documentation/Source/Engine/Landscape/Editing/SculptMode/SculptMode.INT.udn b/Engine/Documentation/Source/Engine/Landscape/Editing/SculptMode/SculptMode.INT.udn
--- a/Engine/Documentation/Source/Engine/Landscape/Editing/SculptMode/SculptMode.INT.udn
+++ b/Engine/Documentation/Source/Engine/Landscape/Editing/SculptMode/SculptMode.INT.udn
@@ -164,7 +164,7 @@ The Erosion tool uses a thermal erosion simulation to adjust the height of the h
 | **Iterations** | The number of iterations performed. Larger values result in more levels of erosion. |
 | **Noise Mode** | Whether to apply noise to raise or lower the heightmap, or do both. |
 | **Noise Scale** | The size of the noise filter used. The noise filter is related to position and scale, which means if you do not change **Noise Scale**, the same filter is applied to the same position many times. | 
-| ![Landscape_Thermal_Erosion_Filter.jpg](Landscape_Thermal_Erosion_Filter.jpg)(w:900 h:325) ||
+| ![](Erosion_New.jpg)(w:900 h:325) ||
 
 
 
@@ -185,7 +185,7 @@ The Hydro Erosion tool uses a hydraulic erosion simulation, i.e. erosion from wa
 | **Sediment Cap.** | The amount of sediment that the water can carry. Larger values will result in more erosion. |
 | **Iterations** | The number of iterations performed. Larger values result in more levels of erosion. |
 | **Initial Rain Distribution** | Whether rain is applied only where the noise filter is positive, or to the entire area. |
-| ![Landscape_Hydraulic_Erosion_Filter.jpg](Landscape_Hydraulic_Erosion_Filter.jpg)(w:900 h:325) ||
+| ![](Hydro_Erosion.jpg)(w:900 h:325) ||
 | **Rain Dist Scale** | The size of the noise filter for applying initial rain to the surface. The noise filter is related to position and scale, which means if you do not change **Rain Dist. Scale**, the same filter is applied to the same position many times. | 
 | **Detail Smooth** | If checked, performs a detail-preserving smoothing to the erosion effect using the specified detail-smoothing value. Larger detail-smoothing values remove more details, while smaller values preserve more details. |
 
@@ -202,7 +202,7 @@ The Noise tool applies a noise filter to the heightmap or layer weight. The stre
 | **Tool Strength** | Controls how much effect each brush stroke has. |
 | **Use Target Value** | If checked, blends the values of the noise being applied toward a target value. |
 | **Noise Mode** | Determines whether to apply all noise effects, only the noise effects that result in raising the heightmap, or only the noise effects that result in lowering the heightmap. |
-| ![Landscape_Noise_Filter.jpg](Landscape_Noise_Filter.jpg)(w:900 h:325) ||
+| ![](Noise.jpg)(w:900 h:325) ||
 | **Noise Scale** | The size of the perlin noise filter used. The noise filter is related to position and scale, which means if you do not change **Noise Scale**, the same filter is applied to the same position many times. |
 
 
