#UE4doc #docs - adjusted bold lettering.

diff --git a/Engine/Documentation/Source/Engine/QuickStart/4/LDQS_4.INT.udn b/Engine/Documentation/Source/Engine/QuickStart/4/LDQS_4.INT.udn
--- a/Engine/Documentation/Source/Engine/QuickStart/4/LDQS_4.INT.udn
+++ b/Engine/Documentation/Source/Engine/QuickStart/4/LDQS_4.INT.udn
@@ -29,13 +29,13 @@ Related:Engine/Rendering/LightingAndShadows
 
 %Steps%
 
-In this step we will begin placing [Actors](Engine/Actors) (which are things like lights or geometry for example) into our empty level. 
-We will cover the two most common ways of adding Actors to the level, one of which is through [Place Mode](Engine/UI/LevelEditor/Modes/PlaceMode) and the other is through the [Content Browser](Engine/Content/Browser). 
+In this step we will begin placing [**Actors**](Engine/Actors) (which are things like lights or geometry for example) into our empty level. 
+We will cover the two most common ways of adding Actors to the level, one of which is through [**Place Mode**](Engine/UI/LevelEditor/Modes/PlaceMode) and the other is through the [**Content Browser**](Engine/Content/Browser). 
 After completing this step, you will know how to place Actors inside your own levels and can begin manipulating those Actors to create your environment. 
 
 ## Steps
 
-1. In the **Modes Panel**, with **Place Mode** enabled, click the [BSP](Engine/Actors/Brushes) category and select the **Box**. 
+1. In the **Modes Panel**, with **Place Mode** enabled, click the [**BSP**](Engine/Actors/Brushes) category and select the **Box**. 
 
 	![](ModesPanel.png)
 
@@ -47,7 +47,7 @@ After completing this step, you will know how to place Actors inside your own le
 
 	![](place2.png)
 
-1. With the **Box** still selected, in the [Details](Engine/UI/LevelEditor/Details) panel (lower-right window of the editor), set **Location** and **Rotation** all to **0**. 
+1. With the **Box** still selected, in the [**Details**](Engine/UI/LevelEditor/Details) panel (lower-right window of the editor), set **Location** and **Rotation** all to **0**. 
 
 1. Set the **Scale** to **4 x 4 x 0.1**. 
 
@@ -59,45 +59,45 @@ After completing this step, you will know how to place Actors inside your own le
 
 	![](placeLight.png)
 	
-1. On the on the [Translation Tool](Engine/UI/LevelEditor/Viewports/ViewportControls), click and drag the Z Axis (Blue) gizmo up away from the surface of the floor.
+1. On the on the [**Translation Tool**](Engine/UI/LevelEditor/Viewports/ViewportControls), click and drag the Z-Axis (Blue) gizmo up away from the surface of the floor.
 
 	![](moveLight.png)
 
 	If the **Directional Light** becomes unselected, you can re-select it by **Left-clicking** on it in the **Level Viewport**. 
 
-1. In the **Modes Panel**, select the **Visual Effects** tab and drag-and-drop an [Atmospheric Fog](Engine/Actors/FogEffects/AtmosphericFog) to the level. 
+1. In the **Modes Panel**, select the **Visual Effects** tab and drag-and-drop an [**Atmospheric Fog**](Engine/Actors/FogEffects/AtmosphericFog) to the level. 
 
 	![](placeFog.png)
 	
 	The **Atmospheric Fog** Actor will add a basic sky to the level and the level will become illuminated instead of dark. 
 
-1. From the **Modes Panel** on the **Basic** tab, drag-and-drop a [Player Start](Engine/Actors/PlayerStart) into the level. 
+1. From the **Modes Panel** on the **Basic** tab, drag-and-drop a [**Player Start**](Engine/Actors/PlayerStart) into the level. 
 
 	![](placeStart.png)
 
-1. From the **Modes Panel** on the **Volumes** tab, drag-and-drop a [Lightmass Importance Volume](Engine/Rendering/LightingAndShadows/Lightmass/Basics) into the level. 
+1. From the **Modes Panel** on the **Volumes** tab, drag-and-drop a [**Lightmass Importance Volume**](Engine/Rendering/LightingAndShadows/Lightmass/Basics) into the level. 
 
 	![](placeVolume.png)
 
 	The **Lightmass Importance Volume** is used to control and concentrate lighting and shadowing affects within the volume. When placing the Lightmass Importance Volume in the level, the default size of the volume does not cover our playable area so we will need to scale it up. 
 	
-1. Inside the **Level Viewport**, press **R** to switch to the [Scale Tool](Engine/UI/LevelEditor/Viewports/ViewportControls#transformcontrols). 
+1. Inside the **Level Viewport**, press **R** to switch to the [**Scale Tool**](Engine/UI/LevelEditor/Viewports/ViewportControls#transformcontrols). 
 
 1. Click and drag the white box in the center of the **Scale Tool** so that the **Lightmass Importance Volume** encapsulates the floor. 
 
 	![](scaleVolume.png)
 
-1. Inside the **Content Browser** under **Content > StarterContent > Props**, drag-and-drop the **SM_TableRound** into the level. 
+1. Inside the **Content Browser** under **Content** > **StarterContent** > **Props**, drag-and-drop the **SM_TableRound** into the level. 
 
 	![](placeTable.png)
 
 	Try to place the table in the center of the floor using the **Move Tool** (Keyboard Shortcut **W** if it is not selected). 
 
-1. Also under **Content > StarterContent > Props**, drag-and-drop the **SM_Chair** into the level. 
+1. Also under **Content** > **StarterContent** > **Props**, drag-and-drop the **SM_Chair** into the level. 
 
 	![](addChair.png)
 
-1. With the **SM_Chair** selected inside the **Level Viewport**, press **E** to access the [Rotation Tool](Engine/UI/LevelEditor/Viewports/ViewportControls#transformcontrols). 
+1. With the **SM_Chair** selected inside the **Level Viewport**, press **E** to access the [**Rotation Tool**](Engine/UI/LevelEditor/Viewports/ViewportControls#transformcontrols). 
 	
 	![](Rotate1.png)(w:440)
 
@@ -109,7 +109,7 @@ After completing this step, you will know how to place Actors inside your own le
 
 	![](Engine/QuickStart/5/placeThings.png)
 
-	Try adding some lights, props, walls and a roof (found under the **Content > StarterContent > Architecture** folder). 
+	Try adding some lights, props, walls and a roof (found under the **Content** > **StarterContent** > **Architecture** folder). 
 
 ## End Result
 
