#UE4doc #docs - updated engine version tested and adjusted bold lettering where needed.

diff --git a/Engine/Documentation/Source/Engine/QuickStart/4/LDQS_4.INT.udn b/Engine/Documentation/Source/Engine/QuickStart/4/LDQS_4.INT.udn
--- a/Engine/Documentation/Source/Engine/QuickStart/4/LDQS_4.INT.udn
+++ b/Engine/Documentation/Source/Engine/QuickStart/4/LDQS_4.INT.udn
@@ -3,7 +3,7 @@ Title:4 - Placing Actors in the Level
 Crumbs:%ROOT%, GettingStarted, Engine/QuickStart
 Description:Here we will begin placing and manipulating Actors in our level to build out our environment. 
 SkillLevel:Beginner
-Version:4.8 
+Version:4.9 
 Related:Engine/UI/LevelEditor/Modes/PlaceMode
 Related:GettingStarted/HowTo/ManipulatingActors
 Related:Engine/UI/LevelEditor/Details
diff --git a/Engine/Documentation/Source/Engine/QuickStart/5/LDQS_5.INT.udn b/Engine/Documentation/Source/Engine/QuickStart/5/LDQS_5.INT.udn
--- a/Engine/Documentation/Source/Engine/QuickStart/5/LDQS_5.INT.udn
+++ b/Engine/Documentation/Source/Engine/QuickStart/5/LDQS_5.INT.udn
@@ -3,7 +3,7 @@ Title:5 - Editing Placed Actors
 Crumbs:%ROOT%, GettingStarted, Engine/QuickStart
 Description:This this step we will take the Actors we placed and edit their properties to give our level a more customized look. 
 SkillLevel:Beginner
-Version:4.8
+Version:4.9
 Related:Engine/Rendering/Materials
 Related:Engine/Rendering/LightingAndShadows/LightTypes
 Related:Engine/Content/FBX/ImportOptions
@@ -27,8 +27,8 @@ Related:
 
 %Steps%
 
-With several different **Actors** placed inside our level, the next step involves [Editing Actor Properites ](GettingStarted/HowTo/DetailsPanel) which can change the look or the way the Actor functions in the level, giving us a more customized looking level. 
-We will start with editing the properties of our **Directional Light Actor** then shift our focus to applying [Materials](Engine/Rendering/Materials) to some of the [Static Mesh Actors](Engine/Actors/StaticMeshActor) that you have placed in your level. 
+With several different **Actors** placed inside our level, the next step involves [**Editing Actor Properites**](GettingStarted/HowTo/DetailsPanel) which can change the look or the way the Actor functions in the level, giving us a more customized looking level. 
+We will start with editing the properties of our **Directional Light Actor** then shift our focus to applying [**Materials**](Engine/Rendering/Materials) to some of the [**Static Mesh Actors**](Engine/Actors/StaticMeshActor) that you have placed in your level. 
 
 Once you have finished this step, you will have seen where to access and modify the properties of Actors so that you can begin editing and experimenting with different settings inside your own levels. 
 
@@ -38,11 +38,11 @@ Once you have finished this step, you will have seen where to access and modify
 
 	![](selectDL.png)
 
-1. In the **Details Panel** under the **Light** category, enable **Used as Atmosphere Sun Light**:
+1. In the **Details Panel** under the **Light** category, enable  **Atmosphere Sun Light**:
 
 	![](enableSun.png)
 
-	Depending on the rotation of your **Directional Light Actor** the sky color will change and if you rotate the Viewport around you will see that the sun now aligns with the Directional Light Actor. 
+	Depending on the rotation of your Directional Light Actor the sky color will change and if you rotate the Viewport around you will see that the sun now aligns with the Directional Light Actor. 
 	This is a real time process so you can rotate the Directional Light Actor (press **E** to switch to **Rotation Mode**) and the sky will change color from night, sunrise, day time, and sunset.
 
 Next we will change the **Material** on one of your placed **Static Mesh Actors** by first selecting it. 
diff --git a/Engine/Documentation/Source/Engine/QuickStart/6/LDQS_6.INT.udn b/Engine/Documentation/Source/Engine/QuickStart/6/LDQS_6.INT.udn
--- a/Engine/Documentation/Source/Engine/QuickStart/6/LDQS_6.INT.udn
+++ b/Engine/Documentation/Source/Engine/QuickStart/6/LDQS_6.INT.udn
@@ -3,7 +3,7 @@ Title:6 - Running the Build Process
 Crumbs:%ROOT%, GettingStarted, Engine/QuickStart
 Description:With a number of Actors placed and customized, one of the final steps for creating any level is the Build Process.
 SkillLevel:Beginner
-Version:4.8
+Version:4.9
 
 [VAR:Steps]
 [OBJECT:Navigation]
diff --git a/Engine/Documentation/Source/Engine/QuickStart/7/LDQS_7.INT.udn b/Engine/Documentation/Source/Engine/QuickStart/7/LDQS_7.INT.udn
--- a/Engine/Documentation/Source/Engine/QuickStart/7/LDQS_7.INT.udn
+++ b/Engine/Documentation/Source/Engine/QuickStart/7/LDQS_7.INT.udn
@@ -3,7 +3,7 @@ Title:7 - On Your Own!
 Crumbs:%ROOT%, GettingStarted, Engine/QuickStart
 Description:Using the methods provided during this guide, polish up your level with additional content and read up on related topics. 
 SkillLevel:Beginner
-Version:4.7
+Version:4.9
 
 [VAR:Steps]
 [OBJECT:Navigation]
@@ -26,7 +26,7 @@ Using the methods that were provided during the course of this guide, try do to
 
 ![](update1.png)(w:830)
 
-* Change the lighting of the level to a moonlight, night scene. 
+* Change the lighting of the level to a moonlit, night scene. 
 * Add another room, attached to the first room. 
 
 ![](update2.png)(h:380) ![](update4.png)(h:380)
@@ -39,23 +39,23 @@ Using the methods that were provided during the course of this guide, try do to
 * Add different kinds of lights with different colors. 
 * Use different Materials on some of your Actors. 
 
-For more information on the topics covered in this quick start guide, and across the entire editor, see the [Unreal Editor Manual](Engine).
+For more information on the topics covered in this quick start guide, and across the entire editor, see the [**Unreal Editor Manual**](Engine).
 
 As for specifics covered in this quick start:
 
-* For more information on the Level Editor, see: [Level Editor](Engine\UI\LevelEditor)
+* For more information on the Level Editor, see: [**Level Editor**](Engine\UI\LevelEditor)
 
-* For more information on Viewports, see: [Viewports](Engine\UI\LevelEditor\Viewports)
+* For more information on Viewports, see: [**Viewports**](Engine\UI\LevelEditor\Viewports)
 
-* For more information on the Editing Modes available in Unreal Editor, 4 see: [Level Editor Modes](Engine\UI\LevelEditor\Modes)
+* For more information on the Editing Modes available in Unreal Editor, 4 see: [**Level Editor Modes**](Engine\UI\LevelEditor\Modes)
 
-* For more information on the **Content Browser**, see: [Content Browser](Engine\Content\Browser)
+* For more information on the **Content Browser**, see: [**Content Browser**](Engine\Content\Browser)
 
-* For more information on the **Details** Panel, see: [Details Panel](Engine\UI\LevelEditor\Details)
+* For more information on the **Details** Panel, see: [**Details Panel**](Engine\UI\LevelEditor\Details)
 
-* For more information on Building, see: [Lightmass](Engine\Rendering\LightingAndShadows\Lightmass)
+* For more information on Building, see: [**Lightmass**](Engine\Rendering\LightingAndShadows\Lightmass)
 
-* For more information on Lighting, see: [Lighting Quick Start Guide](Engine/Rendering/LightingAndShadows/QuickStart)
+* For more information on Lighting, see: [**Lighting Quick Start Guide**](Engine/Rendering/LightingAndShadows/QuickStart)
 
 %Steps%
 
