#UE4doc #docs - bolded links and lettering.

diff --git a/Engine/Documentation/Source/Engine/Rendering/LightingAndShadows/LightsAndPostProcess/LightingAndPostProcess.INT.udn b/Engine/Documentation/Source/Engine/Rendering/LightingAndShadows/LightsAndPostProcess/LightingAndPostProcess.INT.udn
--- a/Engine/Documentation/Source/Engine/Rendering/LightingAndShadows/LightsAndPostProcess/LightingAndPostProcess.INT.udn
+++ b/Engine/Documentation/Source/Engine/Rendering/LightingAndShadows/LightsAndPostProcess/LightingAndPostProcess.INT.udn
@@ -52,11 +52,11 @@ Description:How lights interact with the Post Process Chain.
 
 The post process pipeline takes the lighting in the scene and uses it in various ways, from simulating an organic eye adapting to changes in light or a camera's numerous refracting lenses. While the lighting of a scene may be the first step, the post process steps really drive many effects.
 
-Refer to the [Post Process Effects](Engine/Rendering/PostProcessEffects) page for more information on how to use post processing effects. 
+Refer to the [**Post Process Effects**](Engine/Rendering/PostProcessEffects) page for more information on how to use post processing effects. 
 
 ## Light Adaptation
 
-Used for simulating a camera or human eye's adapting to changing lighting conditions such as walking into a bright sunlit day from a dark room.
+**Auto Exposure** is used for simulating a camera or human eye's adapting to changing lighting conditions such as walking into a bright sunlit day from a dark room.
 
 [REGION:lightbox]
 [![](LO_Adapt_003.png)(w:302)](LO_Adapt_003.png) [![](LO_Adapt_001.png)(w:302)](LO_Adapt_001.png) [![](LO_Adapt_002.png)(w:302)](LO_Adapt_002.png) 
@@ -65,7 +65,7 @@ Used for simulating a camera or human eye's adapting to changing lighting condit
 ## Bloom
 
 [EXCERPT:LightingAndPostProcess_Bloom]
-Bloom simulates the effect of light halos. More information about bloom settings can be found in the [Bloom section of the Post Process Effects documentation](Engine/Rendering\PostProcessEffects\Bloom).
+**Bloom** simulates the effect of light halos. More information about bloom settings can be found in the [**Bloom section of the Post Process Effects documentation**](Engine/Rendering\PostProcessEffects\Bloom).
 [/EXCERPT:LightingAndPostProcess_Bloom]
 
 [REGION:lightbox]
@@ -74,7 +74,7 @@ Bloom simulates the effect of light halos. More information about bloom settings
 
 ## Lens Flare
 
-Lens Flare are those "Visible artifacts, usually in the shape of the lens iris, and are formed when light follows a pathway through the lens that contains one or more reflections from the lens surfaces." A flare can occur from any significantly bright surface or effect.
+**Lens Flare** are those "Visible artifacts, usually in the shape of the lens iris, and are formed when light follows a pathway through the lens that contains one or more reflections from the lens surfaces." A flare can occur from any significantly bright surface or effect.
 
 ![](LI_LFnBloom.png)
 
