#UE4doc #docs - adjusted bold lettering.

diff --git a/Engine/Documentation/Source/Engine/Rendering/Materials/Functions/Reference/Blends/Blends.INT.udn b/Engine/Documentation/Source/Engine/Rendering/Materials/Functions/Reference/Blends/Blends.INT.udn
--- a/Engine/Documentation/Source/Engine/Rendering/Materials/Functions/Reference/Blends/Blends.INT.udn
+++ b/Engine/Documentation/Source/Engine/Rendering/Materials/Functions/Reference/Blends/Blends.INT.udn
@@ -7,7 +7,7 @@ Description:Functions designed to blend one color with another, similar to blend
 
 
 
-A Blend is a type of function that performs mathematical calculations in the color information of a texture so that one texture can blend into another in a particular manner.
+A **Blend** is a type of function that performs mathematical calculations in the color information of a texture so that one texture can blend into another in a particular manner.
 
 Blends will always have a Base and a Blend input, both of which are Vector3. These will each take in a texture and the two will be blended in some way. The manner in which the blend takes place depends on the type of blend function used.
 
@@ -17,7 +17,7 @@ Below is a list of all of the blend material functions.
 
 ### Blend_ColorBurn
 
-Blend_ColorBurn blends the textures in such a way that the darker the Blend color the more of that color will be used in the final result. If the Blend color is white, no change takes place.
+**Blend_ColorBurn** blends the textures in such a way that the darker the Blend color the more of that color will be used in the final result. If the Blend color is white, no change takes place.
 
 
 | Item | Description |
@@ -41,7 +41,7 @@ Blend_ColorBurn blends the textures in such a way that the darker the Blend colo
 
 ### Blend_ColorDodge
 
-Blend_ColorDodge lightens the result by inverting the Base color and dividing it by the Blend color.
+**Blend_ColorDodge** lightens the result by inverting the Base color and dividing it by the Blend color.
 
 | Item | Description |
 | -------- | -------- |
@@ -62,7 +62,7 @@ Blend_ColorDodge lightens the result by inverting the Base color and dividing it
 
 ### Blend_Darken
 
-Blend_Darken works by choosing the darker value for each pixel of the Base and Blend colors. A white Blend color does not produce a change.
+**Blend_Darken** works by choosing the darker value for each pixel of the Base and Blend colors. A white Blend color does not produce a change.
 
 | Item | Description |
 | -------- | -------- |
@@ -83,7 +83,7 @@ Blend_Darken works by choosing the darker value for each pixel of the Base and B
 
 ### Blend_Difference
 
-Blend_Difference creates a sort of inversion-style effect by subtracting the Base from the Blend and then taking the absolute value of the result.
+**Blend_Difference** creates a sort of inversion-style effect by subtracting the Base from the Blend and then taking the absolute value of the result.
 
 | Item | Description |
 | -------- | -------- |
@@ -104,7 +104,7 @@ Blend_Difference creates a sort of inversion-style effect by subtracting the Bas
 
 ### Blend_Exclusion
 
-Blend_Exclusion halves both the Base and Blend textures, combines them, and then does a partial inversion on the result.
+**Blend_Exclusion** halves both the Base and Blend textures, combines them, and then does a partial inversion on the result.
 
 | Item | Description |
 | -------- | -------- |
@@ -125,7 +125,7 @@ Blend_Exclusion halves both the Base and Blend textures, combines them, and then
 
 ### Blend_HardLight
 
-Like a harsher version of Blend_Overlay, Blend_HardLight will either screen or multiply the Base and Blend together. The function does a comparison on the Blend color such that wherever the Blend is brighter than 50% gray, the Base and Blend will be combined via a Screen operation. If the Blend is darker than 50% gray, the Base will be multiplied by the Blend as in the Multiply function. The contrast of the final result is then boosted for a harsh output.
+Like a harsher version of Blend_Overlay, **Blend_HardLight** will either screen or multiply the Base and Blend together. The function does a comparison on the Blend color such that wherever the Blend is brighter than 50% gray, the Base and Blend will be combined via a Screen operation. If the Blend is darker than 50% gray, the Base will be multiplied by the Blend as in the Multiply function. The contrast of the final result is then boosted for a harsh output.
 
 | Item | Description |
 | -------- | -------- |
@@ -146,7 +146,7 @@ Like a harsher version of Blend_Overlay, Blend_HardLight will either screen or m
 
 ### Blend_Lighten
 
-Blend_Lighten compares each pixel of the Base and Blend colors and returns the brighter result.
+**Blend_Lighten** compares each pixel of the Base and Blend colors and returns the brighter result.
 
 | Item | Description |
 | -------- | -------- |
@@ -167,7 +167,7 @@ Blend_Lighten compares each pixel of the Base and Blend colors and returns the b
 
 ### Blend_LinearBurn
 
-Blend_LinearBurn adds Base color to the Blend color and then subtracts 1 from the result.
+**Blend_LinearBurn** adds Base color to the Blend color and then subtracts 1 from the result.
 
 | Item | Description |
 | -------- | -------- |
@@ -188,7 +188,7 @@ Blend_LinearBurn adds Base color to the Blend color and then subtracts 1 from th
 
 ### Blend_LinearDodge
 
-Blend_LinearDodge adds the Base color to the Blend color.
+**Blend_LinearDodge** adds the Base color to the Blend color.
 
 | Item | Description |
 | -------- | -------- |
@@ -209,7 +209,7 @@ Blend_LinearDodge adds the Base color to the Blend color.
 
 ### Blend_LinearLight
 
-Blend_LinearLight is a linear version of Blend_Overlay, providing harsher results. The function does a comparison on the Blend color such that wherever the Blend is brighter than 50% gray, the Base and Blend will be combined via a Screen operation. If the Blend is darker than 50% gray, the Base will be multiplied by the Blend as in the Multiply function.
+**Blend_LinearLight** is a linear version of Blend_Overlay, providing harsher results. The function does a comparison on the Blend color such that wherever the Blend is brighter than 50% gray, the Base and Blend will be combined via a Screen operation. If the Blend is darker than 50% gray, the Base will be multiplied by the Blend as in the Multiply function.
 
 | Item | Description |
 | -------- | -------- |
@@ -230,7 +230,7 @@ Blend_LinearLight is a linear version of Blend_Overlay, providing harsher result
 
 ### Blend_Overlay
 
-Blend_Overlay will either screen or multiply the Base and Blend together. The function does a comparison on the Blend color such that wherever the Blend is brighter than 50% gray, the Base and Blend will be combined via a Screen operation. If the Blend is darker than 50% gray, the Base will be multiplied by the Blend as in the Multiply function.
+**Blend_Overlay** will either screen or multiply the Base and Blend together. The function does a comparison on the Blend color such that wherever the Blend is brighter than 50% gray, the Base and Blend will be combined via a Screen operation. If the Blend is darker than 50% gray, the Base will be multiplied by the Blend as in the Multiply function.
 
 | Item | Description |
 | -------- | -------- |
@@ -251,7 +251,7 @@ Blend_Overlay will either screen or multiply the Base and Blend together. The fu
 
 ### Blend_PinLight
 
-Like Blend_Overlay, Blend_PinLight will either lighten or darken the Base and Blend together. The function does a comparison on the Blend color such that wherever the Blend is brighter than 50% gray, the Base and Blend will be combined via a Screen operation. If the Blend is darker than 50% gray, the Base will be multiplied by the Blend as in the Multiply function. The contrast is softened, making this a less harsh version of Overlay.
+Like Blend_Overlay, **Blend_PinLight** will either lighten or darken the Base and Blend together. The function does a comparison on the Blend color such that wherever the Blend is brighter than 50% gray, the Base and Blend will be combined via a Screen operation. If the Blend is darker than 50% gray, the Base will be multiplied by the Blend as in the Multiply function. The contrast is softened, making this a less harsh version of Overlay.
 
 | Item | Description |
 | -------- | -------- |
@@ -272,7 +272,7 @@ Like Blend_Overlay, Blend_PinLight will either lighten or darken the Base and Bl
 
 ### Blend_Screen
 
-Blend_Screen lightens the Base by the Blend color. It does this by doing a "one minus" on both colors, multiplying them together, and taking a one-minus of the result.
+**Blend_Screen** lightens the Base by the Blend color. It does this by doing a "one minus" on both colors, multiplying them together, and taking a one-minus of the result.
 
 | Item | Description |
 | -------- | -------- |
@@ -293,7 +293,7 @@ Blend_Screen lightens the Base by the Blend color. It does this by doing a "one
 
 ### Blend_SoftLight
 
-Blend_SoftLight is a softer version of Overlay. The function does a comparison on the Blend color such that wherever the Blend is brighter than 50% gray, the Base and Blend will be combined via a Screen operation. If the Blend is darker than 50% gray, the Base will be multiplied by the Blend as in the Multiply function. The contrast is softened, making this a less harsh version of Overlay.
+**Blend_SoftLight** is a softer version of Overlay. The function does a comparison on the Blend color such that wherever the Blend is brighter than 50% gray, the Base and Blend will be combined via a Screen operation. If the Blend is darker than 50% gray, the Base will be multiplied by the Blend as in the Multiply function. The contrast is softened, making this a less harsh version of Overlay.
 
 | Item | Description |
 | -------- | -------- |
