#UE4doc #docs - bolded links.

diff --git a/Engine/Documentation/Source/Engine/Rendering/Materials/Functions/Reference/ImageAdjustment/ImageAdjustment.INT.udn b/Engine/Documentation/Source/Engine/Rendering/Materials/Functions/Reference/ImageAdjustment/ImageAdjustment.INT.udn
--- a/Engine/Documentation/Source/Engine/Rendering/Materials/Functions/Reference/ImageAdjustment/ImageAdjustment.INT.udn
+++ b/Engine/Documentation/Source/Engine/Rendering/Materials/Functions/Reference/ImageAdjustment/ImageAdjustment.INT.udn
@@ -36,7 +36,7 @@ This function blends between 3 input colors based on a grayscale alpha, in the f
 
 ### CheapContrast
 
-This function boosts the contrast of an input by remapping the high end of the histogram to a lower value, and the low end of the histogram to a higher one. This is similar to applying a [Levels adjustment](http://www.cambridgeincolour.com/tutorials/levels.htm) in Photoshop, and pulling the black and white flags in a bit. The user may control the degree to which the contrast is boosted.
+This function boosts the contrast of an input by remapping the high end of the histogram to a lower value, and the low end of the histogram to a higher one. This is similar to applying a [**Levels adjustment**](http://www.cambridgeincolour.com/tutorials/levels.htm) in Photoshop, and pulling the black and white flags in a bit. The user may control the degree to which the contrast is boosted.
 
 | Item | Description |
 | -------- | -------- |
@@ -52,7 +52,7 @@ This function boosts the contrast of an input by remapping the high end of the h
 
 ### CheapContrast_RGB
 
-This function boosts the contrast of an input by remapping the high end of the histogram to a lower value, and the low end of the histogram to a higher one. This is similar to applying a [Levels adjustment](http://www.cambridgeincolour.com/tutorials/levels.htm) in Photoshop, and pulling the black and white flags in a bit. The user may control the degree to which the contrast is boosted.
+This function boosts the contrast of an input by remapping the high end of the histogram to a lower value, and the low end of the histogram to a higher one. This is similar to applying a [**Levels adjustment**](http://www.cambridgeincolour.com/tutorials/levels.htm) in Photoshop, and pulling the black and white flags in a bit. The user may control the degree to which the contrast is boosted.
 
 Unlike the regular *CheapContrast function*, this function can take in a Vector3 as the input, allowing it to perform contrast operations on a color image.
 
@@ -67,7 +67,7 @@ Unlike the regular *CheapContrast function*, this function can take in a Vector3
 
 ### SCurve
 
-The SCurve function boosts contrast of an image by interpolating the values of each channel values of an image along an S-curve. This is similar to applying a [Curves adjustment](http://www.cambridgeincolour.com/tutorials/photoshop-curves.htm) in Photoshop and setting the RGB curve to an S-curve or using the Increase Contrast (RGB) preset.
+The SCurve function boosts contrast of an image by interpolating the values of each channel values of an image along an S-curve. This is similar to applying a [**Curves adjustment**](http://www.cambridgeincolour.com/tutorials/photoshop-curves.htm) in Photoshop and setting the RGB curve to an S-curve or using the Increase Contrast (RGB) preset.
 
 | Item | Description |
 | -------- | -------- |
