UEDOC-1903

Images have been altered to better reflect the Style guide
Gif has been replaced with a YouTube video

diff --git a/Engine/Documentation/Source/Engine/Matinee/HowTo/MHT_3/MHT_3.INT.udn b/Engine/Documentation/Source/Engine/Matinee/HowTo/MHT_3/MHT_3.INT.udn
--- a/Engine/Documentation/Source/Engine/Matinee/HowTo/MHT_3/MHT_3.INT.udn
+++ b/Engine/Documentation/Source/Engine/Matinee/HowTo/MHT_3/MHT_3.INT.udn
@@ -36,19 +36,19 @@ In the following steps, we are going to create a short sequence that uses multip
 
 1.  Minimize the Matinee window, then from the **Modes** menu under the **Basic** tab, **Drag-and-drop** a **Camera** into the viewport.
 
-	![](MHT3_DragCamera.png)(w:580)
+	![](MHT_3_Drag_In_Camera.jpg)(w:580)
 
 	_Try to place the **Camera** somewhere to the left of the template text as seen in the image above._
 
 1.  Place two more **Cameras** in the viewport and rotate them inward to face the box as seen in the image below. 
 
-	![](MHT3_AllCamerasPlaced.png)(w:580)
+	![](MHT3_All_CamerasPlaced.jpg)(w:580)
 
 	_You can rotate a **Camera** by selecting it and pressing the **E** key, then using the widget to orientate it._
 		
 1.  From the **Content Browser** under the **Game/Props** folder, locate the **SM_Chair** and **Drag-and-drop** it into the viewport, onto the box. 
 
-	![](MHT3_AddChair.png)
+	![](MHT3_Drag_In_Chair.jpg)
 			
 1.  With the chair still selected, press **E** to enter **Rotation Mode** and rotate the chair around to **140.00** as seen below. 
 
@@ -82,7 +82,7 @@ In the following steps, we are going to create a short sequence that uses multip
 
 1. Next, we will increase the length of the Matinee by holding **Control** and dragging the red marker (in the yellow box below) to the **15.00** mark (yellow X). 
 	[REGION:fullwidth]
-	![](MHT3_IncreaseMatineeTime.png)
+	![](MHT3_Increase_MatineeTime.jpg)
 	[/REGION]
 
 1. Click on the first keyframe at **0.00** on the **Movement** track of Camera1 the camera should be selected in the viewport. 
@@ -213,6 +213,16 @@ In the following steps, we are going to create a short sequence that uses multip
 
 1. You should now have a Matinee sequence that starts when the level loads and cycles through three different cameras. 
 
-	![](MHT3_FinalResultGif.gif)(convert:false)
+	[OBJECT:EmbeddedVideo]
+	[PARAMLITERAL:width]
+	640
+	[/PARAMLITERAL]
+	[PARAMLITERAL:height]
+	360
+	[/PARAMLITERAL]
+	[PARAMLITERAL:videoid]
+	k3r1aNcbO8Y
+	[/PARAMLITERAL]
+	[/OBJECT]
 
 
diff --git a/Engine/Documentation/Source/Engine/UI/LevelEditor/Viewports/ViewportControls/ViewportControls.INT.udn b/Engine/Documentation/Source/Engine/UI/LevelEditor/Viewports/ViewportControls/ViewportControls.INT.udn
--- a/Engine/Documentation/Source/Engine/UI/LevelEditor/Viewports/ViewportControls/ViewportControls.INT.udn
+++ b/Engine/Documentation/Source/Engine/UI/LevelEditor/Viewports/ViewportControls/ViewportControls.INT.udn
@@ -92,7 +92,7 @@ buttons pressed. These are also the only controls that can be used to navigate t
 
 ### Magic Mouse and Magic Trackpad
 
-After enabling Secondary click  in **System Preference** -> **Mouse** for the Magic Mouse, or enabling Secondary click  with two fingers in **System Preference** -> **Trackpad** for the Magic Trackpad, the following controls are available.
+On a Mac, after enabling Secondary click  in **System Preference** -> **Mouse** for the Magic Mouse, or enabling Secondary click  with two fingers in **System Preference** -> **Trackpad** for the Magic Trackpad, the following controls are available.
 
 | Magic Mouse | Magic Trackpad | Action |
 | ------- | ------- | ------- |
