UEDOC-1869

Updating the images for v4.8

diff --git a/Engine/Documentation/Source/Engine/Animation/AnimHowTo/BlendSpace/BlendspaceHowTo.INT.udn b/Engine/Documentation/Source/Engine/Animation/AnimHowTo/BlendSpace/BlendspaceHowTo.INT.udn
--- a/Engine/Documentation/Source/Engine/Animation/AnimHowTo/BlendSpace/BlendspaceHowTo.INT.udn
+++ b/Engine/Documentation/Source/Engine/Animation/AnimHowTo/BlendSpace/BlendspaceHowTo.INT.udn
@@ -30,13 +30,13 @@ Once downloaded, you can add the Animation Starter Pack to your project via the
 
 1.	Under **Default Pawn Class**, click the dropdown menu and select the **Ue4ASP_Character**. 
 
-	![](Engine\Animation\AnimHowTo\AdditiveAnimations\Animations1.png)
+	![](BS_DefaultPawnClass.jpg)
 
 1.	**Compile** and **Save** then close the Blueprint. 
 
 1.	In the **Content Browser**, click the **Create** Button then under _Animation_, select **Blend Space**. 
 
-	![](BlendSpace1.png)(w:400)
+	![](BS_BlendSpace.jpg)
 
 1.	In the **Pick Skeleton** window, select the **HeroTPP_Skeleton** that is located in the **Game/AnimStaterPack/Character** folder. 
 
@@ -44,7 +44,7 @@ Once downloaded, you can add the Animation Starter Pack to your project via the
 
 	Give the _Blend Space_ a name such as **Walk_Jog_BS**. 
 
-	![](BlendSpace3.png)
+	![](BS_Walk_Jog_BS.jpg)
 
 1.	Open the **Walk_Jog_BS**, then in the center window under the **Parameters** section, enter the values shown. 
 
