UEDOC-1869

Images have been replaced with UE4 images

diff --git a/Engine/Documentation/Source/Engine/Actors/FogEffects/HeightFog/HeightFog.INT.udn b/Engine/Documentation/Source/Engine/Actors/FogEffects/HeightFog/HeightFog.INT.udn
--- a/Engine/Documentation/Source/Engine/Actors/FogEffects/HeightFog/HeightFog.INT.udn
+++ b/Engine/Documentation/Source/Engine/Actors/FogEffects/HeightFog/HeightFog.INT.udn
@@ -121,7 +121,7 @@ The !FogStartDistance can be used to artificially keep some defined area in fron
 The following images show this applied:
 
 [REGION:imagetable]
-| ![FogStart_0.png](FogStart_0.png)(w:300) | ![FogStart_1000.png](FogStart_1000.png)(w:300) | ![FogStart_1000_high.png](FogStart_1000_high.png)(w:300) |
+| ![](Density0.jpg)(w:300) | ![](Density1000.jpg)(w:300) | ![](Density1000HighFog.jpg)(w:300) |
 | --- | --- | --- |
 | Fog Start Distance = 0 | Fog Start Distance = 1000 | Fog Start Distance = 1000 with a high Fog Density value |
 [/REGION]
