UEDOC-1901

Replaced arrows to match the instructional guide.
Replaced gifs with Youtube embedded videos

diff --git a/Engine/Documentation/Source/Engine/Matinee/HowTo/MHT_2/MHT_2.INT.udn b/Engine/Documentation/Source/Engine/Matinee/HowTo/MHT_2/MHT_2.INT.udn
--- a/Engine/Documentation/Source/Engine/Matinee/HowTo/MHT_2/MHT_2.INT.udn
+++ b/Engine/Documentation/Source/Engine/Matinee/HowTo/MHT_2/MHT_2.INT.udn
@@ -20,7 +20,7 @@ In the following steps, we are going to create two lifts that move from one poin
 
 1.  Inside the **Content Browser** and the **Shapes** folder, drag a **Shape_Cube** into the level in front of the **PlayerStart**. 
 
-	![](MHT2_DragInCube.png)(w:720)
+	![](MHT2_Drag_inCube.jpg)(w:720)
 
 1.  In the **Details** panel for the **Shape_Cube**, under **Transform**, set the **Scale** for **X, Y,** and **Z** to **3.0, 3.0** and **0.25**. 
 
@@ -138,7 +138,17 @@ In the following steps, we are going to create two lifts that move from one poin
 
 1. If you click the **Play** button to play in the editor, you should see the lift now moving up/down repeatedly, pausing at the top and bottom positions for a second. 
 
-	![](MHT2_Lift1Gif.gif)(w:420 convert:false)
+	[OBJECT:EmbeddedVideo]
+	[PARAMLITERAL:width]
+	640
+	[/PARAMLITERAL]
+	[PARAMLITERAL:height]
+	360
+	[/PARAMLITERAL]
+	[PARAMLITERAL:videoid]
+	_nPqdhXIoAw
+	[/PARAMLITERAL]
+	[/OBJECT]
 
 
 [OBJECT:TopicList]
@@ -202,7 +212,7 @@ In the following steps, we are going to create two lifts that move from one poin
 
 1. **Left-Click** on the little white box and drag the handle-bar widget downward to create an arc. 
 
-	![](MHT2_CE_DragCurveDown1.png)(w:780)
+	![](MHT2_CEDragCurveDown1.jpg)(w:780)
 
 1. You should also see in the Level Viewport that the **Movement Path** has bent with the adjustment made in the last step. 
 
@@ -210,7 +220,7 @@ In the following steps, we are going to create two lifts that move from one poin
 	
 1. **Left-Click** on the **Keyframe** at the **0.00** mark, then click the **Break** option from the **Toolbar** (see 2 below) and move the little white box (see 1 below) downward to create an arch. 
 
-	![](MHT2_CE_AdjustFirstKey.png)(w:780)
+	![](MHT2_CEAdjustFirstKey.jpg)(w:780)
 
 	_Note: You can switch between the **Keyframes** at 0.00 and 2.00, and continue moving the widget to create a smoother arch._
 
@@ -244,7 +254,17 @@ In the following steps, we are going to create two lifts that move from one poin
 
 1. When you click the **Play** button from the Main Editor's Toolbar, you should have two lifts that move differently and one that uses a curved path. 
 
-	![](MHT2_BothLiftsGif.gif)(w:420 convert:false)
+	[OBJECT:EmbeddedVideo]
+	[PARAMLITERAL:width]
+	640
+	[/PARAMLITERAL]
+	[PARAMLITERAL:height]
+	360
+	[/PARAMLITERAL]
+	[PARAMLITERAL:videoid]
+	SrnTCME2RVU
+	[/PARAMLITERAL]
+	[/OBJECT]
 	
 
 ---
