UEDOC-2169

Gif has been converted into a YouTube video

diff --git a/Engine/Documentation/Source/Engine/Paper2D/HowTo/Animation/2DCharacterAnimation.INT.udn b/Engine/Documentation/Source/Engine/Paper2D/HowTo/Animation/2DCharacterAnimation.INT.udn
--- a/Engine/Documentation/Source/Engine/Paper2D/HowTo/Animation/2DCharacterAnimation.INT.udn
+++ b/Engine/Documentation/Source/Engine/Paper2D/HowTo/Animation/2DCharacterAnimation.INT.udn
@@ -152,7 +152,17 @@ Now that our function is complete, we will call it as part of our Update Animati
 
 1.	**Compile** and **Save** then play in the editor. 
 
-	![](AnimationExample.gif)(convert:false)
+	[OBJECT:EmbeddedVideo]
+	[PARAMLITERAL:width]
+	640
+	[/PARAMLITERAL]
+	[PARAMLITERAL:height]
+	360
+	[/PARAMLITERAL]
+	[PARAMLITERAL:videoid]
+	umA8BlEjgF8
+	[/PARAMLITERAL]
+	[/OBJECT]
 
 	You will now have a character that cycles between the two states based on your condition of whether or not the character is moving. 
 
