UEDOC-1898

All images have been updated and the gif is now a YouTube video

diff --git a/Engine/Documentation/Source/Engine/Matinee/HowTo/MHT_1/MHT_1.INT.udn b/Engine/Documentation/Source/Engine/Matinee/HowTo/MHT_1/MHT_1.INT.udn
--- a/Engine/Documentation/Source/Engine/Matinee/HowTo/MHT_1/MHT_1.INT.udn
+++ b/Engine/Documentation/Source/Engine/Matinee/HowTo/MHT_1/MHT_1.INT.udn
@@ -45,17 +45,17 @@ In the following steps, we are going to create a simple door that opens/closes w
 	
 1. **Left-Click** on the **SM_DoorFrame** Static Mesh and drag-and-drop it into the **Level Viewport**. 
 	
-	![](MHT1_DragFrame.png)(w:640)
+	![](MHT_DoorFrameDragIn.jpg)(w:640)
 
 1. Also in the Props folder, locate the **SM_Door** Static Mesh and drag-and-drop it into the Level Viewport. 
 
-	![](MHT1_DragDoor.png)(w:640)
+	![](MHT_DoorDragIn.jpg)(w:640)
 
 1. In the Content Browser, find the SM_Door Static Mesh again and **Double-Click** on it to open it up in the **Static Mesh** viewer. 
 
 1. Inside the **Details** panel, in the search field type in **Collision** and set the **Collision Complexity** setting to **Use Complex Collision as Simple**. 
 
-	![](MHT1_DoorCollision.png)
+	![](MHT_CollisionComplexity.jpg)
 
 	_Note: The step above will ensure that the door has collision applied to it. _
 
@@ -71,7 +71,7 @@ In the following steps, we are going to create a simple door that opens/closes w
 	
 1. Inside the **Modes** menu and under the **Basic** tab, find the **Box Trigger** and drag-and-drop it into the Level Viewport inside the door frame.
 
-	![](MHT1_DragTrigger.png)(w:640)
+	![](MHT_DragInBoxTrigger.jpg)(w:640)
 
 1. Select the **TriggerBox** and in the Details panel under the **Shape** section, locate and change the **X** value for **Box Extent** to **140**. 
 
@@ -96,7 +96,7 @@ In the following steps, we are going to create a simple door that opens/closes w
 
 1. From the **Toolbar** menu, select **Matinee** and click the **Add Matinee** button. 
 
-	![](MHT1_AddMatinee.png)(w:760)
+	![](MHT_AddMatinee.jpg)(w:760)
 
 1. The Matinee undo warning message will appear, click the **Continue** button. 
 
@@ -166,7 +166,7 @@ In the following steps, we are going to create a simple door that opens/closes w
 
 1. From the **Toolbar** menu, select **Blueprints** and click the **Open Level Blueprint** button. 
 
-	![](MHT1_LevelBluerpint.png)(w:720)
+	![](MHT_OpenLvlBlueprint.jpg)(w:720)
 
 1. Minimize the **Level Blueprint** and through the **Level Viewport** or **World Outliner**, **Left-Click** on the **MatineeActor** to select it. 
 
@@ -174,7 +174,7 @@ In the following steps, we are going to create a simple door that opens/closes w
 
 1. Open the **Level Blueprint** window and **Right-Click** in the **Event Graph**, then in the context menu, select the **Add Reference to MatineeActor** option. 
 
-	![](MHT1_AddMatineeToBP.png)
+	![](MHT_MatineeReference.jpg)
 
 1. This will add the **MatineeActor** to the **Event Graph** of the **Level Blueprint**. 
 
@@ -216,7 +216,17 @@ In the following steps, we are going to create a simple door that opens/closes w
 
 1. Minimize the Level Blueprint and **click** the ![](MHT1_PlayButton.png) **Play** Button to play in the Editor. 
 
-	![](MHT1_SimpleDoor.gif)(w:480 convert:false)
+	[OBJECT:EmbeddedVideo]
+	[PARAMLITERAL:width]
+	640
+	[/PARAMLITERAL]
+	[PARAMLITERAL:height]
+	360
+	[/PARAMLITERAL]
+	[PARAMLITERAL:videoid]
+	Z2son2Wcgbo
+	[/PARAMLITERAL]
+	[/OBJECT]
 
 ---
 
