#UE4doc #docs Added SkillLevel and Version tags.

diff --git a/Engine/Documentation/Source/GettingStarted/DownloadingUnrealEngine/DownloadZip/DownloadZip.INT.udn b/Engine/Documentation/Source/GettingStarted/DownloadingUnrealEngine/DownloadZip/DownloadZip.INT.udn
--- a/Engine/Documentation/Source/GettingStarted/DownloadingUnrealEngine/DownloadZip/DownloadZip.INT.udn
+++ b/Engine/Documentation/Source/GettingStarted/DownloadingUnrealEngine/DownloadZip/DownloadZip.INT.udn
@@ -2,6 +2,8 @@ Availability:NoPublish
 Title:Downloading Unreal as a Zip from GitHub
 Crumbs:%ROOT%, GettingStarted, GettingStarted/DownloadingUnrealEngine
 Description:Downloading the Unreal Engine source code as a Zip file.
+Version: 4.8
+SkillLevel: Advanced
 
 If you prefer to download our source code as a Zip file from GitHub, this section is for you.
 
diff --git a/Engine/Documentation/Source/GettingStarted/DownloadingUnrealEngine/DownloadingUnrealEngine4.INT.udn b/Engine/Documentation/Source/GettingStarted/DownloadingUnrealEngine/DownloadingUnrealEngine4.INT.udn
--- a/Engine/Documentation/Source/GettingStarted/DownloadingUnrealEngine/DownloadingUnrealEngine4.INT.udn
+++ b/Engine/Documentation/Source/GettingStarted/DownloadingUnrealEngine/DownloadingUnrealEngine4.INT.udn
@@ -6,6 +6,7 @@ Prereq: GettingStarted/RecommendedSpecifications
 Related: Programming/Development/BuildingUnrealEngine
 Related: Support/Builds
 Version: 4.8
+SkillLevel: Advanced
 
 [TOC (start:2 end:3)]
 
diff --git a/Engine/Documentation/Source/GettingStarted/DownloadingUnrealEngine/GitHub/GitHub.INT.udn b/Engine/Documentation/Source/GettingStarted/DownloadingUnrealEngine/GitHub/GitHub.INT.udn
--- a/Engine/Documentation/Source/GettingStarted/DownloadingUnrealEngine/GitHub/GitHub.INT.udn
+++ b/Engine/Documentation/Source/GettingStarted/DownloadingUnrealEngine/GitHub/GitHub.INT.udn
@@ -2,6 +2,8 @@
 Title:Forking and Cloning Unreal from GitHub
 Crumbs:%ROOT%, GettingStarted, GettingStarted/DownloadingUnrealEngine
 Description:Forking and cloning the Unreal Engine source code from GitHub.
+Version: 4.8
+SkillLevel: Advanced
 
 Forking and cloning our source code from GitHub is a great way to get started with Unreal.
 
diff --git a/Engine/Documentation/Source/GettingStarted/Installation/InstallingUnrealEngine.INT.udn b/Engine/Documentation/Source/GettingStarted/Installation/InstallingUnrealEngine.INT.udn
--- a/Engine/Documentation/Source/GettingStarted/Installation/InstallingUnrealEngine.INT.udn
+++ b/Engine/Documentation/Source/GettingStarted/Installation/InstallingUnrealEngine.INT.udn
@@ -9,6 +9,7 @@ Related: Programming/QuickStart
 Related: Engine/QuickStart
 Related: Support/Builds
 Version: 4.8
+SkillLevel: Beginner
 
 [REGION:warning]
 We currently don't provide a binary installer of Unreal Engine for Linux.  
diff --git a/Engine/Documentation/Source/Programming/Development/BuildingUnrealEngine/BuildingUnrealEngine.INT.udn b/Engine/Documentation/Source/Programming/Development/BuildingUnrealEngine/BuildingUnrealEngine.INT.udn
--- a/Engine/Documentation/Source/Programming/Development/BuildingUnrealEngine/BuildingUnrealEngine.INT.udn
+++ b/Engine/Documentation/Source/Programming/Development/BuildingUnrealEngine/BuildingUnrealEngine.INT.udn
@@ -11,6 +11,7 @@ Related: Programming/QuickStart
 Related: Engine/QuickStart
 Related: Support/Builds
 Version: 4.8
+SkillLevel: Advanced
 
 %Globals:osselect_linux%
 
