UEDOC-1870

Images have been updated and gif has been replaced with a YouTube video

diff --git a/Engine/Documentation/Source/Engine/Animation/AnimHowTo/LayerEditing/LayerEditingHowTo.INT.udn b/Engine/Documentation/Source/Engine/Animation/AnimHowTo/LayerEditing/LayerEditingHowTo.INT.udn
--- a/Engine/Documentation/Source/Engine/Animation/AnimHowTo/LayerEditing/LayerEditingHowTo.INT.udn
+++ b/Engine/Documentation/Source/Engine/Animation/AnimHowTo/LayerEditing/LayerEditingHowTo.INT.udn
@@ -25,7 +25,7 @@ In this section, we will create a copy of the Idle animation that we can then mo
 
 1.  Inside the **Content Browser**, under the **Game/FirstPersonBP/Animations** folder, open the **Prone_Idle** animation. 
 
-	![](AnimEdit1.png)
+	![](LE_FPIdleM.jpg)
 
 1.	With the animation open, use the **Playback Controls** to pause the animation and set it to frame 0 so it is at the beginning. 
 
@@ -33,7 +33,7 @@ In this section, we will create a copy of the Idle animation that we can then mo
 
 1.	Click the **Create Asset** button from the Toolbar Menu, then under _Create Animation_, select **From Current Animation**. 
 
-	![](AnimEdit3.png)
+	![](LE_FromCurrentAnimation.jpg)
 
 	Here we are creating a new animation based off the idle animation. There are also options for creating a new animation based off a specific pose within an animation or from the base reference pose. For our purposes, we want the entire idle animation from which to work with when modifying it to perform a reload motion, therefore **From Current Animation** is the best choice. 
 
@@ -53,7 +53,7 @@ With the **FirstPerson_Reload** animation created, we can begin editing it throu
 
 1.	In the **Skeleton Tree**, click the **upperarm_l** Bone so that it is selected. 
 
-	![](AnimEdit6.png)
+	![](LE_Leftarm.jpg)
 
 	In order to start editing, you must first select the Bone that will be modified in the **Additive Layer Tracks** which we have done above. 
 
@@ -77,11 +77,11 @@ With the **FirstPerson_Reload** animation created, we can begin editing it throu
 
 1.	Inside the viewport, **Left-click** on the yellow arch of the **Rotation** widget. 
 
-	![](AnimEdit10.png)
+	![](LE_LeftArmRot.jpg)
 
 	It will change to an Axis Rotation widget, move the mouse to the left to the **-60.00** mark. 
 
-	![](AnimEdit11.png)
+	![](LE_LeftArmMinus60.jpg)
 
 	From a first person perspective, this will pull the hand down off screen. 
 
@@ -95,7 +95,7 @@ With the **FirstPerson_Reload** animation created, we can begin editing it throu
 
 1.	In the viewport, **Left-click** on and drag the yellow arch of the **Rotation** widget, then rotate to the right to the **60.00** mark. 
 
-	![](AnimEdit14.png)
+	![](LE_LeftArmUp60.jpg)
 
 1.	From the Toolbar Menu, click the **Add Key** button (or press the **S** key). 
 
@@ -111,7 +111,7 @@ With the **FirstPerson_Reload** animation created, we can begin editing it throu
 
 1.	In the viewport, press **W** to switch to **Translation** mode, then **Left-click** and drag the yellow arrow to the right slightly. 
 
-	![](AnimEdit19.png)
+	![](LE_LeftArmPump.jpg)
 
 	This will shift the arm backwards as if "pumping" the shotgun. You can use your discretion to determine how much to move it.  
 
@@ -141,7 +141,17 @@ With the **FirstPerson_Reload** animation created, we can begin editing it throu
  
 Below is an example of the animation playing on a character in-game. 
 
-![](ReloadGIF.gif)(convert:false)
+[OBJECT:EmbeddedVideo]
+[PARAMLITERAL:width]
+640
+[/PARAMLITERAL]
+[PARAMLITERAL:height]
+360
+[/PARAMLITERAL]
+[PARAMLITERAL:videoid]
+Soior02AUtE
+[/PARAMLITERAL]
+[/OBJECT]
 
 ## Conclusion & FAQ
 
@@ -167,13 +177,13 @@ You can continue to modify this animation by moving to a frame on the timeline a
 
 *	Inside the **Show** menu options, you can enable the **Source Animation** and (or) **Baked Animation** for display in the viewport. The Source Animation can be seen as a white skeleton, while the Baked Animation can be seen as a blue skeleton as shown in the image below. 
 
-	![](AnimEdit26.png)
+	![](LE_SourceBakedAnim.jpg)
 
 **Can Animation Layer Editing be used to create root motion?**
 
 *	Yes. Select your Root Bone and key it, then scrub the timeline forward and move the Root Bone and re-key it. 
 
-	![](AnimEdit27.png)
+	![](LE_ThirdPersonAnim.jpg)
 
 	Above we have enabled the Source and Baked Animation preview where the white skeleton is the source animation. 
 
