UEDOC-1872

Images of the Bluecharacter have been replaced

diff --git a/Engine/Documentation/Source/Engine/Animation/RetargetingDifferentSkeletons/RetargetingDifferentSkeletons.INT.udn b/Engine/Documentation/Source/Engine/Animation/RetargetingDifferentSkeletons/RetargetingDifferentSkeletons.INT.udn
--- a/Engine/Documentation/Source/Engine/Animation/RetargetingDifferentSkeletons/RetargetingDifferentSkeletons.INT.udn
+++ b/Engine/Documentation/Source/Engine/Animation/RetargetingDifferentSkeletons/RetargetingDifferentSkeletons.INT.udn
@@ -5,7 +5,7 @@ Description:Retargeting animations for use between different skeletons.
 Version: 4.5
 
 [REGION:fullwidth]
-![](Retarget1.png)
+![](Retarget1_1.jpg)
 [/REGION]
 
 [TOC(start:2 end:2)]
@@ -37,7 +37,7 @@ In this example, we are using the following assets which are available for downl
 
 Once you have installed the packs above, you can add them to a project via the **Add to project** button in the Unreal Launcher. 
 
-![](Retarget4.png)
+![](AnimationStarterpack2.jpg)
 
 The Animation Starter Pack includes several animations that we will retarget for a character from the Mixamo Character Pack.
 
@@ -51,7 +51,7 @@ The steps below will show you how to access the default Rig and retarget animati
 1.	Click the **Retarget Manager** button, then under **Set up Rig**, click the **Pick Asset** icon and select the **Humanoid** Rig. 
 
 	[REGION:lightbox]
-	[![](RetargetAdd2.png)(w:720)](RetargetAdd2.png)
+	[![](RT_Manager.jpg)(w:720)](RT_Manager.jpg)
 	[/REGION]
 
 	_Click image for full view._
@@ -75,7 +75,7 @@ The steps below will show you how to access the default Rig and retarget animati
 1.	Use the drop-down menus to select the Bones on the Skeleton that are similar in location to the Bones (Nodes) on the Rig. 
 
 	[REGION:lightbox]
-	[![](Retarget8.png)(w:800)](Retarget8.png)
+	[![](RT_Retargeting.jpg)](RT_Retargeting.jpg)
 	[/REGION]
 
 	_Click image for full view._
@@ -175,11 +175,11 @@ The steps below will show you how to access the default Rig and retarget animati
 
 1.	**Right-click** on the animation and choose **Retarget Animation Assets -> Duplicate Anim Assets and Retarget**. 
 
-	![](Retarget9.png)
+	![](RT_DuplicateAnim.jpg)
 
 1.	In the **Select Skeleton** window, you should see your Skeleton in the upper portion of the window, select it and click **Select**. 
 
-	![](Retarget11.png)(w:820)
+	![](Retarget11_2.jpg)(w:820)
 
 1.	A copy of the animation you selected to retarget will be created in the same folder of your Target Skeleton. 
 
@@ -215,11 +215,11 @@ Once you have gone through the process of setting up a Rig for two Skeleton asse
 
 **Right-click** on an Animation Blueprint and select the **Retarget Anim Blueprints -> Duplicate Anim Blueprints and Retarget** option.
 
-![](Retarget16.png)
+![](RT_AnimBP.jpg)(w:820)
 
 The **Select Skeleton** window will appear. Select your Target Skeleton and press **Select**. 
 
-![](Retarget11.png)(w:820)
+![](Retarget11_2.jpg)(w:820)
 
 All the assets associated with the Animation Blueprint (Animation assets, Blend Spaces, etc.), along with any logic defined in the Animation Blueprint itself will be copied and placed in the same location as your Target Skeleton. 
 
@@ -232,7 +232,7 @@ There may be instances where you want to retarget animations to a Skeleton that
 
 Take for example the Skeleton in the image below on the right. 
 
-![](Retarget18.png)(w:820)
+![](Retarget2_2.jpg)(w:820)
 
 Our Source Skeleton on the left is in a T-Pose while our Target Skeleton is using an A-Pose. Even though a Rig has been set up for them to share, if we were to retarget animations as is now, we may encounter issues as seen below. 
 
@@ -250,7 +250,7 @@ If you click the **View Pose** button, you will see the current Base Pose. You c
 
 In the image above, the right arm has been adjusted while the left arm is using the default base pose (which we can adjust by using the rotation widget). Once you have adjusted the pose, you can click the **Save Pose** option to save the retargeted Base Pose. With similar Base Poses between Skeletons, you can now retarget an animation and when the Skeleton Select window appears, you will notice that your Target Skeleton is now using a retargeted Base Pose. 
 
-![](Retarget22.png)(w:820)
+![](Retarget2_2.jpg)(w:820)
 
 Retargeting the animation with the retargeted Base Pose, results in a better looking animation. 
 
