#UE4doc #docs - adjusted bolding and wording changes where needed.

diff --git a/Engine/Documentation/Source/Engine/AI/BehaviorTrees/QuickStart/1/1.INT.udn b/Engine/Documentation/Source/Engine/AI/BehaviorTrees/QuickStart/1/1.INT.udn
--- a/Engine/Documentation/Source/Engine/AI/BehaviorTrees/QuickStart/1/1.INT.udn
+++ b/Engine/Documentation/Source/Engine/AI/BehaviorTrees/QuickStart/1/1.INT.udn
@@ -36,7 +36,7 @@ This Quick Start makes some assumptions about what is available in the project,
 	The **Starter Assets** are optional for this quick start guide.
 
 
-If you need further explanation of starting a new project, see the [Project Browser](Engine/Basics/Projects/Browser) User Guide.
+If you need further explanation of starting a new project, see the [**Project Browser**](Engine/Basics/Projects/Browser) User Guide.
 
 [/EXCERPT:AIProjectStart]
 %Steps%
diff --git a/Engine/Documentation/Source/Engine/AI/BehaviorTrees/QuickStart/2/2.INT.udn b/Engine/Documentation/Source/Engine/AI/BehaviorTrees/QuickStart/2/2.INT.udn
--- a/Engine/Documentation/Source/Engine/AI/BehaviorTrees/QuickStart/2/2.INT.udn
+++ b/Engine/Documentation/Source/Engine/AI/BehaviorTrees/QuickStart/2/2.INT.udn
@@ -36,7 +36,7 @@ The NavMesh in Unreal has a number of functions that work well with Behavior Tre
 
 	![](createNavVol.png)
 
-1. Translate and scale the volume so it encapsulates the areas you want your characters to be able to navigate.
+1. **Translate** and **Scale** the volume so it encapsulates the areas you want your characters to be able to navigate.
 
 	* You can use more than one Volume if you like, or one large one that encapsulates the entire level.
 
diff --git a/Engine/Documentation/Source/Engine/AI/BehaviorTrees/QuickStart/3/3.INT.udn b/Engine/Documentation/Source/Engine/AI/BehaviorTrees/QuickStart/3/3.INT.udn
--- a/Engine/Documentation/Source/Engine/AI/BehaviorTrees/QuickStart/3/3.INT.udn
+++ b/Engine/Documentation/Source/Engine/AI/BehaviorTrees/QuickStart/3/3.INT.udn
@@ -44,7 +44,7 @@ To start, you will need create 4 assets in the **Content Browser**: An AI Contro
 
 	Name this **Blueprint** `AI_Character`.
 
-1. **New** -> **Miscellaneous** -> **Behavior Tree**
+1. **New** -> **Artificial Intelligence** -> **Behavior Tree**
 
 	[REGION:raw]
 	![](newBT.png)
@@ -52,7 +52,7 @@ To start, you will need create 4 assets in the **Content Browser**: An AI Contro
 
 	Name the **Behavior Tree** `FollowerBT`.
 
-1. **New** -> **Miscellaneous** ->  **Blackboard**.
+1. **New** -> **Artificial Intelligence** ->  **Blackboard**.
 	
 	[REGION:raw]
 	![](newBB.png)
diff --git a/Engine/Documentation/Source/Engine/AI/BehaviorTrees/QuickStart/4/4.INT.udn b/Engine/Documentation/Source/Engine/AI/BehaviorTrees/QuickStart/4/4.INT.udn
--- a/Engine/Documentation/Source/Engine/AI/BehaviorTrees/QuickStart/4/4.INT.udn
+++ b/Engine/Documentation/Source/Engine/AI/BehaviorTrees/QuickStart/4/4.INT.udn
@@ -33,7 +33,7 @@ This Character will become our new AI Character. The reason we use a Character,
 
 	![](meshComp.png)
 
-1. In the **Details** panel, find the **Mesh** category and assign the `HeroTPP` (or your own) **Skeletal Mesh**.
+1. In the **Details** panel, find the **Mesh** category and assign the `SK_Mannequin` (or your own) **Skeletal Mesh**.
 
 	![](assignSkelMesh.png)
 
@@ -41,7 +41,7 @@ This Character will become our new AI Character. The reason we use a Character,
 
 1. Rotate the **Mesh** Component so it faces the direction the blue Arrow Component is pointing.
 
-1. In the **Details** panel, find the **Animation** category and assign the `HeroTPP_AnimBlueprint` to the **Animation Blueprint Generated Class** property.
+1. In the **Details** panel, find the **Animation** category and assign the `TopDown_AnimBP` to the **Animation Blueprint Generated Class** property.
 
 	![](assignAnimBP.png)
 
@@ -53,7 +53,7 @@ This Character will become our new AI Character. The reason we use a Character,
 
 	![](charmve.png)
 
-1. In the **Details** panel, in the Movement Component category, adjust the Agent Radius and Agent Height values.
+1. In the **Details** panel, in the Movement Component category, adjust the _Agent Radius_ and _Agent Height_ values.
 
 	![](agentProps.png)
 
