UEDOC-1460

Origin is now defined as 0,0,0 with no rotation.

This page is fixed internally and will show correctly after the next publishing.

diff --git a/Engine/Documentation/Source/Engine/Animation/RootMotion/RootMotion.INT.udn b/Engine/Documentation/Source/Engine/Animation/RootMotion/RootMotion.INT.udn
--- a/Engine/Documentation/Source/Engine/Animation/RootMotion/RootMotion.INT.udn
+++ b/Engine/Documentation/Source/Engine/Animation/RootMotion/RootMotion.INT.udn
@@ -43,7 +43,7 @@ Another example of **Root Motion** applied to a character can be found on the [A
 Put simply, root motion is the motion of a character that is based off animation from the root bone of the skeleton. Most in-game animations are handled via cycles in which the character's root remains stationary. This cannot always be the case, however, as we can see in the example above. To handle this, we need to take the motion of the root away from the character and apply it to the character's capsule instead. That is the essence of what Root Motion does in UE4.
 
 [REGION:warning]
-In order to properly use Root Motion, it is important to note that the Root Bone of your character should be at the origin (on the floor) as this allows the system to isolate physical movement (capsule) from animated movement (character). 
+In order to properly use Root Motion, it is important to note that the Root Bone of your character should be at the origin (0,0,0 with no rotation) as this allows the system to isolate physical movement (capsule) from animated movement (character). 
 [/REGION]
 
 
