#UE4 DOCS: added all the information to the stats ref page

diff --git a/Engine/Documentation/Source/Engine/Performance/StatCommands/StatCommands.INT.udn b/Engine/Documentation/Source/Engine/Performance/StatCommands/StatCommands.INT.udn
--- a/Engine/Documentation/Source/Engine/Performance/StatCommands/StatCommands.INT.udn
+++ b/Engine/Documentation/Source/Engine/Performance/StatCommands/StatCommands.INT.udn
@@ -32,7 +32,7 @@ This is necessary to get anything useful out of any of the "dump" commands.
 | **stat Memory** || Shows stats on how much memory is being used by various sub systems in Unreal Engine. |
 | **stat Streaming** || Displays various statistics on streaming assets, like how much memory streaming textures are using, or how many streaming textures there are in the scene. |
 | **stat StreamingDetails** || More detailed statistics on streaming, like breaking down general texture streaming into more specific groups (lightmaps, static textures, and dynamic textures). |
-| **stat Slow [-ms=0.3] [-maxdepth=5]** || Show any stat that has recently taken longer than 0.3 ms up to 5 levels deep (as things can be layered). The defaults are **-ms=1.0 -maxdepth=4**. Be careful with really small values in **-ms=** as it will start displaying everything and will drastically reduce performance (thus making more things show up on the list...). |
+| **stat Slow [-ms=0.3] [-maxdepth=5]** || [INCLUDE:#statslow] |
 | **stat Grouped** || Disables stat Slow.|
 [/REGION]
 
@@ -44,10 +44,36 @@ This is necessary to get anything useful out of any of the "dump" commands.
 Note: removed until this can be checked against profiler info. StatsViewer didn't appear to be in P4 or a binary build
 | **stat StartFile** **stat StopFile** | Generates a .ue4stats file that can be analyzed offline (With the StatsViewer tool from: UE4\Engine\Binaries\DotNET\StatsViewer). A GPU stall can appear as high workload the frame before.
 
+[EXCERPT:statSlow]
+This command will display stats for the game thread and the render thread. All stats will rendered as one big stats group. This command can be useful, if we don’t have access to the profiler or the log file, or we want to test the basic performance of our game.
 
+Here is the example in the running game:
 
--->
+[REGION:lightbox]
+[![](statSlow1.png)(w:600)](statSlow1.png)
+[/REGION]
+
+This command can be configured as follows:
+
+	stat slow [-ms=1.0] [-maxdepth=4]
+
+By default only stat items larger than 1.0 ms and stat items that are not nested more that 4 levels are displayed.
+
+Be careful as lowering the ms or increasing the depth may affect the overall performance.
+
+**Stat slow** is filtering anything lower than 0.1ms, to improve the performance and to avoid clutter. Here is an example of what you will get if you decrease your -ms and increase your -maxdepth:
 
+	stat slow -ms=0.3 -maxdepth=5
+
+[REGION:lightbox]
+[![](statSlow2.png)(w:600)](statSlow2.png)
+[/REGION]
+
+**Stat slow** can eat a lot of screen space, using **stat slow** again (or **stat Grouped**) will remove all visible stats groups. Further, if you enable other stats groups, the **stat slow** will be disabled as well. 
+
+[/EXCERPT]
+
+-->
 
 
 
