Saving before checking in this time

diff --git a/Engine/Documentation/Source/Engine/UI/LevelEditor/Modes/MeshPaintMode/VertexColor/MaterialSetup/Color/Color.INT.udn b/Engine/Documentation/Source/Engine/UI/LevelEditor/Modes/MeshPaintMode/VertexColor/MaterialSetup/Color/Color.INT.udn
--- a/Engine/Documentation/Source/Engine/UI/LevelEditor/Modes/MeshPaintMode/VertexColor/MaterialSetup/Color/Color.INT.udn
+++ b/Engine/Documentation/Source/Engine/UI/LevelEditor/Modes/MeshPaintMode/VertexColor/MaterialSetup/Color/Color.INT.udn
@@ -14,8 +14,9 @@ In order to make this paintable, we need to get the Vertex Color material expres
 
 **Click for full size.**
 
- [![](OriginalColorMaterial.png)(w:720)](OriginalColorMaterial.png)
-
+[REGION:lightbox]
+[![](OriginalColorMaterial.png)(w:720)](OriginalColorMaterial.png)
+[/REGION]
 
 [REGION:note]
 This material has been desaturated a bit, giving us more room to add color via the Mesh Paint tool.
@@ -25,9 +26,10 @@ This material has been desaturated a bit, giving us more room to add color via t
  To get our vertex colors worked into the material, we simply need to intercept the Base Color connection and multiply it by a Vertex Color expression's RGB value, like so:
 
  **Click for full size.**
-
+ 
+[REGION:lightbox]
 [![](VertexColorAdded.png)(w:720)](VertexColorAdded.png)
-
+[/REGION]
 
 Within the commented area, you can see where we have added a Multiply node and a corresponding Vertex Color node.
 
