#UE4doc #docs - updated images and bold lettering where needed.

diff --git a/Engine/Documentation/Source/Engine/AI/BehaviorTrees/UserGuide/BehaviorTreeUserGuide.INT.udn b/Engine/Documentation/Source/Engine/AI/BehaviorTrees/UserGuide/BehaviorTreeUserGuide.INT.udn
--- a/Engine/Documentation/Source/Engine/AI/BehaviorTrees/UserGuide/BehaviorTreeUserGuide.INT.udn
+++ b/Engine/Documentation/Source/Engine/AI/BehaviorTrees/UserGuide/BehaviorTreeUserGuide.INT.udn
@@ -38,20 +38,20 @@ Or, if you have started to edit a Behavior Tree (or simply want a new Blackboard
 
 ### Creating a Task (Blueprint)
 
-1. Click the New Task ![](Engine/AI/BehaviorTrees/QuickStart/12/newTask_button.png) button in the Behavior Tree Editor (while in Behavior Tree Edit Mode).
-1. Choose a Behavior Tree Task class to inherit from, `BTTask_BlueprintBase` is the default "blank" class.
+1. Click the ![](Engine/AI/BehaviorTrees/QuickStart/12/newTask_button.png) button in the Behavior Tree Editor (while in Behavior Tree Edit Mode).
+1. Choose a **Behavior Tree Task** class to inherit from, `BTTask_BlueprintBase` is the default "blank" class.
 1. The new Task Blueprint will be created in the same folder as your Behavior Tree asset.
 
 ### Creating a Decorator (Blueprint)
 
-1. Click the New Decorator ![](Engine/AI/BehaviorTrees/QuickStart/13/newDecorator_button.png) button in the Behavior Tree Editor (while in Behavior Tree Edit Mode).
-1. Choose a Behavior Tree Decorator class to inherit from, `BTDecorator_BlueprintBase` is the default "blank" class.
+1. Click the ![](Engine/AI/BehaviorTrees/QuickStart/13/newDecorator_button.png) button in the Behavior Tree Editor (while in Behavior Tree Edit Mode).
+1. Choose a **Behavior Tree Decorator** class to inherit from, `BTDecorator_BlueprintBase` is the default "blank" class.
 1. The new Decorator Blueprint will be created in the same folder as your Behavior Tree asset.
 
 ### Creating a Service (Blueprint)
 
-1. Click the New Service ![](Engine/AI/BehaviorTrees/QuickStart/11/newService_button.png) button in the Behavior Tree Editor (while in Behavior Tree Edit Mode).
-1. Choose a Behavior Tree Service class to inherit from, `BTService_BlueprintBase` is the default "blank" class.
+1. Click the ![](Engine/AI/BehaviorTrees/QuickStart/11/newService_button.png) button in the Behavior Tree Editor (while in Behavior Tree Edit Mode).
+1. Choose a **Behavior Tree Service** class to inherit from, `BTService_BlueprintBase` is the default "blank" class.
 1. The new Service Blueprint will be created in the same folder as your Behavior Tree asset.
 
 ## Editing a Blackboard
@@ -66,7 +66,7 @@ Or by switching a currently opened Behavior Tree Editor to Blackboard Edit Mode:
 
 ### Add Keys
 
-1. Click the New Key ![](Engine/AI/BehaviorTrees/QuickStart/5/newKey_button.png) button.
+1. Click the ![](Engine/AI/BehaviorTrees/QuickStart/5/newKey_button.png) button.
 1. Select the Type of Key to add.
 1. Adjust the name of the Key.
 
@@ -87,7 +87,7 @@ Object Keys are special in that you can change their specific type. This enables
 
 As some Decorator and Service functions require Actor Keys, this is how you change an Object Key into an Actor Key.
 
-1. Select an Object Key
+1. Select an **Object** Key
 1. In the **Details** Panel, expand the Key Type property.
 1. Change the Base Class property to the desired Object Type.
 
@@ -103,7 +103,7 @@ Or by switching a currently opened Behavior Tree Editor to Behavior Tree Edit Mo
 
 ### Setting the Blackboard Asset to Use
 
-1. Select the Root node of the Behavior Tree (or deselect all nodes)
+1. Select the **Root** node of the Behavior Tree (or deselect all nodes)
 1. In the **Details** Panel, set the Blackboard Asset to the desired Blackboard.
 
 [REGION:raw]
@@ -164,7 +164,7 @@ To connect nodes:
 
 ### Removing Decorators and Services from Nodes
 
-1. Select a Decorator or Service on a node by **left-clicking** it (the Decorator or Service will gain a redish outline).
+1. Select a **Decorator** or **Service** on a node by **left-clicking** it (the Decorator or Service will gain a redish outline).
 
 	![](selectDec.png)
 
@@ -172,11 +172,11 @@ To connect nodes:
 
 ### Editing Decorators and Services
 
-1. Select a Decorator or Service on a node by **left-clicking** it (the Decorator or Service will gain a redish outline).
+1. Select a **Decorator** or **Service** on a node by **left-clicking** it (the Decorator or Service will gain a redish outline).
 
 	![](selectDec.png)
 
-1. Adjust the Decorator or Service properties using the **Details** panel.
+1. Adjust the **Decorator** or **Service** properties using the **Details** panel.
 
 	![](details.png)
 
