#UE4doc #docs Minor edits.

diff --git a/Engine/Documentation/Source/Programming/Development/CompilingProjects/CompilingProjects.INT.udn b/Engine/Documentation/Source/Programming/Development/CompilingProjects/CompilingProjects.INT.udn
--- a/Engine/Documentation/Source/Programming/Development/CompilingProjects/CompilingProjects.INT.udn
+++ b/Engine/Documentation/Source/Programming/Development/CompilingProjects/CompilingProjects.INT.udn
@@ -21,42 +21,44 @@ project created from a Blueprints Only template.
 UE4 and the UnrealBuildTool use different build configurations to determine how the engine is compiled. Which configuration you use will be determined by the purposes of the build you want to create.
 
 [PUBLISH:Rocket]
+[EXCERPT:windowsBuildConfig]
 Each build configuration contains two keywords.  The first of these indicates the state of the engine and your game project. For instance, if you compile using a **Debug** configuration, you will
 be able to debug your game's code.  The second keyword indicates the target you are building for.  If you would like to open a project in Unreal, you need to build with
 the **Editor** target keyword, but if you are building an executable version of your game, you would need to build using the empty target keyword.
 
-|Build Configuration - State | Description |
-|-------------|-------------|
-| **Debug** | This configuration contains symbols for debugging. If you compile your project using the **Debug** configuration and want to open the project with the Unreal Editor, you must use the "-debug" flag in order to see your code changes reflected in your project.  |
-| **Development** | This configuration is equivalent to Release. In order to see your project's code changes reflected when you open the project later in the Unreal Editor, you must compile in the **Development** configuration.|
-| **Shipping** | This is the configuration for optimal performance and shipping your game.  This configuration strips out console commands, stats, and profiling tools. |
-
-|Build Configuration - Target | Description|
-| ------------------| ----------|
-|  [empty]      | This configuration builds a stand-alone executable version of your project, but requires cooked content specific to the platform.   |
-| **Editor**        | To be able to open a project in Unreal Editor and see all code changes reflected, the project must be built in an **Editor** configuration.  |  
-| **Uncooked**      | Projects built using configurations that have an **Uncooked** type should be opened in Unreal Editor with the "-game" flag.  This runs your game with uncooked content, in a new window, and is equivalent to **Play in > New Window at Default Player Start** in the editor.       | 
+|Build Configuration - State  | Description                                                                                                                                                                                                                                                        |
+|-------------                |-------------                                                                                                                                                                                                                                                       |
+| **Debug**                   | This configuration contains symbols for debugging. If you compile your project using the **Debug** configuration and want to open the project with the Unreal Editor, you must use the "-debug" flag in order to see your code changes reflected in your project.  |
+| **Development**             | This configuration is equivalent to Release. In order to see your project's code changes reflected when you open the project later in the Unreal Editor, you must compile in the **Development** configuration.                                                    |
+| **Shipping**                | This is the configuration for optimal performance and shipping your game.  This configuration strips out console commands, stats, and profiling tools.                                                                                                             |
+
+|Build Configuration - Target | Description                                                                                                                                                                                                                                                                         |
+| ------------------          | ----------                                                                                                                                                                                                                                                                          |
+|  [empty]                    | This configuration builds a stand-alone executable version of your project, but requires cooked content specific to the platform.                                                                                                                                                   |
+| **Editor**                  | To be able to open a project in Unreal Editor and see all code changes reflected, the project must be built in an **Editor** configuration.                                                                                                                                         |  
+| **Uncooked**                | Projects built using configurations that have an **Uncooked** type should be opened in Unreal Editor with the "-game" flag.  This runs your game with uncooked content, in a new window, and is equivalent to **Play in > New Window at Default Player Start** in the editor.       | 
+[/EXCERPT]
 [/PUBLISH:Rocket]
 [PUBLISH:Licensee]
 Each build configuration contains two keywords.  The first of these indicates the state of the engine and your game project. For instance, if you compile using a **DebugGame** configuration, you will
 be able to debug the game code, but not the engine code.  The second keyword indicates the target you are building for.  If you would like to open a project in UE4, you need to build with
 the **Editor** target keyword, but if you are building an executable version of your project that will run on a server, you would need to build using the **Server** target keyword.
 
-|Build Configuration - State | Description |
-|-------------|-------------|
-| **Debug** | This configuration builds both engine and game code in debug configuration.|
-| **DebugGame** | This configuration builds the engine as optimized, but leaves the game code debuggable. This configuration is ideal for debugging only game modules. |
-| **Development** | This configuration is equivalent to Release.  Both engine and game code will be built in this configuration. |
-| **Shipping** | This is the configuration for optimal performance and shipping your game.  This configuration strips out console commands, stats, and profiling tools. |
-| **Test** | This configuration is the **Shipping** configuration, but with some console commands, stats, and profiling tools enabled. |
-
-|Build Configuration - Target | Description|
-| ------------------| ----------|
-|  [empty]      | This configuration builds a stand-alone executable version of your project, but requires cooked content specific to the platform.   |
-| **Client**        | If there is a `<Game>Client.Target.cs` file, the **Client** build configurations will be valid.   |
-| **Editor**        | To be able to open a project in Unreal Editor and see all code changes reflected, the project must be built in an **Editor** configuration.  |  
-| **Server**        | If there is a `<Game>Server.Target.cs` file, the **Server** build configurations will be valid.  |
-| **Uncooked**      | Projects built using configurations that have an **Uncooked** type should be opened in Unreal Editor with the "-game" flag.  This runs your game with uncooked content, in a new window, and is equivalent to **Play in > New Window at Default Player Start** in the editor.       | 
+|Build Configuration - State  | Description                                                                                                                                            |
+|-------------                |-------------                                                                                                                                           |
+| **Debug**                   | This configuration builds both engine and game code in debug configuration.                                                                            |
+| **DebugGame**               | This configuration builds the engine as optimized, but leaves the game code debuggable. This configuration is ideal for debugging only game modules.   |
+| **Development**             | This configuration is equivalent to Release.  Both engine and game code will be built in this configuration.                                           |
+| **Shipping**                | This is the configuration for optimal performance and shipping your game.  This configuration strips out console commands, stats, and profiling tools. |
+| **Test**                    | This configuration is the **Shipping** configuration, but with some console commands, stats, and profiling tools enabled.                              |
+
+|Build Configuration - Target | Description                                                                                                                                                                                                                                                                         |
+| ------------------          | ----------                                                                                                                                                                                                                                                                          |
+|  [empty]                    | This configuration builds a stand-alone executable version of your project, but requires cooked content specific to the platform.                                                                                                                                                   |
+| **Client**                  | If there is a `<Game>Client.Target.cs` file, the **Client** build configurations will be valid.                                                                                                                                                                                     |
+| **Editor**                  | To be able to open a project in Unreal Editor and see all code changes reflected, the project must be built in an **Editor** configuration.                                                                                                                                         |  
+| **Server**                  | If there is a `<Game>Server.Target.cs` file, the **Server** build configurations will be valid.                                                                                                                                                                                     |
+| **Uncooked**                | Projects built using configurations that have an **Uncooked** type should be opened in Unreal Editor with the "-game" flag.  This runs your game with uncooked content, in a new window, and is equivalent to **Play in > New Window at Default Player Start** in the editor.       | 
 [/PUBLISH:Licensee]
 [/EXCERPT:Main]
 
