#UE4doc #docs Started 'Build Configurations' documentation.

diff --git a/Engine/Documentation/Source/Programming/Development/BuildConfigurations/BuildConfigurations.INT.udn b/Engine/Documentation/Source/Programming/Development/BuildConfigurations/BuildConfigurations.INT.udn
new file mode 100644
--- /dev/null
+++ b/Engine/Documentation/Source/Programming/Development/BuildConfigurations/BuildConfigurations.INT.udn
@@ -0,0 +1,57 @@
+Availability: Docs
+Title: Build Configurations Reference
+Description: Reference guide for the different build configurations available to developers.
+Crumbs: %ROOT%, Programming, Programming/Development
+Prereq: GettingStarted/RecommendedSpecifications
+Prereq: GettingStarted/DownloadingUnrealEngine
+Prereq: Programming/Development/BuildingUnrealEngine
+Related: Support/Builds
+SkillLevel: Advanced
+Version: 4.8
+
+Unreal Engine 4 (UE4) uses a custom building method via the UnrealBuildTool. This tool handles all the complex aspects of compiling the engine. 
+
+### Build Configuration
+
+%Globals:osselect_linux%
+
+[OBJECT:ToggleButtonContent]
+	[PARAMLITERAL:category]
+	OS
+	[/PARAMLITERAL]
+	[PARAMLITERAL:id]
+	windows
+	[/PARAMLITERAL]
+	[PARAMLITERAL:active]
+	active_button_content
+	[/PARAMLITERAL]
+	[PARAM:content]
+		Windows Test.
+	[/PARAM]
+[/OBJECT]
+[OBJECT:ToggleButtonContent]
+	[PARAMLITERAL:category]
+	OS
+	[/PARAMLITERAL]
+	[PARAMLITERAL:id]
+	mac
+	[/PARAMLITERAL]
+	[PARAMLITERAL:active]
+	[/PARAMLITERAL]
+	[PARAM:content]
+		Mac Test.
+	[/PARAM]
+[/OBJECT]
+[OBJECT:ToggleButtonContent]
+	[PARAMLITERAL:category]
+	OS
+	[/PARAMLITERAL]
+	[PARAMLITERAL:id]
+	linux
+	[/PARAMLITERAL]
+	[PARAMLITERAL:active]
+	[/PARAMLITERAL]
+	[PARAM:content]
+		Linux Test.
+	[/PARAM]
+[/OBJECT]
\ No newline at end of file
