#UE4doc #docs Added INCLUDE to 'Build Configuration' table from 'CompilingProjects' documentation.

diff --git a/Engine/Documentation/Source/Programming/Development/BuildConfigurations/BuildConfigurations.INT.udn b/Engine/Documentation/Source/Programming/Development/BuildConfigurations/BuildConfigurations.INT.udn
--- a/Engine/Documentation/Source/Programming/Development/BuildConfigurations/BuildConfigurations.INT.udn
+++ b/Engine/Documentation/Source/Programming/Development/BuildConfigurations/BuildConfigurations.INT.udn
@@ -9,9 +9,17 @@ Related: Support/Builds
 SkillLevel: Advanced
 Version: 4.8
 
+[TOC (start:2 end:3)]
+
+## Build Configurations
+
 Unreal Engine 4 (UE4) uses a custom building method via the UnrealBuildTool. This tool handles all the complex aspects of compiling the engine. 
 
-### Build Configuration
+### Unreal Engine Source Version
+
+[INCLUDE:Programming/Development/CompilingProjects#winbuildconfig]
+
+### Unreal Engine Binary Version
 
 %Globals:osselect_linux%
 
@@ -26,7 +34,7 @@ Unreal Engine 4 (UE4) uses a custom building method via the UnrealBuildTool. Thi
 	active_button_content
 	[/PARAMLITERAL]
 	[PARAM:content]
-		Windows Test.
+	Test Windows.
 	[/PARAM]
 [/OBJECT]
 [OBJECT:ToggleButtonContent]
