UEDOC-2175

Images have been updated and gif has been converted into a YouTube embedded video

diff --git a/Engine/Documentation/Source/Engine/AI/BehaviorTrees/QuickStart/14/14.INT.udn b/Engine/Documentation/Source/Engine/AI/BehaviorTrees/QuickStart/14/14.INT.udn
--- a/Engine/Documentation/Source/Engine/AI/BehaviorTrees/QuickStart/14/14.INT.udn
+++ b/Engine/Documentation/Source/Engine/AI/BehaviorTrees/QuickStart/14/14.INT.udn
@@ -56,19 +56,19 @@ You may have already done this back in part 8 of this guide, but just in case yo
 1. You can **Right-click** in the Graph Editor to create nodes
 	1. Add a **Composites** -> **Sequence** Node.
 
-		![](Engine/AI/BehaviorTrees/QuickStart/10/sequence.png)
+		![](sequence.jpg)
 
 	1. Add two (2) **Composites** -> **Selector** Nodes.
 
-		![](Engine/AI/BehaviorTrees/QuickStart/10/selector.png)
+		![](selector.jpg)
 
 1. Connect the Root Node to one of the Selector Nodes
 
-	![](connectNodes.png)
+	![](connectNodes.jpg)
 
 1. Connect the newly connected Sequence Node to the Selector Node and the other Sequence Node. Then arrange them like so:
 
-	![](connectNodes2.png)
+	![](connectNodes2.jpg)
 
 	_2nd Selector Node on the left, Sequence Node on the right._
 
@@ -81,7 +81,7 @@ You may have already done this back in part 8 of this guide, but just in case yo
 1. Select the **Agro Check** box on the **Selector** Node, and using the **Details** panel, set the **Target to Follow** property to `TargetToFollow`, set **Target Location** to `TargetLocation`, adjust the **Interval** to `0.5` and the **Random Deviation** to `0.0`
 
 	[REGION:raw]
-	![](setAgroCheck.png)
+	![](setAgroCheck.jpg)
 	[/REGION]
 
 ### Add the Tasks
@@ -102,13 +102,13 @@ You may have already done this back in part 8 of this guide, but just in case yo
 1. Set the **Blackboard Key** property on the right-most **Move To** Node to `HomeLocation`
 
 	[REGION:raw]
-	![](adjustMoveTo.png)
+	![](adjustMoveTo.jpg)
 	[/REGION]
 
 1. Select the **Wait** Node, and using the **Details** panel, set its **Wait Time** property to `2.5`
 
 	[REGION:raw]
-	![](adjustWait.png)
+	![](adjustWait.jpg)
 	[/REGION]
 
 1. Select the **RapidMoveTo** node, and using the **Details** panel, set its **Target to Follow** property to `TargetToFollow`.
@@ -119,14 +119,14 @@ You may have already done this back in part 8 of this guide, but just in case yo
 1. **Right-click** the lower **Selector** node and select **Add Decorators..** -> **Close Enough**.
 
 	[REGION:raw]
-	![](addDecorator1.png)
+	![](addDecorator1.jpg)
 	[/REGION]
 
 1. Set the **Observer Aborts** property on the **Blackboard** decorator to `Both`, **Key Query** to `Is Set`, and the **Blackboard Key** to `TargetToFollow`
 1. Set the **Acceptable Distance** property on the **Close Enough** decorator to `100.0`, and **Target to Follow** property to `TargetToFollow`.
 
 	[REGION:raw]
-	![](setBB1.png)
+	![](setBB1.jpg)
 	[/REGION]
 
 1. **Right-click** the **Sequence** node and select **Add Decorators..** -> **Blackboard**.
@@ -149,7 +149,17 @@ You may have already done this back in part 8 of this guide, but just in case yo
 And that is it! Leave the **Behavior Tree** editor open while using **Play in Editor** and you will be able to see how the tree is operating. If you get close enough to the AI and it has line of sight to you, it will begin to chase you. If you are able to break line of sight, he will go to your last known location, wait, then return to where he started.
 
 
-![](BTFinal.gif)(convert:false)
+[OBJECT:EmbeddedVideo]
+[PARAMLITERAL:width]
+640
+[/PARAMLITERAL]
+[PARAMLITERAL:height]
+360
+[/PARAMLITERAL]
+[PARAMLITERAL:videoid]
+1td5wSD3PPU
+[/PARAMLITERAL]
+[/OBJECT]
 
 %Steps%
 
