UEDOC-1887

Images have been updated and are showing correctly

diff --git a/Engine/Documentation/Source/Engine/Levels/Levels.INT.udn b/Engine/Documentation/Source/Engine/Levels/Levels.INT.udn
--- a/Engine/Documentation/Source/Engine/Levels/Levels.INT.udn
+++ b/Engine/Documentation/Source/Engine/Levels/Levels.INT.udn
@@ -56,10 +56,10 @@ When playing a Video Game, every object that you see or interact with in the wor
 The ability to create new levels is a must when trying to develop a game. Creating new levels inside UE4 is a very simple process that can be accomplished by following the instructions listed below.
 
 1. To create a new level, we first need to go to the **File** option that is located on the **Main** tool bar.
-![](Level_File_New_Level.png)
+![](Level_File_New_Level.jpg)
 
 1. Next from the **File** menu, we need to select the **New Level** option.
-	![](Level_Create_New_Level.png)
+	![](Level_Create_New_Level.jpg)
 	[Region:tip]
 	You can also create new levels by using the keyboard shortcut  **Crtl + N**. 
 	[/region]
@@ -72,7 +72,7 @@ The ability to create new levels is a must when trying to develop a game. Creati
 	* **Blank:** This will give you a new level that is completely empty. 
 
 1. Select the level that you want by clicking on it with the **Left Mouse  Button**. After you do this, the level type that you picked will now be loaded up in the editor. 
-![](Level_Both_Types.png)
+![](Level_Both_Types.jpg)
 
 [region:tip]
 When connected to source control, new levels that you create and save will be automatically added to source control.
@@ -88,7 +88,7 @@ Opening already created levels inside UE4 is a very simple process that can be a
 	![](Level_File_New_Level.png)
 
 1. Next from the **File** menu, we need to select the **Open Level ** option.
-	![](Level_Open_New_Level.png)
+	![](Level_Open_New_Level.jpg)
 	[Region:tip]
 	You can also open levels by using the keyboard shortcut  **Crtl + O**, or by simply **Double Clicking** on the map asset inside the **Content Browser**. 
 	[/region]
@@ -97,7 +97,7 @@ Opening already created levels inside UE4 is a very simple process that can be a
 	![](Level_Available_Levels.png)
 
 1. To open a level, simply find the level you wish to open in the list, click on it, and then press the button labeled **Open**.
-	![](Level_Open_Selected_Level.png)
+	![](Level_Open_Selected_Level.jpg)
 
 1. Once you have pressed the **Open** button, the level that you have selected will be loaded in the editor. 
 	![](Level_Selected_Level.png)
@@ -109,14 +109,14 @@ Saving levels inside UE4 is a very simple process that can be accomplished by fo
 	![](Level_File_New_Level.png)
 
 1. Next from the **File** menu, we need to select the **Save** option.
-	![](Level_Save_Level.png)
+	![](Level_Save_Level.jpg)
 	[Region:tip]
 	You can also save all opened levels by using the keyboard shortcut  **Crtl + S**. 
 	[/region]
 
 1. Depending on if this is a new map or not, two things can happen when you press **Save**. 
 	* **New Map:** If this is a new map, the first time you press save, you will be prompted with the **Save Level As** dialogue box. In this box, you can give your new level a name and location to be saved to. 
-	![](Level_Save_Level_Box.png)
+	![](Level_Save_Level_Box.jpg)
 	* **Pre-Existing Map:** If this is not a new map but one that already exists in the project, then when you press save, the map will be saved. 
 
 1. Once your level has been saved, you are now free to continue working, open a new level, or close down the Unreal Editor. 
@@ -125,19 +125,19 @@ Saving levels inside UE4 is a very simple process that can be accomplished by fo
 When creating any game, you will need the ability to set a default level for the game to start on. UE4 not only offers you the ability to define what level your game starts on, it also allows you to define what level your editor will start on as well. As with everything in UE4, the process for changing your default level is quite simple and outlined below. 
 
 1. To change any settings specific to our project, we first need to go to the **Edit** option that is located on the **Main** tool bar.
-	![](Level_Edit_Options.png)
+	![](Level_Edit_Options.jpg)
 
 1. From the **Edit** menu, we then need to select the **Project Settings** option.
-	![](Level_Project_Settings.png)
+	![](Level_Project_Settings.jpg)
 
 1. Selecting the **Project Settings** option will open up your projects settings.   
 	![](Level_Project_Settings_Menu.png)
 
 1. In this menu, you can edit all kinds of options about your project but for now we will be looking at the **Maps & Modes** Section. To active the **Maps & Modes** section, simply click on it in the menu. When you do so, you should now have something that looks like this in the **Project Settings** box.
-	![](Level_Maps_&_Modes.png)
+	![](Level_Maps_&_Modes.jpg)
 
 1. To change your projects default game and editor starting levels, simply press on the arrow icon to the right of the name box and then select the level you want from the list that is displayed. 
-	![](Level_Change_Default_Map.png)
+	![](Level_Change_Default_Map.jpg)
 
 1. Repeat the process above if you wish to change the Editor start up map as well. 
 
@@ -146,10 +146,10 @@ When creating any game, you will need the ability to set a default level for the
 Each level can have unique settings applied to it via the **World Settings** options. The **World Settings** options allow you to do everything from making sure the right **Game Mode** is activated when you play the level too allowing you to adjust how the Global Illumination will look for that level. To activate the **World Settings** options, you need to do the following.
 
 1. On the **Toolbar** that is across the top of the editor, click on the **Settings** option. 
-	![](Level_Settings_Menu.png)
+	![](Level_Settings_Menu.jpg)
 
 1. Clicking on the **Settings** option will display the game specific settings menu and from the menu select **World Settings.**
-	![](Level_Settings_Options.png)
+	![](Level_Settings_Options.jpg)
 
 1. Clicking on the **World Settings** option will open up the levels settings in the **World Settings** tab of the **Details** panel.
 	![](Level_Settings_Options_Menu.png)
@@ -162,19 +162,19 @@ Sometimes you might need to move entire levels and all of their content from one
 1. First, make sure that you have the level and all the content that you want to move saved. 
 
 1. Once you have everything saved, find the level that you want to **Migrate** in the **Content Browser.**
-	![](Level_Level_To_Migrate.png)
+	![](Level_Level_To_Migrate.jpg)
 
 1. After the level is found, all you have to do is **Right-Click** on the level and then choose the **Migrate** option from the list.  
-	![](Level_Select_Migrate.png)
+	![](Level_Select_Migrate.jpg)
 
 1. Once the **Migrate** option was selected, you will be presented with an **Asset Report** box that will allow you to pick and choose what you would like to migrate over to the new project. Once you have selected what you want to move, press the **Ok** button.
-	![](Level_What_To_Migrate.png)
+	![](Level_What_To_Migrate.jpg)
 
 1. After you press the **Ok** button, you will then be asked where you want the content to be moved to. Using the **Browse For Folder** dialogue box, find a location where you would like to move the content to and then press the **Ok** button. Remember that you have to move it to a location that is under the ** Content** folder or the migrate might not work correctly. 
-	![](Level_Where_To_Migrate_to.png)
+	![](Level_Where_To_Migrate_to.jpg)
 
 1. After you press **Ok**, the migration process will start and run until it is done. You can tell if the migration was successful or not if you get the following message in the lower right hand portion of your editor window. 
-	![](Level_Migrate_Completed_Scuess.png)
+	![](Level_Migrate_Completed_Scuess.jpg)
 
 ### Level Migration Troubleshooting
 During the process of migrating a level, you might receive a few error messages. Here is a list of known error message and ways to avoid them. 
