UEDOC-1814

Images have been updated to match the style guide

diff --git a/Engine/Documentation/Source/Engine/Blueprints/BP_HowTo/ActorReference/ActorReference.INT.udn b/Engine/Documentation/Source/Engine/Blueprints/BP_HowTo/ActorReference/ActorReference.INT.udn
--- a/Engine/Documentation/Source/Engine/Blueprints/BP_HowTo/ActorReference/ActorReference.INT.udn
+++ b/Engine/Documentation/Source/Engine/Blueprints/BP_HowTo/ActorReference/ActorReference.INT.udn
@@ -12,13 +12,13 @@ _For this example, we are using a new **Blueprint Third Person Template** with *
 
 1. From the **Modes** Menu under **Basic**, drag a **Box Trigger** into your level. 
 
-	![](1AddTriggerBox.png)
+	![](1AddTriggerBox.jpg)
 
 	We are going to use this to trigger movement in another Actor in the level. 
 
 1. In the **Details** panel for the **Box Trigger**, select the **Convert to Class Blueprint** option. 
 
-	![](2ReplaceWithBlueprint.png)
+	![](2ReplaceWithBlueprint.jpg)
 
 	We are going to create a **Blueprint** out of this Actor and assign a variable within it which will store the reference Actor.
 
@@ -26,29 +26,29 @@ _For this example, we are using a new **Blueprint Third Person Template** with *
 
 1. In the Blueprint Editor, under **My Blueprint**, click the **Add Variable** button. 
 
-	![](3AddVariable.png)	
+	![](3AddVariable.jpg)	
 
 1. Click on the variable, then in the **Details** panel, click the **Variable Type** button and select **Actor** under **Object Reference**. 
 
-	![](4VariableType.png)
+	![](4VariableType.jpg)
 
 	Here we are specifying that the reference is of the **Actor** type and is an **Object** in our level. 
 
 1. Name the variable **TargetActor** and click the **Editable** checkbox. 
 
-	![](5NameVariable.png)
+	![](5NameVariable.jpg)
 
 	This indicates that the variable is set to **Public** which allows it to be accessed and set outside of this Blueprint via the **Details** panel in the Main Editor window. This is useful in situations where you want to modify the variable without having to go into the Blueprint to modify it. 
 
 1. You can also enter in text to appear as the **Tooltip** for what the variable does or is intended for. 
 
-	![](7ToolTip.png)
+	![](7ToolTip.jpg)
 
 	Here we have entered in the text "The Actor to Target". 
 
 1. Click the **Compile** button in the upper left portion of the Blueprint Editor window. 
 
-	![](20Compile.png)
+	![](20Compile.jpg)
 
 	It will change to a green check mark after compiling. 
 
@@ -58,11 +58,11 @@ _For this example, we are using a new **Blueprint Third Person Template** with *
 
 1. In the **Details** panel for the window, you should see under **Default** the **Target Actor** variable and tooltip. 
 
-	![](8DetailsPanel.png)
+	![](8DetailsPanel.jpg)
 
 1. From the **Content Browser**, under **Content/StarterContent/Shapes** add three **Shape_Cubes** to the level. 
 
-	![](9AddingAnActor.png)
+	![](9AddingAnActor.jpg)
 
 	[REGION:note]
 	If you are using your own project, you can add any Actor, **Static Mesh**, or **Blueprint** to the level instead. 
@@ -70,7 +70,7 @@ _For this example, we are using a new **Blueprint Third Person Template** with *
 
 1. Whatever you add to the level, in the **Details** panel for each Actor, set their mobility to **Moveable**. 
 
-	![](22Moveable.png)
+	![](22Moveable.jpg)
 
 	This will allow us to move them via Blueprint script. 
 
@@ -78,41 +78,41 @@ _For this example, we are using a new **Blueprint Third Person Template** with *
 
 1. In the **Details** panel, click the **None** box under **Target Actor** and search for **Shape** (or the name of the Actor(s) added).
 
-	![](10DropDownSelection.png)
+	![](10DropDownSelection.jpg)
 
 	When clicking **None** all Actors currently in the level are displayed, however we are going to target one of the cubes we added. 
 
 1. Click the eye dropper icon next to the magnifying glass icon. 
 
-	![](11PickActor.png)
+	![](11PickActor.jpg)
 
 	This will allow you to select an Actor from your level as the Target Actor by clicking on it in the viewport. 
 
 1. Move the eye dropper icon over one of the cubes (or other asset you have added) and click it to make it the Target Actor. 
 
-	![](11ActorSelected.png)
+	![](11ActorSelected.jpg)
 
 1. When you select an Actor, the **Target Actor** variable will update with the selected Actor in the **Details** panel. 
 
-	![](13ActorSelected.png)
+	![](13ActorSelected.jpg)
 
 	Here we have **Set** reference to an Actor which will be our Target Actor. 
 
 1. Open the **Box Trigger** Blueprint, then in the **MyBlueprint** window, click the eye icon and enable **Show Inherited Variables**. 
 
-	![](showVariables.png)
+	![](showVariables.jpg)
 
 	This will allow us to add events on the Components that were inherited when we first converted the Trigger Box to a Class Blueprint. 
 	
 1. Under **Trigger Base**, **Right-click** on **CollisionComponent** and choose **Add Event** then **Add OnComponentBeginOverlap**. 
 
-	![](14AddOverlap.png)
+	![](14AddOverlap.jpg)
 
 	This will create a new **Event** node on the graph. 
 
 1. Hold **Control** and drag the **TargetActor** variable into the graph. 
 
-	![](15GetActorReference.png)
+	![](15GetActorReference.jpg)
 
 	This creates a **Get** node which will get whatever is assigned to the variable which we have set to an Actor in our level. 
 
@@ -124,21 +124,21 @@ _For this example, we are using a new **Blueprint Third Person Template** with *
 
 1. Connect the nodes as shown below and set the **Z** value on the **Add Actor Local Offset** node to **500.0**. 
 
-	![](17AddZValue.png)
+	![](17AddZValue.jpg)
 
 	This will get the target's current location and move it up by 500 along the Z axis whenever the trigger box is overlapped. 
 
 1.	In the **Components** window, click **TriggerBox_Blueprint** so it is selected. 
 
-	![](selectTriggerBox.png)
+	![](selectTriggerBox.jpg)
 
 	If you do not see the **Components** window, you can enable it from the **Window** menu option. 
 
-	![](enableComponentWindow.png)
+	![](enableComponentWindow.jpg)
 
 1.	In the **Details** panel for the Trigger Box, under **Rendering**, un-check the **Actor Hidden In Game** option. 
 
-	![](19ActorHiddenInGame.png)
+	![](19ActorHiddenInGame.jpg)
 
 	This will make the trigger box visible in-game so that you can see where it is to enter it with your player. 
 
@@ -152,7 +152,7 @@ When you enter the Trigger Box, your selected Actor should move up by 500 units.
 
 1.	In the **Details** panel for the **Box Trigger**, click the eye dropper icon next to **Target Actor**. 
 
-	![](11PickActor.png)
+	![](11PickActor.jpg)
 
 1.	In the level viewport, click a different Actor than the one you used before. 
 
