#UE4doc #docs - Removed overview header.

diff --git a/Engine/Documentation/Source/Gameplay/HowTo/SetUpInput/Blueprints/SetInput_BP.INT.udn b/Engine/Documentation/Source/Gameplay/HowTo/SetUpInput/Blueprints/SetInput_BP.INT.udn
--- a/Engine/Documentation/Source/Gameplay/HowTo/SetUpInput/Blueprints/SetInput_BP.INT.udn
+++ b/Engine/Documentation/Source/Gameplay/HowTo/SetUpInput/Blueprints/SetInput_BP.INT.udn
@@ -5,9 +5,9 @@ Description:A How To Guide for Setting Up Inputs in Blueprints.
 
 [TOC (start:2 end:2)]
 
-## Overview
 
-On this page, you will create a simple character and set it up to receive input, you will then assign the character to a **GameMode** so that it is the default pawn during gameplay. If you are familiar with setting up a character or have a character that you want to provide input to (and have a **GameMode** set up to use your character), you can skip the creating a character section and jump right into creating new inputs and implementing them. 
+
+On this page, you will create a simple character and set it up to receive input, you will then assign the character to a **GameMode** so that it is the default pawn during gameplay. If you are familiar with setting up a character or have a character that you want to provide input to (and have a GameMode set up to use your character), you can skip the creating a character section and jump right into creating new inputs and implementing them. 
 
 ## Creating a Simple Character
 
