UEDOC-1871

Replacing Images with the Blue Guy in them

diff --git a/Engine/Documentation/Source/Engine/Animation/AnimHowTo/PreviewingAnimations/PreviewingHowTo.INT.udn b/Engine/Documentation/Source/Engine/Animation/AnimHowTo/PreviewingAnimations/PreviewingHowTo.INT.udn
--- a/Engine/Documentation/Source/Engine/Animation/AnimHowTo/PreviewingAnimations/PreviewingHowTo.INT.udn
+++ b/Engine/Documentation/Source/Engine/Animation/AnimHowTo/PreviewingAnimations/PreviewingHowTo.INT.udn
@@ -19,11 +19,11 @@ This page covers the different methods in which you can load and preview animati
 
 The most common method of opening and previewing an animation asset is by **Double-clicking** on it inside the **Content Browser**. 
 
-![](Preview1.png)
+![](PA_Intro.jpg)
 
 This will open the asset selected inside the **Persona Animation Editor**. 
 
-![](Preview2.png)
+![](pa_DoubleClickOpen.jpg)
 
 The selected asset will automatically start playing inside the [](Engine/Animation/Persona/Viewport) and auto-loop playback of the motion. 
 
@@ -32,7 +32,7 @@ The selected asset will automatically start playing inside the [](Engine/Animati
 If **Persona** is already open, on the **Animation** tab, you can **Double-click** an asset in the **Asset Browser** to preview it in the Viewport. 
 
 [REGION:lightbox]
-[![](Preview3.png)](Preview3.png)
+[![](PA_SwitchingAnim.jpg)](PA_SwitchingAnim.jpg)
 [/REGION:lightbox]
 
 _Click for full image._
@@ -43,12 +43,12 @@ Also while on the **Animation** tab of **Persona**, you can set the preview anim
 
 After clicking the pick asset button, a pop-up menu will appear where you can select the animation asset to preview. 
 
-![](Preview7.png)
+![](PA_AssetToPreview.jpg)
 
 If you are on a different tab of **Persona**, you can enable the **Asset Browser** window from the **Menu Toolbar** / **Window** option. 
 
 [REGION:lightbox]
-[![](Preview4.png)](Preview4.png)
+[![](PA_Skeleton.jpg)](PA_Skeleton.jpg)
 [/REGION:lightbox]
 
 _Click for full image._
@@ -56,14 +56,14 @@ _Click for full image._
 The **Asset Browser** window will then appear, which you can dock to any other window. 
 
 [REGION:lightbox]
-[![](Preview5.png)](Preview5.png)
+[![](PA_SkelAssetBrowser.jpg)](PA_SkelAssetBrowser.jpg)
 [/REGION:lightbox]
 
 _Click for full image._
 
 You can also get a smaller, looping preview directly inside the **Asset Browser** when highlighting an animation asset. 
 
-![](Preview8.png)
+![](PA_AssetPreview.jpg)
 
 
 ## Persona Viewport Controls
@@ -107,7 +107,7 @@ These controls are mirrored on the arrow keys and num pad to provide alternate a
 
 This section illustrates the options contained within the **Persona Viewport**. 
 
-![](Preview13.png)
+![](PA_PersonaViewportOptions.jpg)
 
 1.	**View Options Toolbar**
 2.	**Movement Toolbar**
@@ -125,7 +125,7 @@ This section illustrates the options contained within the **Persona Viewport**.
 
 You can toggle the information displayed in the viewport from the **Show** option under **Display info**. 
 
-![](Preview14.png)
+![](PA_AnimationDisplay.jpg)
 
 **Basic** will show a condensed version of the display info, while **Detailed** will show more information and **None** will hide the display info. 
 
@@ -141,11 +141,11 @@ When you open an **AnimSequence** or **AnimMontage**, at the bottom of the cente
 
 You can preview the blending of animations when applied to a [Blend Space](Engine/Animation/Blendspaces) or [](Engine/Animation/AimOffset) inside **Persona** as long as the **Enable Preview BlendSpace** option is checked. 
 
-![](Preview10.png)
+![](Preview10_1.jpg)
 
 To see the blending between animations, move the **Mouse Cursor** inside the **Parameters** grid. 
 
-![](Preview11.png)
+![](Preview11_1.jpg)
 
 Remove the Tooltip Display by unchecking the **Enable Tooltip Display** option. 
 
