#UE4doc #docs - removed overview header

diff --git a/Engine/Documentation/Source/Gameplay/Networking/Actors/Properties/Properties.INT.udn b/Engine/Documentation/Source/Gameplay/Networking/Actors/Properties/Properties.INT.udn
--- a/Engine/Documentation/Source/Gameplay/Networking/Actors/Properties/Properties.INT.udn
+++ b/Engine/Documentation/Source/Gameplay/Networking/Actors/Properties/Properties.INT.udn
@@ -8,7 +8,7 @@ Related: Gameplay/Networking/Actors/RPCs
 
 [TOC (start:2 end:3)]
 
-## Overview
+
 
 Each actor maintains a list of properties, and each property can be marked to be replicated to clients. Whenever the value of the variable changes on the server side, the server sends the client the updated value. The variable may have changed on client side too - in which case the new value will overwrite it. Property updates only come from the server, the client will never send property updates to the server.
 
