#UE4doc #docs - removed overview header.

diff --git a/Engine/Documentation/Source/Gameplay/Networking/Actors/Roles/Roles.INT.udn b/Engine/Documentation/Source/Gameplay/Networking/Actors/Roles/Roles.INT.udn
--- a/Engine/Documentation/Source/Gameplay/Networking/Actors/Roles/Roles.INT.udn
+++ b/Engine/Documentation/Source/Gameplay/Networking/Actors/Roles/Roles.INT.udn
@@ -7,7 +7,7 @@ Related: Gameplay/Networking/Server
 
 [TOC (start:2 end:3)]
 
-## Overview
+
 
 There are two properties on actor that are important when it comes to replication. `Role` and `RemoteRole`. 
 
