UEDOC-1999

Added Lightboxes to 2 images on the Doc page

diff --git a/Engine/Documentation/Source/Resources/Showcases/BlueprintExamples/SecurityDoor/SecurityDoor.INT.udn b/Engine/Documentation/Source/Resources/Showcases/BlueprintExamples/SecurityDoor/SecurityDoor.INT.udn
--- a/Engine/Documentation/Source/Resources/Showcases/BlueprintExamples/SecurityDoor/SecurityDoor.INT.udn
+++ b/Engine/Documentation/Source/Resources/Showcases/BlueprintExamples/SecurityDoor/SecurityDoor.INT.udn
@@ -46,8 +46,8 @@ The Components List, Construction Script, and Event Graph for this Blueprint per
 
 
 
-[REGION:fullwidth]
-![](CameraEG.png)
+[REGION:lightbox]
+![](CameraEG.png)(w:920)
 [/REGION]
 
 
@@ -98,7 +98,9 @@ The Components List, Construction Script, and Event Graph for this Blueprint per
 | **Construction Script** | [INCLUDE:#AlarmCS] |
 | **Event Graph** | [INCLUDE:#AlarmEG]  |
 
-![](AlarmGraph.png)
+[Region:lightbox]
+![](AlarmGraph.png)(w:920)
+[/region]
 
 
 ### Editable Variables
