#UE4doc #docs - Updated images, and removed overview header.

diff --git a/Engine/Documentation/Source/GettingStarted/HowTo/ManipulatingActors/ManipulatingActors.INT.udn b/Engine/Documentation/Source/GettingStarted/HowTo/ManipulatingActors/ManipulatingActors.INT.udn
--- a/Engine/Documentation/Source/GettingStarted/HowTo/ManipulatingActors/ManipulatingActors.INT.udn
+++ b/Engine/Documentation/Source/GettingStarted/HowTo/ManipulatingActors/ManipulatingActors.INT.udn
@@ -11,7 +11,7 @@ Manipulating Actors in Unreal Engine 4 (or **Transforming Actors**) refers to th
 [TOC(start:2 end:2)]
 
 [REGION:note]
-For this How To guide, we have created a new project using the **Blueprint Third Person** template with **Starter Content** enabled. If you have not created a project based off a template before or are not sure how to enable Starter Content, refer to the [](Engine/Basics/Projects/Browser) page for information on creating/opening projects, templates, and project settings. 
+For this How To guide, we have created a new project using the **Blueprint Third Person** template with **Starter Content** enabled. If you have not created a project based off a template before or are not sure how to enable Starter Content, refer to the [**Project Browser**](Engine/Basics/Projects/Browser) page for information on creating/opening projects, templates, and project settings. 
 [/REGION]
 
 ## Moving Actors
@@ -93,7 +93,7 @@ This section will show you how to rotate Actors inside the Level Editor Viewport
 
 	![](Rotate5.png)(w:640)
 
-1.	Drag left or right to rotate around the arch, rotating the Actor. 
+1.	**Drag** left or right to rotate around the arch, rotating the Actor. 
 
 	![](Rotate6.png)(w:640)
 
