UEDOC-1998

Lightbox has been added to the image

diff --git a/Engine/Documentation/Source/Resources/Showcases/BlueprintExamples/Foliage/Foliage.INT.udn b/Engine/Documentation/Source/Resources/Showcases/BlueprintExamples/Foliage/Foliage.INT.udn
--- a/Engine/Documentation/Source/Resources/Showcases/BlueprintExamples/Foliage/Foliage.INT.udn
+++ b/Engine/Documentation/Source/Resources/Showcases/BlueprintExamples/Foliage/Foliage.INT.udn
@@ -30,8 +30,8 @@ The Components List, Construction Script, and Event Graph for this Blueprint per
 | **Construction Script** | [INCLUDE:#FoliageCS] |
 | **Event Graph** | This Blueprint has no **Event Graph**, meaning that once its properties are set, it remains the same throughout gameplay. |
 
-[REGION:fullwidth]
-![](FoliageCS.png)
+[REGION:lightbox]
+![](FoliageCS.png)(w:920)
 [/REGION]
 
 
