#UE4doc #docs - Removed Overview header

diff --git a/Engine/Documentation/Source/Platforms/Android/Reference/AndroidGettingStarted.INT.udn b/Engine/Documentation/Source/Platforms/Android/Reference/AndroidGettingStarted.INT.udn
--- a/Engine/Documentation/Source/Platforms/Android/Reference/AndroidGettingStarted.INT.udn
+++ b/Engine/Documentation/Source/Platforms/Android/Reference/AndroidGettingStarted.INT.udn
@@ -9,8 +9,6 @@ Description:How to install different Android SDKs, set environment variables, an
 
 
 
-## Overview 
-
 Here is what we currently require for Android game development:
 
 
