[UEDOC-2014] Added an image of the final product to the tutorial.

diff --git a/Engine/Documentation/Source/Programming/Tutorials/PlayerCamera/3/PlayerCamera_3.INT.udn b/Engine/Documentation/Source/Programming/Tutorials/PlayerCamera/3/PlayerCamera_3.INT.udn
--- a/Engine/Documentation/Source/Programming/Tutorials/PlayerCamera/3/PlayerCamera_3.INT.udn
+++ b/Engine/Documentation/Source/Programming/Tutorials/PlayerCamera/3/PlayerCamera_3.INT.udn
@@ -135,15 +135,18 @@ version:4.9
 		}
 	Using **GetActorForwardVector** and **GetActorRightVector** allows us to move relative to the direction the actor is facing. Since the camera faces the same way as the actor, this ensures that our forward key is always forward relative to what the player sees.
 
-
 1. We have finished coding. We can now compile our code and drag an instance of our new class from the **Content Browser** into the **Level Editor** window inside the **Unreal Engine** editor.
 
 	[![](ClassInContentBrowser.png)](ClassInContentBrowser.png)
 
-	Feel free to add a **Static Mesh** or other visual Component, or play without one. You should feel a smooth acceleration and deceleration to your camera's movement as it follows you around the level, but your rotation should feel tight and instantaneous. Try changing some of the properties on the **SpringArmComponent** to see how they affect the feel of your controls, such as adding "Camera Rotation Lag" or incresing or decreasing Camera Lag.
+	Feel free to add a **Static Mesh** or other visual Component, or play without one. You should feel a smooth acceleration and deceleration to your camera's movement as it follows you around the level, but your rotation should feel tight and instantaneous. Try changing some of the properties on the **SpringArmComponent** to see how they affect the feel of your controls, such as adding "Camera Rotation Lag" or increasing or decreasing Camera Lag.
 
 	[![](SpringArmValues.png)](SpringArmValues.png)
 
+1. With a rotated cone mesh, our final product should look something like this:
+
+	[![](FinalScreen.png)](FinalScreen.png)
+
 ----
 
 [OBJECT:Section]
