#UE4doc #docs Re-organizing documentation.

diff --git a/Engine/Documentation/Source/Programming/Development/CompilingProjects/CompilingProjects.INT.udn b/Engine/Documentation/Source/Programming/Development/CompilingProjects/CompilingProjects.INT.udn
--- a/Engine/Documentation/Source/Programming/Development/CompilingProjects/CompilingProjects.INT.udn
+++ b/Engine/Documentation/Source/Programming/Development/CompilingProjects/CompilingProjects.INT.udn
@@ -37,7 +37,6 @@ the **Editor** target keyword, but if you are building an executable version of
 | ------------------          | ----------                                                                                                                                                                                                                                                                          |
 |  [empty]                    | This configuration builds a stand-alone executable version of your project, but requires cooked content specific to the platform.                                                                                                                                                   |
 | **Editor**                  | To be able to open a project in Unreal Editor and see all code changes reflected, the project must be built in an **Editor** configuration.                                                                                                                                         |  
-| **Uncooked**                | Projects built using configurations that have an **Uncooked** type should be opened in Unreal Editor with the "-game" flag.  This runs your game with uncooked content, in a new window, and is equivalent to **Play in > New Window at Default Player Start** in the editor.       | 
 | **Client**                  | If there is a `<Game>Client.Target.cs` file, the **Client** build configurations will be valid.                                                                                                                                                                                     |
 | **Server**                  | If there is a `<Game>Server.Target.cs` file, the **Server** build configurations will be valid.                                                                                                                                                                                     |                                                                                                                                
 [/EXCERPT:winbuildconfig]
@@ -60,6 +59,8 @@ The build configuration can be set in the Visual Studio toolbar.
 
 ### Setting the Solution Platform
 
+The solution platform can be set in the Visual Studio toolbar.
+
 [PUBLISH:Rocket]
 UE4 currently supports the following platforms:
 
@@ -183,3 +184,6 @@ Xcode's LLDB debugger needs to be configured to correctly handle breakpoints in
 -->
 [/PUBLISH:Rocket]
 [/EXCERPT:Xcode]
+
+<!-- The following was deprecated per LaurenR 29.7.15
+| **Uncooked**                | Projects built using configurations that have an **Uncooked** type should be opened in Unreal Editor with the "-game" flag.  This runs your game with uncooked content, in a new window, and is equivalent to **Play in > New Window at Default Player Start** in the editor.       | -->
\ No newline at end of file
