#UE4doc #docs Updating 'Build Configuration' documentation.

diff --git a/Engine/Documentation/Source/Programming/Development/BuildConfigurations/BuildConfigurations.INT.udn b/Engine/Documentation/Source/Programming/Development/BuildConfigurations/BuildConfigurations.INT.udn
--- a/Engine/Documentation/Source/Programming/Development/BuildConfigurations/BuildConfigurations.INT.udn
+++ b/Engine/Documentation/Source/Programming/Development/BuildConfigurations/BuildConfigurations.INT.udn
@@ -15,10 +15,19 @@ Version: 4.8
 
 Unreal Engine 4 (UE4) uses a custom building method via the UnrealBuildTool. This tool handles all the complex aspects of compiling the engine. 
 
+[INCLUDE:Programming/Development/CompilingProjects#winbuildconfig]
+
 ### Unreal Engine Source Version
 
-[INCLUDE:Programming/Development/CompilingProjects#winbuildconfig]
+The following build configurations are available when building the source version of Unreal Engine:
 
+|             | Debug    | DebugGame | Development | Shipping | Test     |
+| ----------- | -----    | --------- | ----------- | -------- | ----     |
+| **[empty]** | &#x2713; | &#x2713;  | &#x2713;    | &#x2713; | &#x2713; |
+| **Editor ** | &#x2713; | &#x2713;  | &#x2713;    |          |          |
+| **Client ** | &#x2713; | &#x2713;  | &#x2713;    | &#x2713; | &#x2713; |
+| **Server ** | &#x2713; | &#x2713;  | &#x2713;    | &#x2713; | &#x2713; |
+													   
 ### Unreal Engine Binary Version
 
 %Globals:osselect_linux%
@@ -34,7 +43,13 @@ Unreal Engine 4 (UE4) uses a custom building method via the UnrealBuildTool. Thi
 	active_button_content
 	[/PARAMLITERAL]
 	[PARAM:content]
-	Test Windows.
+	The following build configurations are available when building the binary version of Unreal Engine:
+	|             | Debug    | DebugGame | Development | Shipping | Test     |
+    | ----------- | -----    | --------- | ----------- | -------- | ----     |
+    | **[empty]** |          | &#x2713;  | &#x2713;    | &#x2713; |          |
+    | **Editor ** |          | &#x2713;  | &#x2713;    |          |          |
+    | **Client ** |          |           |             |          |          |
+    | **Server ** |          |           |             |          |          |
 	[/PARAM]
 [/OBJECT]
 [OBJECT:ToggleButtonContent]
