[UEDOC-2064] Added screenshots and text.

diff --git a/Engine/Documentation/Source/Programming/Tutorials/UMG/5/UMG_5.INT.udn b/Engine/Documentation/Source/Programming/Tutorials/UMG/5/UMG_5.INT.udn
--- a/Engine/Documentation/Source/Programming/Tutorials/UMG/5/UMG_5.INT.udn
+++ b/Engine/Documentation/Source/Programming/Tutorials/UMG/5/UMG_5.INT.udn
@@ -66,6 +66,11 @@ version:4.9
 
 	[REGION:note]After deactivating the menu with the Play Game Button, we will be unable to do anything further in our game. This is the point at which we would normally load the first level, play an introductory cutscene, or spawn and possess a **Pawn**.[/REGION]
 
+1. We should now have two screens that look roughly as follows:
+
+	[![](FinalScreen.png)](FinalScreen.png)
+
+	[![](FinalScreen2.png)](FinalScreen2.png)
 
 ----
 
