#UE4doc #docs: ##Overview removed as according to the instructional guide.

diff --git a/Engine/Documentation/Source/Programming/Rendering/PrimitiveTypes/PrimitiveTypes.INT.udn b/Engine/Documentation/Source/Programming/Rendering/PrimitiveTypes/PrimitiveTypes.INT.udn
--- a/Engine/Documentation/Source/Programming/Rendering/PrimitiveTypes/PrimitiveTypes.INT.udn
+++ b/Engine/Documentation/Source/Programming/Rendering/PrimitiveTypes/PrimitiveTypes.INT.udn
@@ -9,10 +9,6 @@ Description:Programmer's guide to creating new types of render primitives.
 
 [TOC (start:2 end:3)]
 
-
-
-## Overview
-
 UPrimitiveComponents are Actor components that are rendered and collided against.  This page describes how to create a new UPrimitiveComponent type.
 
 
