UEDOC-2133

lightboxes have been added to imgaes on this page

diff --git a/Engine/Documentation/Source/Engine/Rendering/Materials/HowTo/Subsurface_Profile/Subsurface_Profile.INT.udn b/Engine/Documentation/Source/Engine/Rendering/Materials/HowTo/Subsurface_Profile/Subsurface_Profile.INT.udn
--- a/Engine/Documentation/Source/Engine/Rendering/Materials/HowTo/Subsurface_Profile/Subsurface_Profile.INT.udn
+++ b/Engine/Documentation/Source/Engine/Rendering/Materials/HowTo/Subsurface_Profile/Subsurface_Profile.INT.udn
@@ -24,7 +24,9 @@ When completed, your **Content Browser** should look like this.
 	![](SSP_Create_New_Material.png)
 
 1. Next, open up the Material by **Double-Clicking** on it in the **Content Browser** with the **Right Mouse Button**.
+	[REGION:lightbox]
 	![](SSP_Open_New_Material.png)
+	[/REGION]
 
 1. Now with the Material open, we need to change the Materials **Shading Model** from **Default Lit** to **Subsurface Profile**. 
 	[OBJECT:EmbeddedVideo]
@@ -66,7 +68,9 @@ In the following steps, we will be creating a Material and a Material Instance t
 1. The first thing that we need to do is to make sure that our Subsurface Material that was created above is open. If you have the Material open, please go on to the next step. 
    If not, to open the Material, first find it in the **Content Browser** and then **Double-Click** on it with the **Left Mouse Button** to open it up.
    When completed, you should have something like this on your screen. 
+	[REGION:lightbox]
 	![](SSP_Open_New_Material.png)
+	[/REGION]
 
 1. With the Material open, it is time to place down some Material Expression nodes. For this tutorial, we are going to make use of the following nodes. 
 	* **Vector Parameter** x 1
@@ -123,7 +127,9 @@ Once your mouse is over the input you want to use, release the **Left Mouse Butt
 
 1. With the Subsurface Profile applied and the Material expression linked to the main Material node, it is now time to compile the shader and get ready to make a Material Instance out of this Material. 
    When completed with the Material, it should look something like this. 
+	[REGION:lightbox]
 	![](SSP_Completed_Material.png)
+	[/REGION]
 
 ## Using a Subsurface Profile in a Material Instance
 
@@ -136,11 +142,15 @@ When completed, your **Content Browser** should look something like this.
 
 1. Now that the Material Instance has been created, open it up by using the **Left Mouse Button** to **Double-Click** on it in the **Content Browser**. Once open, you should see something that looks like this. 
 
+	[REGION:lightbox]
 	![](SSP_Material_Instance.png)
+	[/REGION]
 
 1. To change the Subsurface Profile in a Material Instance, first make sure that you enable the **Override Subsurface Profile** and then provide the Subsurface Profile that you want to use 
 	by using the **Left Mouse  Button** to click on the **Subsurface Profile** input. Once you click on the input you should see the Subsurface Profiles that are available to choose from. 
+	[REGION:lightbox]
 	![](SSP_SSP_In_Material_Instance.png)
+	[/REGION]
 
 	[region:note]
 	If you have already supplied a Subsurface Profile in the Master Material, then you do not need to override it in the Material Instance. This is only if you wish you use a different Subsurface Profile than what is already being used. 
@@ -179,7 +189,9 @@ When completed, your **Content Browser** should look something like this.
 
 	When completed, you should have something that looks like the following.
 
+	[REGION:lightbox]
 	![](Engine\Rendering\Materials\HowTo\Subsurface_Scattering\SSM_Light_Settings.png)
+	[/REGION]
 
 1. With all of the lights added, we can now add a Static Mesh to test with. For this example, we will be using the following mesh. 
 
