#UE4doc #docs: ##Overview removed as according to the instructional guide. Updated engine version to 4.9.

diff --git a/Engine/Documentation/Source/Programming/Rendering/MaterialExpressions/CreatingMaterialExpressions.INT.udn b/Engine/Documentation/Source/Programming/Rendering/MaterialExpressions/CreatingMaterialExpressions.INT.udn
--- a/Engine/Documentation/Source/Programming/Rendering/MaterialExpressions/CreatingMaterialExpressions.INT.udn
+++ b/Engine/Documentation/Source/Programming/Rendering/MaterialExpressions/CreatingMaterialExpressions.INT.udn
@@ -2,21 +2,19 @@ Availability: Docs
 Title:Material Expressions Technical Guide
 Crumbs:%ROOT%, Programming, Programming/Rendering
 Description:Guide to creating new material expressions for use in the Material Editor.
-Version: 4.5
+Version: 4.9
 
 [TOC (start:2 end:3)]
 
-## Overview
-
 Creating a new material expression in UE4 is straightforward.  As an example, I will create a new material expression that squares the input, called MaterialExpressionSquare.  Note that all material expressions in UE4 are implemented in `MaterialExpressions.cpp`.
 
 
 
 ## Declaration
 
-First, add a new engine script file with a new type that derives from MaterialExpression.  Declare the expression's inputs and virtual functions, for example:<br>  
+First, add a new engine script file with a new type that derives from MaterialExpression.  Declare the expression's inputs and virtual functions, for example:<br> 
 
-`MaterialExpressionSquare.h`
+	MaterialExpressionSquare.h
 
 	UCLASS(HeaderGroup=Material, collapsecategories, hidecategories=Object)
 	public class UMaterialExpressionSquare : public UMaterialExpression
