#UE4doc #docs: Made minor spelling and grammar changes: its -> it is, compliment -> complement, classes needs -> classes need

diff --git a/Engine/Documentation/Source/Programming/Introduction/IntroToProgramming.INT.udn b/Engine/Documentation/Source/Programming/Introduction/IntroToProgramming.INT.udn
--- a/Engine/Documentation/Source/Programming/Introduction/IntroToProgramming.INT.udn
+++ b/Engine/Documentation/Source/Programming/Introduction/IntroToProgramming.INT.udn
@@ -20,7 +20,7 @@ It is possible to write "plain old C++ code" in Unreal Engine, but you will be m
 
 Unreal Engine provides two methods, C++ and Blueprints Visual Scripting, to create new gameplay elements. Using C++, programmers add the base gameplay systems that designers can then build upon or with to create the custom gameplay for a level or the game. In these cases, the C++ programmer works in their favorite IDE (usually Microsoft Visual Studio, or Apple's Xcode) and the designer works in the Unreal Editor's Blueprint Editor.
 
-The gameplay API and framework classes are available to both of these systems, which can be used separately, but show their true power when used in conjunction to compliment each other. What does that really mean, though? It means that the engine works best when programmers are creating gameplay building blocks in C++ and designers take those blocks and make interesting gameplay.
+The gameplay API and framework classes are available to both of these systems, which can be used separately, but show their true power when used in conjunction to complement each other. What does that really mean, though? It means that the engine works best when programmers are creating gameplay building blocks in C++ and designers take those blocks and make interesting gameplay.
 
 With that said, let us take a look at a typical workflow for the C++ programmer that is creating building blocks for the designer. In this case, we are going to create a class that is later extended via Blueprints by a designer or programmer. In this class, we are going to create some properties that the designer can set and we are going to derive new values from those properties. The whole process is very easy to do using the tools and C++ macros we provide for you.
 
@@ -418,7 +418,7 @@ In this section we will go over basic memory management and the garbage collecti
 
 ### UObjects and Garbage Collection
 
-UE4 uses the reflection system to implement a garbage collection system.  With garbage collection, you will not have to manually manage deleting your UObjects, you just need to maintain valid references to them.  Your classes needs to derive from UObject in order to be enabled for garbage collection.  Here is the simple example class we will be using:
+UE4 uses the reflection system to implement a garbage collection system.  With garbage collection, you will not have to manually manage deleting your UObjects, you just need to maintain valid references to them.  Your classes need to derive from UObject in order to be enabled for garbage collection.  Here is the simple example class we will be using:
 
 	UCLASS()
 	class MyGCType : public UObject
@@ -467,7 +467,7 @@ This is a more common case, where you have actors with UObject properties.
 		World->SpawnActor<AMyActor>(Location, Rotation);
 	}
 
-When we call the above function, we spawn an actor into the world.  The actor’s constructor creates two objects.  One gets assigned to a UPROPERTY, the other to a bare pointer.  Since actors are automatically a part of the root set, SafeObject will not be garbage collected because it can be reached from a root set object.  DoomedObject, however, will not fare so well.  We didn’t mark it with UPROPERTY, so the collector does not know its being referenced, and will eventually destroy it.
+When we call the above function, we spawn an actor into the world.  The actor’s constructor creates two objects.  One gets assigned to a UPROPERTY, the other to a bare pointer.  Since actors are automatically a part of the root set, SafeObject will not be garbage collected because it can be reached from a root set object.  DoomedObject, however, will not fare so well.  We didn’t mark it with UPROPERTY, so the collector does not know it is being referenced, and will eventually destroy it.
 
 When a UObject is garbage collected, all UPROPERTY references to it will be set to nullptr for you.  This makes it safe for you to check if an object has been garbage collected or not.
 
