#UE4doc #docs: ##Overview removed as according to the instructional guide.

diff --git a/Engine/Documentation/Source/Programming/Development/Tools/Commandlets/MakingACommandlet/MakingACommandlet.INT.udn b/Engine/Documentation/Source/Programming/Development/Tools/Commandlets/MakingACommandlet/MakingACommandlet.INT.udn
--- a/Engine/Documentation/Source/Programming/Development/Tools/Commandlets/MakingACommandlet/MakingACommandlet.INT.udn
+++ b/Engine/Documentation/Source/Programming/Development/Tools/Commandlets/MakingACommandlet/MakingACommandlet.INT.udn
@@ -9,10 +9,6 @@ Description:How to create completely new, custom commandlets to analyze or modif
 
 [TOC(start:2)]
 
-
-
-## Overview
-
 Commandlets are command line programs that run inside of the Unreal Engine environment. They are most often used to make bulk changes to content, iterate over content to get information about it, or as a unit testing mechanism.
 
 They are basically little applets that have the engine initialized but are mostly just confined to Windows console output.
