#UE4doc #docs: ##Overview removed as according to the instructional guide. Updated engine version to 4.9. Grammar fixes: all of objects -> all objects

diff --git a/Engine/Documentation/Source/Programming/Assets/ThumbnailRendering/ThumbnailRendering.INT.udn b/Engine/Documentation/Source/Programming/Assets/ThumbnailRendering/ThumbnailRendering.INT.udn
--- a/Engine/Documentation/Source/Programming/Assets/ThumbnailRendering/ThumbnailRendering.INT.udn
+++ b/Engine/Documentation/Source/Programming/Assets/ThumbnailRendering/ThumbnailRendering.INT.udn
@@ -2,7 +2,7 @@ Availability:Docs
 Title:Thumbnail Rendering Technical Guide
 Crumbs: %ROOT%, Engine, Engine/UI, Programming
 Description:Programmer's guide to rendering thumbnails for asset types in the Content Browser.
-Version: 4.5
+Version: 4.9
 
 [REGION:todo]
 	This page is converted from Unreal Engine 3 documentation and currently being reviewed. Information may be outdated.
@@ -10,10 +10,6 @@ Version: 4.5
 
 [TOC(start:2)]
 
-
-
-## Overview
-
 The previous rendering system was part of the object in question. It was defined in Core where there was no access to the rendering interfaces needed to render a thumbnail. This posed a problem for object like Archetypes, which had no way to provide a thumbnail view. Also, this editor only function was being included on consoles, where it was just consuming space. Another problem with the system was that it lacked good reuse for items that rendered their thumbnails using the same approach (namely icons for thumbnails).
 
 To alleviate the code duplication, make the code exist in the editor only, and make the system late bound, we created the ThumbnailManager and the support interfaces ThumbnailRenderer and ThumbnailLabelRenderer. The ThumbnailManager uses an object registration system to match the appropriate rendering component(s) for a given resource type. This is similar to the GenericBrowserTypes with the exception being that all of the relevant  information comes from the `MyGameEditor.ini` file rather than through native classes discovery. ThumbnailRenderers render the object they are  associated with in a thumbnail area. ThumbnailLabelRenderers build and render the text that appears beneath the thumbnail. These are separate  objects in order to maximize flexibility and reuse. Most objects use the same label renderer and separating out the interfaces allows  a unique thumbnail rendering object and a shared label rendering object. It also allows for some interesting options at run time, like overriding the suggested label renderer with a memory counting one.
@@ -34,7 +30,7 @@ IconThumbnailRenderer -- Provides an implementation that will render a specified
 
 DefaultSizedThumbnailRenderer -- Provides common data for rendering objects that do not have an inherent size (particles, material instances, etc.) 
 
-TextureThumbnailRenderer -- Used as the base class for all of objects that support the UTexture interface.
+TextureThumbnailRenderer -- Used as the base class for all objects that support the UTexture interface.
 
 
 ### ThumbnailHelpers.h
@@ -58,7 +54,7 @@ MemCountThumbnailLabelRenderer -- Counts the memory used for an object and displ
 
 ## Registering your renderers with the ThumbnailManager
 
-The ThumbnailManager uses the UObject config system to register the classes that are supported for thumbnail rendering. All class are registered by name to avoid package dependencies and allow licensees to create their own in a MyGameUnrealEd package, rather than having to modify Epic code. The snippet below show how you would add a custom renderer. The second entry shows you how to use the existing icon renderer rather than providing a custom one. Note: the line breaks are there for readability and should not be in your `DefaultEditor.ini` when you customize it.
+The ThumbnailManager uses the UObject config system to register the classes that are supported for thumbnail rendering. All class are registered by name to avoid package dependencies and allow licensees to create their own in a MyGameUnrealEd package, rather than having to modify Epic code. The snippet below shows how you would add a custom renderer. The second entry shows you how to use the existing icon renderer rather than providing a custom one. Note: the line breaks are there for readability and should not be in your `DefaultEditor.ini` when you customize it.
 
 
     
