#UE4doc #docs: ##Overview removed as according to the instructional guide.

diff --git a/Engine/Documentation/Source/Programming/Assets/CustomImporters/CustomImporters.INT.udn b/Engine/Documentation/Source/Programming/Assets/CustomImporters/CustomImporters.INT.udn
--- a/Engine/Documentation/Source/Programming/Assets/CustomImporters/CustomImporters.INT.udn
+++ b/Engine/Documentation/Source/Programming/Assets/CustomImporters/CustomImporters.INT.udn
@@ -9,10 +9,6 @@ Description:Programmer's guide to writing importers for custom asset types.
 
 [TOC(start:2)]
 
-
-
-## Overview
-
 UFactory classes are used in the editor for creating new assets, most commonly based on external files.
 
 The sister class of UFactory is UExporter, which is used to export assets for copy/duplicate operations, as well as for exporting assets to file. See [Writing Custom Exporters](Programming/Assets/CustomExporters) for information on how to write an exporter.
