#UE4doc #docs: ##Overview removed as according to the instructional guide.

diff --git a/Engine/Documentation/Source/Programming/Assets/CustomExporters/CustomExporters.INT.udn b/Engine/Documentation/Source/Programming/Assets/CustomExporters/CustomExporters.INT.udn
--- a/Engine/Documentation/Source/Programming/Assets/CustomExporters/CustomExporters.INT.udn
+++ b/Engine/Documentation/Source/Programming/Assets/CustomExporters/CustomExporters.INT.udn
@@ -9,10 +9,6 @@ Description:Programmer's guide to writing exporters for custom asset types.
 
 [TOC(start:2)]
 
-
-
-## Overview
-
 UExporter classes are used in the editor to export assets, for use either by the editor's copy/paste/duplicate functionality, or to export assets to file.  Exporters classes declare the type of object supported by the exporter and a list of supported export formats.  Exporters support either to text or binary.
 
 The sister class of UExporter is UFactory, which is used to import and paste assets, as well as to create new assets.  See [WritingCustomImporters](Programming/Assets/CustomImporters) for more information on how to write an importer.
