#UE4doc #docs: Updated engine version to 4.9. Removed extraneous commas: execution, but, instead, are used -> execution, but instead are used

diff --git a/Engine/Documentation/Source/Programming/Assertions/Assertions.INT.udn b/Engine/Documentation/Source/Programming/Assertions/Assertions.INT.udn
--- a/Engine/Documentation/Source/Programming/Assertions/Assertions.INT.udn
+++ b/Engine/Documentation/Source/Programming/Assertions/Assertions.INT.udn
@@ -2,7 +2,7 @@
 Title: Assertions
 Crumbs:
 Description:
-Version:4.7
+Version:4.9
 
 Assertions are tool to verify assumptions that a given piece of code relies on. This can be as simple as verifying that a pointer is not NULL or as complicated as verifying a particular function is not reentered. UE4 provides a series of macros to perform these types of validations. They are macros so that they can compile out in certain build configurations, either for performance reasons or because they aren't needed in a final build. If you want to see the macros for yourself, you can find them in:
 
@@ -115,7 +115,7 @@ Examples:
 	checkfSlow(Class->IsA(AActor::StaticClass()), TEXT("Class (%s) is wrong type"), Class->GetName());
 	verifySlow(LastValue == Current);
 
-The final class of runtime assertions do not halt execution, but, instead, are used to create a callstack report to aid in tracking down issues. The expression in these macros is always executed and can be placed inside conditionals. These macros are enabled with the DO_CHECK define.
+The final class of runtime assertions do not halt execution, but instead are used to create a callstack report to aid in tracking down issues. The expression in these macros is always executed and can be placed inside conditionals. These macros are enabled with the DO_CHECK define.
 
 **ensure(expression);**
 
