#UE4doc #docs Updated build configuration tables. Availibility: Docs.

diff --git a/Engine/Documentation/Source/Programming/Development/BuildConfigurations/BuildConfigurations.INT.udn b/Engine/Documentation/Source/Programming/Development/BuildConfigurations/BuildConfigurations.INT.udn
--- a/Engine/Documentation/Source/Programming/Development/BuildConfigurations/BuildConfigurations.INT.udn
+++ b/Engine/Documentation/Source/Programming/Development/BuildConfigurations/BuildConfigurations.INT.udn
@@ -17,9 +17,11 @@ Unreal Engine 4 (UE4) uses a custom building method via the UnrealBuildTool. Thi
 
 [INCLUDE:Programming/Development/CompilingProjects#winbuildconfig]
 
-### Unreal Engine Source Version
+### UE4 Build Configuration - Source
 
-The following build configurations are available when building the source version of Unreal Engine:
+(#bookmark1)
+
+The following build configurations are available when building your project with the source version of UE4:
 
 |             | Debug    | DebugGame | Development | Shipping | Test     |
 | ----------- | -----    | --------- | ----------- | -------- | ----     |
@@ -28,7 +30,9 @@ The following build configurations are available when building the source versio
 | **Client ** | &#x2713; | &#x2713;  | &#x2713;    | &#x2713; | &#x2713; |
 | **Server ** | &#x2713; | &#x2713;  | &#x2713;    | &#x2713; | &#x2713; |
 													   
-### Unreal Engine Binary Version
+### UE4 Build Configuration - Binary
+
+The following build configurations are available when building your project with the binary version of UE4:
 
 %Globals:osselect_linux%
 
@@ -43,13 +47,12 @@ The following build configurations are available when building the source versio
 	active_button_content
 	[/PARAMLITERAL]
 	[PARAM:content]
-	The following build configurations are available when building the binary version of Unreal Engine:
-	|             | Debug    | DebugGame | Development | Shipping | Test     |
-    | ----------- | -----    | --------- | ----------- | -------- | ----     |
-    | **[empty]** |          | &#x2713;  | &#x2713;    | &#x2713; |          |
-    | **Editor ** |          | &#x2713;  | &#x2713;    |          |          |
-    | **Client ** |          |           |             |          |          |
-    | **Server ** |          |           |             |          |          |
+	    |             | Debug | DebugGame | Development | Shipping | Test |
+        | ----------- | ----- | --------- | ----------- | -------- | ---- |
+        | **[empty]** |       | &#x2713;  | &#x2713;    | &#x2713; |      |
+        | **Editor ** |       | &#x2713;  | &#x2713;    |          |      |
+        | **Client ** |       |           |             |          |      |
+        | **Server ** |       |           |             |          |      |
 	[/PARAM]
 [/OBJECT]
 [OBJECT:ToggleButtonContent]
@@ -62,7 +65,12 @@ The following build configurations are available when building the source versio
 	[PARAMLITERAL:active]
 	[/PARAMLITERAL]
 	[PARAM:content]
-		Mac Test.
+	    |             | Debug | DebugGame | Development | Shipping | Test |
+        | ----------- | ----- | --------- | ----------- | -------- | ---- |
+        | **[empty]** |       | &#x2713;  | &#x2713;    | &#x2713; |      |
+        | **Editor ** |       | &#x2713;  | &#x2713;    |          |      |
+        | **Client ** |       |           |             |          |      |
+        | **Server ** |       |           |             |          |      |
 	[/PARAM]
 [/OBJECT]
 [OBJECT:ToggleButtonContent]
@@ -75,6 +83,10 @@ The following build configurations are available when building the source versio
 	[PARAMLITERAL:active]
 	[/PARAMLITERAL]
 	[PARAM:content]
-		Linux Test.
+	    [REGION:warning]
+        We currently don't provide a binary installer of Unreal Engine for Linux.  
+		Linux users should refer to build configurations listed in the ['Unreal Engine Source Version' section of this document](#bookmark1).  
+        _Refer to our [](GettingStarted/DownloadingUnrealEngine) documentation to learn how to build Unreal Engine from source._
+        [/REGION]
 	[/PARAM]
 [/OBJECT]
\ No newline at end of file
