#UE4doc #docs: Fixed a spelling error: instuctions -> instructions

diff --git a/Engine/Documentation/Source/Platforms/HTML5/GettingStarted/HTML5GettingStarted.INT.udn b/Engine/Documentation/Source/Platforms/HTML5/GettingStarted/HTML5GettingStarted.INT.udn
--- a/Engine/Documentation/Source/Platforms/HTML5/GettingStarted/HTML5GettingStarted.INT.udn
+++ b/Engine/Documentation/Source/Platforms/HTML5/GettingStarted/HTML5GettingStarted.INT.udn
@@ -100,7 +100,7 @@ For example:
 	Chrome(32bit)="C:/Program Files (x86)/Google/Chrome/Application/chrome.exe"
 	Chrome(64bit)="C:/Program Files/Google/Chrome/Application/chrome.exe"
 
-Follow the instuctions on [https://github.com/EpicGames/UnrealEngine](https://github.com/EpicGames/UnrealEngine) to build the editor.  Set your solution configuration to **Development** and your solution platform to **HTML5** to build HTML5. 
+Follow the instructions on [https://github.com/EpicGames/UnrealEngine](https://github.com/EpicGames/UnrealEngine) to build the editor.  Set your solution configuration to **Development** and your solution platform to **HTML5** to build HTML5. 
 
 ### Mac OS
 
