#UE4doc #docs: ##Overview removed as according to the instructional guide. Added a [REGION] to a note as according to the intructional guide.

diff --git a/Engine/Documentation/Source/Platforms/HTML5/BuildingEmscripten/HTML5BuildingEmscripten.INT.udn b/Engine/Documentation/Source/Platforms/HTML5/BuildingEmscripten/HTML5BuildingEmscripten.INT.udn
--- a/Engine/Documentation/Source/Platforms/HTML5/BuildingEmscripten/HTML5BuildingEmscripten.INT.udn
+++ b/Engine/Documentation/Source/Platforms/HTML5/BuildingEmscripten/HTML5BuildingEmscripten.INT.udn
@@ -5,9 +5,6 @@ Description:Building the Fastcomp-enabled version of Empscripten
 
 [TOC(start:2 end:3)]
 
-
-## Overview
-
 Currently, the HTML5 pipeline requires a more recent version of Emscripten than is available via the `emsdk` tool. This more recent version makes use of a new compiler core called **Fastcomp**.
 
 Instructions for this process are available on the Emscripten Wiki:
@@ -27,7 +24,9 @@ CMake must be installed and accessible via the command line.
 
 In order to successfully compile the Emscripten LLVM source using MSVC, you must use Visual Studio 2010 SP1. You can alternatively use the MinGW toolchain for compilation but this has not been tested internally at Epic. 
 
-Note: **Service Pack 1 is a necessity**. Without it, the CMake tool will error out while generating the LLVM solution file.
+[REGION:note]
+**Service Pack 1 is a necessity**. Without it, the CMake tool will error out while generating the LLVM solution file.
+[/region]
 
 ## Procedure
 
