#UE4doc #docs:
Resubmit due to changes not being saved

diff --git a/Engine/Documentation/Source/Platforms/GearVR/Debugging/GearVRDebugging.INT.udn b/Engine/Documentation/Source/Platforms/GearVR/Debugging/GearVRDebugging.INT.udn
--- a/Engine/Documentation/Source/Platforms/GearVR/Debugging/GearVRDebugging.INT.udn
+++ b/Engine/Documentation/Source/Platforms/GearVR/Debugging/GearVRDebugging.INT.udn
@@ -12,7 +12,7 @@ Related: Platforms/Mobile/Materials
 Related: Engine/Rendering/Materials/MaterialProperties/BlendModes
 Related: Engine/Rendering/Materials/CustomizedUVs
 Skilllevel: Intermediate
-Version: 4.8
+Version: 4.9
 
 [TOC(start:2 end:3)]
 
@@ -142,7 +142,7 @@ Check out the official Android documentation for much more detailed information
 
 1. The information that is contained within the log file was generated from adding **UE4**, **OVR**,  **VrApi**, **VrLib**,**VRSVC**, and **HmdInfo** to the logcat command. 
 	When you look at the log file that was created and search for anyone of those terms you will find all debug information related to those terms. 
-	You can see an example of this in the image below, some of the terms used, **OVR**, **UE4**, **Vrapi**, and **VrApi** all show up in the log with information that is related to those terms.
+	You can see an example of this in the image below, some of the terms used, OVR, UE4, Vrapi, and VrApi all show up in the log with information that is related to those terms.
 	
 	![](T_ADB_Log_Info.png)
 
