#UE4doc #docs:
Changed Engine Version
Changed color for multiple borders/arrows in pictures from green to yellow
Unbolded a few terms

diff --git a/Engine/Documentation/Source/Platforms/GearVR/BestPractices/GearVRBestPractices.INT.udn b/Engine/Documentation/Source/Platforms/GearVR/BestPractices/GearVRBestPractices.INT.udn
--- a/Engine/Documentation/Source/Platforms/GearVR/BestPractices/GearVRBestPractices.INT.udn
+++ b/Engine/Documentation/Source/Platforms/GearVR/BestPractices/GearVRBestPractices.INT.udn
@@ -12,7 +12,7 @@ Related: Platforms/Mobile/Materials
 Related: Engine/Rendering/Materials/MaterialProperties/BlendModes
 Related: Engine/Rendering/Materials/CustomizedUVs
 Skilllevel: Intermediate
-Version: 4.8
+Version: 4.9
 
 [TOC(start:2 end:3)]
 
@@ -24,13 +24,13 @@ using the Samsung Gear Virtual Reality(VR) Headset and Samsung Galaxy Note 4 Sma
 ## Gear VR Rendering Options
 While combining the Gear VR Headset with the Samsung Galaxy Note 4 Smartphone can produce some very immersive VR experiences, 
 there are some rendering features that UE4 can not perform while using the Gear VR Headset due to hardware limitations. 
-Many of the restrictions that you will encounter are similar to restrictions for **[mobile devices](Platforms/Mobile/Performance)**.
+Many of the restrictions that you will encounter are similar to restrictions for [mobile devices](Platforms/Mobile/Performance).
 However because of the unique rendering requirements for the Gear VR Headset there are few extra things that must be considered.
 In the following sections we will cover what these requirements are and how they can effect your Gear VR project. 
 
 ### Mobile HDR
-While you may be tempted to enable Mobile HDR so that you enable more advanced rendering features for your projects but **DO NOT** do this!
-Mobile HDR **will not work** with the Gear VR Headset and should be left off. 
+While you may be tempted to enable Mobile HDR so that you enable more advanced rendering features for your projects but do not do this.
+Mobile HDR will not work with the Gear VR Headset and should be left off. 
 Enabling this feature will cause your project to not be rendered correctly when placed into the Gear VR Headset.
 
 ### Post Processing
@@ -59,19 +59,19 @@ For example only the following Material inputs will work with the Gear VR Headse
 | **Opacity Mask** |Opacity Mask helps to define if a certain pixel is see through or not. |
 
 [REGION:note]
-Please keep in mind that the both **Translucent** and **Masked** Materials can be very expensive to render, especially on hardware like Smartphone's.
+Please keep in mind that the both Translucent and Masked Materials can be very expensive to render, especially on hardware like Smartphone's.
 Whenever possible you should use [Opaque Materials](Engine/Rendering/Materials/MaterialProperties/BlendModes) as they are much cheaper to render.
 [/region]
 
 ### Material Shading Models
-Due to hardware limitations only the following Material Shading Models are available for use in projects that use the Gear VR Headset.
+Due to hardware limitations only the following **Material Shading Models** are available for use in projects that use the Gear VR Headset.
 | Shading Model Name | Description |
 | -------- | ---------- |
 | **Default Lit** |Default Lit is, as the name suggests, the default shading model and will likely be the one that you will use most often. This shading model makes use of direct and indirect lighting, as well as specularity for reflections.. |
 | **Unlit** |The Unlit Shading Model only outputs Emissive for color, making this perfect for special effects such as fire or illuminating objects.|
 
 [region:note]
-Try to use the **Unlit** Shading Model as much as possible with your Gear VR projects as it is the cheapest to use. 
+Try to use the Unlit Shading Model as much as possible with your Gear VR projects as it is the cheapest to use. 
 [/region]
 
 ### Material Blend Modes
@@ -80,13 +80,13 @@ When creating content for the Gear VR Headset the following Material Blend Modes
 
 | Blend Mode | Description |
 | ---- | ----- |
-| **Opaque** | For a Solid object, **Opaque** is the cheapest material. But, as its name implies, it has no support for translucency. That means it will never create overdraw, and rendering one pixel of the material will always take the same time depending on the complexity of the **Node Graph**. |
+| **Opaque** | For a Solid object, Opaque is the cheapest material. But, as its name implies, it has no support for translucency. That means it will never create overdraw, and rendering one pixel of the material will always take the same time depending on the complexity of the **Node Graph**. |
 | **Masked** |  If you need some pixels not drawn on your material, this is the way to go. The material only creates overdraw for the area that is masked out. But the material is 100\% opaque or 100\% see through. The bigger the area that is masked out the more overdraw is created. That means more expensive. |
-| **Translucent** | This is the most expensive **BlendMode**. It can look like a **Masked** and **Opaque** but will have an additional effect on performance regardless of actually being transparent or not. |
-| **Additive** | Is a cheaper alternative to **Translucent**. It does not need to sort the pixels, they just get added to the scene. But that means it can only brighten. Black pixels on the material have no effect on the scene.|
-| **Modulated** | Is the second cheapest alternative to **Translucent**, but it multiplies its Pixels with the scene. That means it can only darken the scene.|
+| **Translucent** | This is the most expensive **BlendMode**. It can look like a Masked and Opaque but will have an additional effect on performance regardless of actually being transparent or not. |
+| **Additive** | Is a cheaper alternative to Translucent. It does not need to sort the pixels, they just get added to the scene. But that means it can only brighten. Black pixels on the material have no effect on the scene.|
+| **Modulated** | Is the second cheapest alternative to Translucent, but it multiplies its Pixels with the scene. That means it can only darken the scene.|
 [region:note]
-When using the **Modulate** blend mode make sure to disable **Separate Translucency** or your Material will not compile correctly.
+When using the Modulate blend mode make sure to disable **Separate Translucency** or your Material will not compile correctly.
   ![](T_Separate_Trans.png)
 [/region]
 
@@ -95,11 +95,11 @@ On the **Main Material Node** under the **Details Panel** you will find a sectio
 The settings contained here are for helping to lessen Material rendering costs on mobile devices.  
 | Property | Description |
 | -------- | ---------- |
-| bFullyRough | When enabled, saves a number of shader ALU instructions and one sampler.  This effectively disables Specular while maintaining overall reflective energy. |
-| bUseLightmapDirectionality | When disabled, lighting from light maps will be flat but cheaper. |
+| **Fully Rough** | When enabled, saves a number of shader ALU instructions and one sampler.  This effectively disables Specular while maintaining overall reflective energy. |
+| **Use Lightmap Directionality** | When disabled, lighting from light maps will be flat but cheaper. |
 
 ### Customized UVs
-**[Customized UVs](Engine/Rendering/Materials/CustomizedUVs)** allow you to perform certain shader operations in vertex shaders instead of pixel shaders to help avoid dependent texture fetches. 
+[Customized UVs](Engine/Rendering/Materials/CustomizedUVs) allow you to perform certain shader operations in vertex shaders instead of pixel shaders to help avoid dependent texture fetches. 
 This offers an excellent way to claw back much needed performance when doing simple things like UV manipulation to make a texture title more than it currently does. 
 Whenever possible make sure to do all UV manipulation using Customized UVs especially when developing content for the Gear VR Headset.
 
@@ -108,15 +108,15 @@ Customized UVs are currently only implemented for certain component types such a
 [/region]
 
 ### Reflections
-The Gear VR does not support Screen Space Reflections(SSR) which means you will only be able generate reflections that come from **[Reflection Capture Actors](Resources/ContentExamples/Reflections/1_4/)** that are placed in levels. 
+The Gear VR does not support Screen Space Reflections(SSR) which means you will only be able generate reflections that come from [Reflection Capture Actors](Resources/ContentExamples/Reflections/1_4/) that are placed in levels. 
 Also keep in mind that the reflections that are generated from the reflections actors will not be of very high fidelity due to hardware limitations. 
 
 ## Gear VR & Content Setup
-When creating Art content for your Gear VR projects there are few things that you need to keep in mind to ensure that your project can run as smooth as possible on the Gear VR Headset. 
-In the following section we will cover what you need to be concerned with when making Art content for the Gear VR Headset.
+When creating art content for your Gear VR projects there are few things that you need to keep in mind to ensure that your project can run as smooth as possible on the Gear VR Headset. 
+In the following section we will cover what you need to be concerned with when making art content for the Gear VR Headset.
 
 ### Content Hard limits
-Due to certain hardware restrictions on mobile devices the following restrictions **Must Be** adhered to when creating content. 
+Due to certain hardware restrictions on mobile devices the following restrictions must be adhered to when creating content. 
 Failing to follow these restrictions will result in your project now working.
 
 * All mesh types can have up to but not over 65,000 vertices due to the lack of 32 bit index support on mobile hardware.
@@ -125,35 +125,35 @@ Failing to follow these restrictions will result in your project now working.
 ### Mesh Considerations
 When creating Meshes to use with your Gear VR Project the following information needs to be kept in mind to ensure that what you are creating will work with the Gear VR Headset. 
 
-* Use as few **[Material IDs](Engine/Content/FBX/Materials)** on 3D meshes as possible to save on draw calls.
-* Make sure to make use of **[Static Mesh LODs](Engine/Content/FBX/StaticMeshes)** whenever possible.
-* Keep an eye on your 3D meshes polygon count and try and keep it as **[Low Poly](https://en.wikipedia.org/wiki/Low_poly)** as possible.
-* Make sure that all static models have a 2nd UV set so they can make use of **[Light Mapping](Engine/Content/Types/StaticMeshes/LightmapUnwrapping)**.
-* Whenever possible try and make use if **[Imposter Sprites](Engine/Content/Tools/RenderToTextureTools/3)** to replace far away Static Meshes as they are much cheaper to render.
+* Use as few [Material IDs](Engine/Content/FBX/Materials) on 3D meshes as possible to save on draw calls.
+* Make sure to make use of [Static Mesh LODs](Engine/Content/FBX/StaticMeshes) whenever possible.
+* Keep an eye on your 3D meshes polygon count and try and keep it as [Low Poly](https://en.wikipedia.org/wiki/Low_poly) as possible.
+* Make sure that all static models have a 2nd UV set so they can make use of [Light Mapping](Engine/Content/Types/StaticMeshes/LightmapUnwrapping).
+* Whenever possible try and make use if [Imposter Sprites](Engine/Content/Tools/RenderToTextureTools/3) to replace far away Static Meshes as they are much cheaper to render.
 
 ### Texture Considerations
 When creating Textures to use with your Gear VR Project the following information needs to be kept in mind to ensure that what you are creating will not be detrimental to your projects overall performance. 
 
-* Textures use the **[ECT2](http://en.wikipedia.org/wiki/Ericsson_Texture_Compression)** Texture compression format.
-* Textures **can not** be bigger than 2048 pixels in either the X or Y dimension.
-* Textures sizes **MUST BE** in a power of 2 (2, 4, 8, 16, 32, 64, 128, 256, 512, 1024, 2048)
+* Textures use the [ECT2](http://en.wikipedia.org/wiki/Ericsson_Texture_Compression) Texture compression format.
+* Textures can not be bigger than 2048 pixels in either the X or Y dimension.
+* Textures sizes must be in a power of 2 (2, 4, 8, 16, 32, 64, 128, 256, 512, 1024, 2048)
 * Loading up to 512 MB of Textures into memory is possible with the Gear VR however it is recommended that you do not load more that 128 MB of Textures into memory.
 * Use square Textures whenever possible as them make a more efficient use of memory. 
-* Wherever possible try and use **[Texture Atlasing](http://en.wikipedia.org/wiki/Texture_atlas)** to reduce the number of Textures needed in memory.
-* **[RGB Texture Masking / Packing](Engine/Rendering/Materials/HowTo/Masking)** is a great way to cut down on unique Textures in memory.
+* Wherever possible try and use [Texture Atlasing](http://en.wikipedia.org/wiki/Texture_atlas) to reduce the number of Textures needed in memory.
+* [RGB Texture Masking / Packing](Engine/Rendering/Materials/HowTo/Masking) is a great way to cut down on unique Textures in memory.
 * Only **TC_Default** and **TC_NormalMap** compression settings should be used as they are the only two that will work on mobile platforms. 
 
 ### Level Considerations
 Getting the maximum performance out of your levels on the Gear VR is a delicate balancing act that will require time to get just right. 
 The following suggestions / things to look out for can help you quickly identify areas of your levels that might cause performance issues when viewed on the Gear VR Headset.
 
-* Always make sure that the level lighting has been **[built](Engine/Rendering/LightingAndShadows/Lightmass)** to ensure you get the best performance possible. If you see a red **Lighting needs to be rebuilt** error message on the screen, this means that your lighting needs to be built.
-* Make sure that all the lights in you level are set to **[Static](Engine/Rendering/LightingAndShadows/LightMobility/StaticLights)** as the other lighting options can be very expensive to render.
-* Check to see if your level is making use of **[Precomputed Visibility](Engine/Actors/Volumes)** and if not make sure to add one to the level and re-build lighting to enable it in the level's World Settings. 
-* Draw calls for the entire scene should not be greater than 100 for any view. This can be seen with **[stat openglrhi](Engine/Performance/StatCommands)** on the device, **[stat d3d11rhi](Engine/Performance/StatCommands)** in ES2 previewer on PC.
-* The total triangle count of the level should not be greater than 100,000 for any view. This can be seen with **[stat openglrhi](Engine/Performance/StatCommands)** on the device, **[stat d3d11rhi](Engine/Performance/StatCommands)** in ES2 previewer on PC.
+* Always make sure that the level lighting has been [built](Engine/Rendering/LightingAndShadows/Lightmass) to ensure you get the best performance possible. If you see a red **Lighting needs to be rebuilt** error message on the screen, this means that your lighting needs to be built.
+* Make sure that all the lights in you level are set to [Static](Engine/Rendering/LightingAndShadows/LightMobility/StaticLights) as the other lighting options can be very expensive to render.
+* Check to see if your level is making use of [Precomputed Visibility](Engine/Actors/Volumes) and if not make sure to add one to the level and re-build lighting to enable it in the level's World Settings. 
+* Draw calls for the entire scene should not be greater than 100 for any view. This can be seen with [stat openglrhi](Engine/Performance/StatCommands) on the device, [stat d3d11rhi](Engine/Performance/StatCommands) in ES2 previewer on PC.
+* The total triangle count of the level should not be greater than 100,000 for any view. This can be seen with [stat openglrhi](Engine/Performance/StatCommands) on the device, [stat d3d11rhi](Engine/Performance/StatCommands) in ES2 previewer on PC.
 * Be very careful in how Translucent and Opacity Masked Materials are used in the level as they are very expensive to render.
-* Make sure to **[Performance Profile](Engine/Performance)** your level often to make sure you catch any performance issues sooner rather than later.
+* Make sure to [Performance Profile](Engine/Performance) your level often to make sure you catch any performance issues sooner rather than later.
 
 
 
