#UE4doc #docs:
Updated Engine Version

diff --git a/Engine/Documentation/Source/Platforms/GearVR/QuickStart/GearVRQuickStart.INT.udn b/Engine/Documentation/Source/Platforms/GearVR/QuickStart/GearVRQuickStart.INT.udn
--- a/Engine/Documentation/Source/Platforms/GearVR/QuickStart/GearVRQuickStart.INT.udn
+++ b/Engine/Documentation/Source/Platforms/GearVR/QuickStart/GearVRQuickStart.INT.udn
@@ -5,7 +5,7 @@ Description:Instructions on how to quickly setup and deploy projects for use wit
 Navigation:platform
 Platform:Gear VR
 Skilllevel: Intermediate
-Version: 4.8
+Version: 4.9
 
 The Gear VR Quick Start will walk you through how to setup your development PC and Samsung Galaxy Note 4 Smartphone to work with the Samsung Gear VR headset.
 
