#UE4doc #docs:
Updated Engine Version
Changed font settings for multiple terms
Fixed capitalization when it comes to file extensions

diff --git a/Engine/Documentation/Source/Platforms/GearVR/QuickStart/3/GearVRQuickStart_3.INT.udn b/Engine/Documentation/Source/Platforms/GearVR/QuickStart/3/GearVRQuickStart_3.INT.udn
--- a/Engine/Documentation/Source/Platforms/GearVR/QuickStart/3/GearVRQuickStart_3.INT.udn
+++ b/Engine/Documentation/Source/Platforms/GearVR/QuickStart/3/GearVRQuickStart_3.INT.udn
@@ -44,18 +44,18 @@ Version: 4.8
 
 	![](GVR_Deployed_Files.png)
 
-1. Next, run the **.BAT** file which is located in the folder that the build was packaged to by **Double - Clicking** on it with your **Left Mouse Button**.
-	Once you run the .BAT file, a command line window will be displayed that shows the current progress of the deployment to your Smartphone.
+1. Next, run the .bat file which is located in the folder that the build was packaged to by **Double-clicking** on it with your **Left Mouse Button**.
+	Once you run the .bat file, a command line window will be displayed that shows the current progress of the deployment to your Smartphone.
 	When your project has been successfully deployed to your Smartphone, the command line window will shut itself down.
 	
 	![](GVR_Run_Bat_File.png)
 
-1. Once you run the .BAT file, a command line window will be displayed that will show you the deployment progress as well as any errors that may have occurred. 
+1. Once you run the .bat file, a command line window will be displayed that will show you the deployment progress as well as any errors that may have occurred. 
 
 	![](GVR_Deployment_Window.png)
 
 	[region:note]
-	The command line window that is created when running the .BAT file will close down on it's own when it is completed. Closing this window manually could cause an issue with your deployment and should be avoided. 
+	The command line window that is created when running the .bat file will close down on it's own when it is completed. Closing this window manually could cause an issue with your deployment and should be avoided. 
 	[/region]
 
 1. Finally, after the command line window has closed, check your Smartphone for your newly created application which should use the same name used for your UE4 C++ project that it was created from.
