UEDOC-1864

Images have been updated with the Silver guy

diff --git a/Engine/Documentation/Source/Engine/Components/Navigation/NavigationComponents.INT.udn b/Engine/Documentation/Source/Engine/Components/Navigation/NavigationComponents.INT.udn
--- a/Engine/Documentation/Source/Engine/Components/Navigation/NavigationComponents.INT.udn
+++ b/Engine/Documentation/Source/Engine/Components/Navigation/NavigationComponents.INT.udn
@@ -20,11 +20,11 @@ These **Area Classes** can define basic things like **cost** to enter an area or
 
 For example, the red area has a very high cost to move through, but the Pawn has no other choice but to move through it.:
 
-![](throughRed.png)(convert:false)
+![](AroundTheRed.jpg)
 
  But, if you remove a wall:
 
-![](aroundRed.png)(convert:false)
+![](ThroughTheRed.jpg)
 
 Then the pawn will avoid the red area even though it means taking a much longer route.
 
