#UE4doc #docs - Updated images and doc page.

diff --git a/Engine/Documentation/Source/Engine/Rendering/Materials/HowTo/ColoredTransluscentShadows/ColoredTransluscentShadows.INT.udn b/Engine/Documentation/Source/Engine/Rendering/Materials/HowTo/ColoredTransluscentShadows/ColoredTransluscentShadows.INT.udn
--- a/Engine/Documentation/Source/Engine/Rendering/Materials/HowTo/ColoredTransluscentShadows/ColoredTransluscentShadows.INT.udn
+++ b/Engine/Documentation/Source/Engine/Rendering/Materials/HowTo/ColoredTransluscentShadows/ColoredTransluscentShadows.INT.udn
@@ -11,7 +11,7 @@ Related: Engine/Rendering/LightingAndShadows/LitTranslucency
 
 
 
-In this "How To," we intend to show you how you can setup and use **Translucent Shadow Color** to create a material that can cast a colored shadow. 
+In this "How-to," we intend to show you how you can setup and use **Translucent Shadow Color** to create a material that can cast a colored shadow. 
 There are many applications where this could be useful to inject some color into your scene, but a common use would be for something like a stained glass window.
  
  **Example:**
@@ -22,7 +22,7 @@ There are many applications where this could be useful to inject some color into
 ## Translucent Shadow Color
 
 Translucent Shadow Color is the term used to describe how shadows will pick up colors of objects that they pass through also known as **Transmission**. 
-The amount of light that passes through the material is determined is directly determined by its opacity value between 0 and 1 and how much light is being cast onto the material. 
+The amount of light that passes through the material is determined by its opacity value between 0 and 1 and how much light is being cast onto the material. 
 For example, if the opacity was set to 0, the light transmission for a colored shadow from the material would be none. If the opacity value were set to 1, then it would be fully opaque. 
 
 
@@ -45,42 +45,42 @@ Notes
 
 ![](TCS_Lit_Masked.PNG)
 
-1. Select the Material Inputs and change the following:
+1. Select the **Material Attributes** node and change the following in its **Details** panel:
 * Blend Mode: Translucent
 * (Optional) Two Sided: Checked 
 * Translucency Lighting Model: TLM_Surface
-1. Plug the alpha channel of the Texture sample into Opacity to mask parts of material.
+1. Plug the **Alpha** channel of the Texture Sample into **Opacity** slot to mask parts of material.
 
 #### Custom Opacity
 
 ![](TCS_Lit_CustomOpacity.PNG)(w:850)
 
-1. Select the Material Inputs and change the following:
+1. Select the **Material Attributes** and change the following in its **Details** panel:
 * Blend Mode: Translucent
 * (Optional) Two Sided: Checked 
 * Translucency Lighting Model: TLM_Surface
-1. Create a Multiply Node and plug into Opacity on the Material Inputs
-1. Plug the alpha of the Texture Sample into pin A
-1. Create a Constant value and set the value to anything below 1.  Zero being fully transparent and 1 being fully opaque. 
+1. Create a **Multiply** node and plug into **Opacity** on the Material Attriubutes node.
+1. Plug the **Alpha** of the Texture Sample into pin **A**.
+1. Create a **Constant** value and set the value to anything below 1.  With 0 being fully transparent and 1 being fully opaque. 
 ![](TCS_LIT_CO_constantValue.PNG)
 
 [REGION:note]
-Notes 
-* Two sided is a good choice to use if you intend the player to see both sides of the mesh with the material. 
+
+* **Two Sided** is a good choice to use if you intend the player to see both sides of the mesh with the material. 
 * When **not** using Two Sided, the light must cast onto the visible side of the mesh with the material.  
-* To mask part of the image, you will need to have an alpha channel with your texture.
+* To mask part of the image, you will need to have an Alpha channel with your texture.
 [/REGION]
 
 ### Unlit Material
 
 ![](TCS_Unlit_Mat.PNG)
 
-1. Select the Material Inputs and change the following: 
+1. Select the **Material Attributes** node and change the following in its **Details** panel: 
 * Blend Mode: Translucent
 * Shading Model: Unlit
 * (Optional) Two Sided: Checked
 * Translucency Lighting Model: TLM_Surface
-1. Plug the Texture Sample into Emissive
+1. Plug the Texture Sample into **Emissive**.
 
 ## Colored Translucent Shadows in Action
 
@@ -99,7 +99,7 @@ Notes
 
 ## Shadow Sharpness
 
-Shadow sharpness can be affected by several different factors, including the lightmap resolution of the mesh receiving the **Translucent Shadow Color**, the source angle of the light, and the quality of the texture sample. 
+Shadow sharpness can be affected by several different factors, including the lightmap resolution of the mesh receiving the transmitted shadow color, the source angle of the light, and the quality of the texture sample. 
 
 ### Lightmap Resolution
 
@@ -117,7 +117,7 @@ Shadow sharpness can be affected by several different factors, including the lig
 
 ##Conclusion
 
-The most important thing to remember about colored shadows is that it only supports **Static Lights**. 
+The most important thing to remember about colored shadows is that it **ONLY** supports **Static Lights**. 
 Also, you will want to make sure that you are selecting the best blending mode (Translucent, Additive, and Modulate) that fits your needs for your design. 
 There are a number of applications for colored shadows and each blend mode can be used to work within your needs for your projects. 
 
