#UE4doc #docs Removed Licensee regions, deprecated some outdated information, and re-organized documentation.

diff --git a/Engine/Documentation/Source/Programming/Development/CompilingProjects/CompilingProjects.INT.udn b/Engine/Documentation/Source/Programming/Development/CompilingProjects/CompilingProjects.INT.udn
--- a/Engine/Documentation/Source/Programming/Development/CompilingProjects/CompilingProjects.INT.udn
+++ b/Engine/Documentation/Source/Programming/Development/CompilingProjects/CompilingProjects.INT.udn
@@ -50,56 +50,50 @@ the **Editor** target keyword, but if you are building an executable version of
 
 The build configuration can be set in the Visual Studio toolbar.
 
-[PUBLISH:Rocket]
+### Project - Configuration Menu
+
 ![Project - Configuration Menu](rocket_configuration_menu.png)
-[/PUBLISH:Rocket]
-[PUBLISH:Licensee]
-![UE4 Solution - Configuration Menu](configuration_menu.png)
-[/PUBLISH:Licensee]
+<!--
+### Solution - Configuration Menu
 
+![UE4 Solution - Configuration Menu](configuration_menu.png)
+-->
 ### Setting the Solution Platform
 
 The solution platform can be set in the Visual Studio toolbar.
-
-[PUBLISH:Rocket]
+<!--
 UE4 currently supports the following platforms:
 
 * **Windows 64-bit**
 
-
 ![Project - Platform Menu](rocket_platform_menu.png)
-[/PUBLISH:Rocket]
-
-[PUBLISH:Licensee]
+![UE4 Solution - Platform Menu](platform_menu.png)
+-->
 UE4 currently supports the following platforms:
 
 * **Windows 32-bit**
 * **Windows 64-bit**
 
-
-When working with Unreal Engine 4, you will generally use the **Win64** platform. This is the only one included by default when generating project files; the [](Programming/UnrealBuildSystem/ProjectFileGenerator) page
+When working with Unreal Engine 4, you will typically use the **Win64** platform. This is the only one included by default when generating project files; the [](Programming/UnrealBuildSystem/ProjectFileGenerator) page
 has instructions for generating project files for additional platforms.
 
-![UE4 Solution - Platform Menu](platform_menu.png)
-[/PUBLISH:Licensee]
-
 ### Building the Project
 
-[PUBLISH:Rocket]
 [REGION:note]
 Please make sure you are running Visual Studio 2013 Express, Community, or Professional for Windows Desktop installed before proceeding.  If you are using Mac, make sure to have Xcode 6.1 or higher installed.
 [/REGION]
+
+<!--
+
 [REGION:tip]
 When building in Visual Studio, you are only compiling the game project, not the Unreal Editor.  
 [/REGION]
-[/PUBLISH:Rocket]
-[PUBLISH:Licensee]
 [REGION:tip]
 When building in Visual Studio, you are compiling the engine in addition to the game project.  
 [/REGION]
-[/PUBLISH:Licensee]
 
-[PUBLISH:Licensee]
+-->
+
 1.	Set the **Solution Configuration** to the configuration you want to build. In this example, it is set to **Development Editor**. Refer to the [Build Configuration](#BuildConfiguration) section for descriptions of each available configuration.
 
 	![Project - Set Configuration to Development](config_development_menu.png)
@@ -108,8 +102,6 @@ When building in Visual Studio, you are compiling the engine in addition to the
 
 	![](rebuild_project_outside.png)
 
-[/PUBLISH:Licensee]
-[PUBLISH:Rocket]
 1.	Set the **Solution Configuration** to the configuration you want to build. In this example, it is set to **Development**. Refer to the [Build Configuration](#BuildConfiguration) section for descriptions of each available configuration.
 
 	![](rocket_config_development_menu.png)
@@ -117,39 +109,30 @@ When building in Visual Studio, you are compiling the engine in addition to the
 1. **Right-click** your game project and choose **Rebuild** to recompile.
 
 	![](rebuild_project_rocket.png)
-[/PUBLISH:Rocket]
 
 You can now [run the engine](GettingStarted/RunningUnrealEngine) with your compiled project.  
 
-
 [EXCERPT:MatchEXE]
-[PUBLISH:Licensee]
 [REGION:note]
 When running the engine, it is important to use the Unreal Engine executable that matches the build configuration you rebuilt your project in.  For example, if you compiled your project in the **DebugGame Uncooked**
 build configuration, you would run the `UE4-Win64-DebugGame.exe` executable with your game information as an argument.  More information on the binary naming convention can be found on the [Building Unreal Engine](Programming/Development/BuildingUnrealEngine) page.
 [/REGION:note]
-[/PUBLISH:Licensee]	
 
-[PUBLISH:Rocket]
 [REGION:note]
 When running Unreal, it is important to add the `-game` flag if you rebuilt your project in any **Uncooked** configuration, and the `-debug` flag if you rebuilt your project in any **Debug** configuration.  
 [/REGION:note]
-[/PUBLISH:Rocket]	
 [/EXCERPT:MatchEXE]
 
-[PUBLISH:Rocket]
 ### Visual Studio Known Issues
 
-
 | Issue | Solution |
 | ----- | -------- |
 | "Project is out of date" message always appears | Visual Studio thinks the project is out of date when it really is up to date. You can safely suppress this message by checking the **Do not show this dialog again** box and pressing **No**. |
 | **No Debugging Information** window appears when using the _Debug_ configuration | The reason for this is that the UE4editor.exe was compiled using the _Development_ configuration. Your game project will still be debuggable when compiled using the _Debug_ configuration despite this warning. You can safely suppress this message by checking the **Do not show this dialog again** box and pressing **Yes**. |
-[/PUBLISH:Rocket]
 [/EXCERPT:VS]
 
 [EXCERPT:Xcode]
-[PUBLISH:Rocket]
+
 ## Building with Xcode
 
 [REGION:tip]
@@ -182,7 +165,7 @@ Xcode's LLDB debugger needs to be configured to correctly handle breakpoints in
 	settings set target.inline-breakpoint-strategy always
 [/EXCERPT:XcodeDebug]
 -->
-[/PUBLISH:Rocket]
+
 [/EXCERPT:Xcode]
 
 <!-- The following was deprecated per LaurenR 29.7.15
