#UE4doc #docs Completed first draft. Availibility: Docs.

diff --git a/Engine/Documentation/Source/Programming/Development/BuildConfigurations/BuildConfigurations.INT.udn b/Engine/Documentation/Source/Programming/Development/BuildConfigurations/BuildConfigurations.INT.udn
--- a/Engine/Documentation/Source/Programming/Development/BuildConfigurations/BuildConfigurations.INT.udn
+++ b/Engine/Documentation/Source/Programming/Development/BuildConfigurations/BuildConfigurations.INT.udn
@@ -3,25 +3,26 @@ Title: Build Configurations Reference
 Description: Reference guide for the different build configurations available to developers.
 Crumbs: %ROOT%, Programming, Programming/Development
 Prereq: GettingStarted/RecommendedSpecifications
-Prereq: GettingStarted/DownloadingUnrealEngine
-Prereq: Programming/Development/BuildingUnrealEngine
+Related: GettingStarted/DownloadingUnrealEngine
+Related: Programming/Development/BuildingUnrealEngine
+Related: Programming/Development/CompilingProjects
 Related: Support/Builds
 SkillLevel: Advanced
 Version: 4.8
 
 [TOC (start:2 end:3)]
 
-## Build Configurations
+## Build Configuration Descriptions
 
 Unreal Engine 4 (UE4) uses a custom building method via the UnrealBuildTool. This tool handles all the complex aspects of compiling the engine. 
 
 [INCLUDE:Programming/Development/CompilingProjects#winbuildconfig]
 
-### UE4 Build Configuration - Source
+### Build Configuration for UE4 Solution
 
 (#bookmark1)
 
-The following build configurations are available when building your project with the source version of UE4:
+When compiling a UE4 solution, you are compiling our engine's source code together with your project's source code. The following build configurations are available when building your project this way:
 
 |             | Debug    | DebugGame | Development | Shipping | Test     |
 | ----------- | -----    | --------- | ----------- | -------- | ----     |
@@ -30,9 +31,9 @@ The following build configurations are available when building your project with
 | **Client ** | &#x2713; | &#x2713;  | &#x2713;    | &#x2713; | &#x2713; |
 | **Server ** | &#x2713; | &#x2713;  | &#x2713;    | &#x2713; | &#x2713; |
 													   
-### UE4 Build Configuration - Binary
+### Build Configuration for UE4 Project
 
-The following build configurations are available when building your project with the binary version of UE4:
+When compiling a UE4 project, you are only compiling your project's source code. The following build configurations are available when building your project this way:
 
 %Globals:osselect_linux%
 
@@ -85,7 +86,7 @@ The following build configurations are available when building your project with
 	[PARAM:content]
 	    [REGION:warning]
         We currently don't provide a binary installer of Unreal Engine for Linux.  
-		Linux users should refer to build configurations listed in the ['Unreal Engine Source Version' section of this document](#bookmark1).  
+		Linux users should refer to the [Build Configuration for UE4 Solution](#bookmark1) section.  
         _Refer to our [](GettingStarted/DownloadingUnrealEngine) documentation to learn how to build Unreal Engine from source._
         [/REGION]
 	[/PARAM]
