#UE4doc #docs Minor edit to page description. Availability: Docs.

diff --git a/Engine/Documentation/Source/Programming/Development/BuildConfigurations/BuildConfigurations.INT.udn b/Engine/Documentation/Source/Programming/Development/BuildConfigurations/BuildConfigurations.INT.udn
--- a/Engine/Documentation/Source/Programming/Development/BuildConfigurations/BuildConfigurations.INT.udn
+++ b/Engine/Documentation/Source/Programming/Development/BuildConfigurations/BuildConfigurations.INT.udn
@@ -1,6 +1,6 @@
 Availability: Docs
 Title: Build Configurations Reference
-Description: Reference guide for the different build configurations available to developers.
+Description: Reference guide for solution and project build configurations.
 Crumbs: %ROOT%, Programming, Programming/Development
 Prereq: GettingStarted/RecommendedSpecifications
 Related: GettingStarted/DownloadingUnrealEngine
