#UE4doc #docs Minor edit.

diff --git a/Engine/Documentation/Source/Programming/Development/BuildConfigurations/BuildConfigurations.INT.udn b/Engine/Documentation/Source/Programming/Development/BuildConfigurations/BuildConfigurations.INT.udn
--- a/Engine/Documentation/Source/Programming/Development/BuildConfigurations/BuildConfigurations.INT.udn
+++ b/Engine/Documentation/Source/Programming/Development/BuildConfigurations/BuildConfigurations.INT.udn
@@ -7,7 +7,6 @@ Related: GettingStarted/DownloadingUnrealEngine
 Related: Programming/Development/BuildingUnrealEngine
 Related: Programming/Development/CompilingProjects
 Related: Support/Builds
-SkillLevel: Advanced
 Version: 4.8
 
 [TOC (start:2 end:3)]
