UEDOC-1902

Images have been replaced with Updated images

diff --git a/Engine/Documentation/Source/Engine/Rendering/Materials/LayeredMaterials/LayeredMaterials.INT.udn b/Engine/Documentation/Source/Engine/Rendering/Materials/LayeredMaterials/LayeredMaterials.INT.udn
--- a/Engine/Documentation/Source/Engine/Rendering/Materials/LayeredMaterials/LayeredMaterials.INT.udn
+++ b/Engine/Documentation/Source/Engine/Rendering/Materials/LayeredMaterials/LayeredMaterials.INT.udn
@@ -42,7 +42,9 @@ One of the benefits to using Layered Materials are that they allow you to create
 Consider the complexity of the network below, which uses no material layers and blends a chrome and a copper effect, such as the one seen above:
 
 **Without Material Layers**
-![](WithoutLayers.png)(w:900)
+[REGION:lightbox]
+![](Rocket_Material_Beforelayers.jpg)(w:900)
+[/REGION]
 
 
 However, with a Layered Material, each different Material type is contained within its own node, making blending much simpler and far easier for an artist to edit and debug. By using the Make Material Attributes and Break Material Attributes nodes, all of the different Material Layers can be connected directly without having to worry about wiring up individual properties.
@@ -50,7 +52,9 @@ However, with a Layered Material, each different Material type is contained with
 The network below works the same way as the one above, but with the chrome and copper effects modularized into their own Material Layers:
 
 **With Material Layers**
-![](WithLayers.png)(w:900)
+[REGION:lightbox]
+![](MatieralLayering.jpg)
+[/REGION]
 
 Another benefit of Material Layers is that since they utilize Material Functions, they are reusable. This allows you to set up a library of Material archetypes, or Materials that define a basic real-world surface type. For example, you can create layers that represent generic surfaces such as aluminum, steel, leather, plastic, rubber, and so on. You can then use a Layered Material to blend between them. This can be very useful when wanting to create a detailed object - such as a character - without having to create a massive set of materials to apply separately to the surface.
 
