#UE4doc #docs Moved Epic Account Creation content from 'InstallingUnrealEngine' documentation to 'AccountSetup' documentation.

diff --git a/Engine/Documentation/Source/GettingStarted/Installation/InstallingUnrealEngine.INT.udn b/Engine/Documentation/Source/GettingStarted/Installation/InstallingUnrealEngine.INT.udn
--- a/Engine/Documentation/Source/GettingStarted/Installation/InstallingUnrealEngine.INT.udn
+++ b/Engine/Documentation/Source/GettingStarted/Installation/InstallingUnrealEngine.INT.udn
@@ -18,7 +18,7 @@ Linux users; however, can compile Unreal Engine from source by referring to our
 
 By the end of this tutorial, you'll be able to:
 
-*	[Create an Epic Games account _(if you don't already have one)_](#Bookmark1)
+*	[Create an Epic Games Account _(if you don't already have one)_](#Bookmark1)
 *	[Download and Run the Installer (Epic Games Launcher Set-Up Program)](#Bookmark2)
 *	[Sign into the Epic Games Launcher](#Bookmark3)
 *	[Install Unreal Engine](#Bookmark4)
@@ -28,37 +28,13 @@ By the end of this tutorial, you'll be able to:
 
 _The Epic Games Launcher enables our subscribers to conveniently download and install updates for UE4 (example pictured above)._
 
-## Creating an Epic Games Account
-
-(#Bookmark1)
-
 [REGION:warning]
 If you already have an Epic Games account, skip ahead to learn about [Downloading and Running the Installer](#Bookmark2).
 [/REGION]
 
-In order to run the Epic Games Launcher, you need to have a valid Epic Games account.
-
-1.	**Navigate** to [UnrealEngine.com](https://www.unrealengine.com/what-is-unreal-engine-4).
-
-1.	**Click** the **Get Unreal** button (pictured below).
-
-	![](UnrealWebsite.png)
-
-1.	**Fill-out** our **'Join the Community'** Form and **click** the **Sign Up** button (pictured below).
-	
-	![](SignupScreen.png)
-
-	[REGION:note]
-	Make sure to read through our [terms of service](http://epicgames.com/tos) prior to signing up.
-	[/REGION]
-
-1.	**Read** our [Unreal Engine End User License Agreement](https://www.unrealengine.com/eula) (EULA) prior to **clicking** the **Accept** button (pictured below).
-
-	![](EULA.png)
+(#Bookmark1)
 
-	[REGION:note]
-	Refer to our [EULA FAQ](https://www.unrealengine.com/faq#legal) if you have additional questions about the EULA's terms and conditions.
-	[/REGION]
+[INCLUDE:GettingStarted/AccountSetup#epicAccount]
 
 ## Downloading and Running the Installer
 
