#UE4doc #docs Completed first draft of 'AccountSetup' documentation.

diff --git a/Engine/Documentation/Source/GettingStarted/AccountSetup/AccountSetup.INT.udn b/Engine/Documentation/Source/GettingStarted/AccountSetup/AccountSetup.INT.udn
new file mode 100644
--- /dev/null
+++ b/Engine/Documentation/Source/GettingStarted/AccountSetup/AccountSetup.INT.udn
@@ -0,0 +1,38 @@
+Availability:Docs
+Title:Setting Up Your Epic Games Account
+Crumbs: %ROOT%, GettingStarted
+Description:How to set up an Epic Games account.
+Related: GettingStarted/RecommendedSpecifications
+Related: Support/Builds
+Version: 4.8
+SkillLevel: Beginner
+
+[EXCERPT:epicAccount]
+
+## Creating an Epic Games Account
+
+In order to run the Epic Games Launcher, you need to have a valid Epic Games account.
+
+1.	**Navigate** to [UnrealEngine.com](https://www.unrealengine.com/what-is-unreal-engine-4).
+
+1.	**Click** the **Get Unreal** button (pictured below).
+
+	![](UnrealWebsite.png)
+
+1.	**Fill-out** our **'Join the Community'** form and **click** the **Sign Up** button (pictured below).
+
+	![](SignupScreen.png)
+
+	[REGION:note]
+	Make sure to read through our [terms of service](http://epicgames.com/tos) prior to signing up.
+	[/REGION]
+
+1.	**Read** our [Unreal Engine End User License Agreement](https://www.unrealengine.com/eula) (EULA) prior to **clicking** the **Accept** button (pictured below).
+
+	![](EULA.png)
+
+	[REGION:note]
+	Refer to our [EULA FAQ](https://www.unrealengine.com/faq#legal) if you have additional questions about the EULA's terms and conditions.
+	[/REGION]
+
+[/EXCERPT:epicAccount]
\ No newline at end of file
