#UE4doc #docs Changed Availability from Docs to NoPublish.

diff --git a/Engine/Documentation/Source/GettingStarted/AccountSetup/AccountSetup.INT.udn b/Engine/Documentation/Source/GettingStarted/AccountSetup/AccountSetup.INT.udn
--- a/Engine/Documentation/Source/GettingStarted/AccountSetup/AccountSetup.INT.udn
+++ b/Engine/Documentation/Source/GettingStarted/AccountSetup/AccountSetup.INT.udn
@@ -1,4 +1,4 @@
-Availability:Docs
+Availability:NoPublish
 Title:Setting Up Your Epic Games Account
 Crumbs: %ROOT%, GettingStarted
 Description:How to set up an Epic Games account.
