#UE4doc #docs Starting 'FirstPersonShooter' C++ tutorial.

diff --git a/Engine/Documentation/Source/Programming/Tutorials/FirstPersonShooter/FirstPersonShooter.INT.udn b/Engine/Documentation/Source/Programming/Tutorials/FirstPersonShooter/FirstPersonShooter.INT.udn
new file mode 100644
--- /dev/null
+++ b/Engine/Documentation/Source/Programming/Tutorials/FirstPersonShooter/FirstPersonShooter.INT.udn
@@ -0,0 +1,7 @@
+Availability: Docs
+Title: First Person Shooter C++ Tutorial
+Crumbs: %ROOT%, GettingStarted, Programming/Tutorials
+Description: How to develop a First Person Shooter.
+SkillLevel: Intermediate
+Version: 4.8
+
