TypoCorrection

diff --git a/Engine/Documentation/Source/Platforms/HTML5/GettingStarted/HTML5GettingStarted.INT.udn b/Engine/Documentation/Source/Platforms/HTML5/GettingStarted/HTML5GettingStarted.INT.udn
--- a/Engine/Documentation/Source/Platforms/HTML5/GettingStarted/HTML5GettingStarted.INT.udn
+++ b/Engine/Documentation/Source/Platforms/HTML5/GettingStarted/HTML5GettingStarted.INT.udn
@@ -1,4 +1,4 @@
-Availability:Public
+﻿Availability:Public
 Title: Getting Started: Developing HTML5 Projects
 Crumbs:%ROOT%, Platforms, Platforms/HTML5
 Description:Setting up for development for the HTML5 platform.
@@ -89,14 +89,14 @@ If you need to specify different paths for different OSs you can do the followin
 
 	\[HTML5SDKPaths\]
 	Windows="C:/html5sdk/path/on/windows"
-	Mac="/html5sdk/path/on/mac”
+	Mac="/html5sdk/path/on/mac"
 
 Make sure you have major browsers installed and links in `HTML5Engine.ini` under `\[HTML5DevicesWindows\]` section point to the right browser binaries. 
 
 For example:
 
 	\[HTML5DevicesWindows\]
-	Firefox="C:/Program Files (x86)/Mozilla Firefox/firefox.exe”
+	Firefox="C:/Program Files (x86)/Mozilla Firefox/firefox.exe"
 	Chrome(32bit)="C:/Program Files (x86)/Google/Chrome/Application/chrome.exe"
 	Chrome(64bit)="C:/Program Files/Google/Chrome/Application/chrome.exe"
 
