#UE4doc #docs - Add images and separated methods of import for Destructible and Cloth assets. Also setup import instructions for APEX cloth assets.

diff --git a/Engine/Documentation/Source/Engine/Physics/Apex/ApexTypes/ApexTypes.INT.udn b/Engine/Documentation/Source/Engine/Physics/Apex/ApexTypes/ApexTypes.INT.udn
--- a/Engine/Documentation/Source/Engine/Physics/Apex/ApexTypes/ApexTypes.INT.udn
+++ b/Engine/Documentation/Source/Engine/Physics/Apex/ApexTypes/ApexTypes.INT.udn
@@ -61,6 +61,8 @@ See the [APEX Physics Lab](Engine/Physics/Apex) documentation on creating and ex
 
 ## Import
 
+###Destructible Assets
+
 In the **Content Browser**:
 
 1. Click the ![New Asset](CB_button_NewAsset.png) button.
@@ -68,9 +70,19 @@ In the **Content Browser**:
 
 Or
 
-1. Select your .APX or .APB file(s) file in a Windows File Explorer window.
+1. Select your **.APX** or **.APB** file(s) file in a Windows File Explorer window.
 1. Click and drag it (them) into the Content Browser in Unreal Editor 4.
 
+###Cloth Assets
+
+1. Click the ![New Asset](CB_button_NewAsset.png) button.
+1. Select your Skeletal Mesh that uses the cloth asset.
+1. Open the **Skeletal Mesh**. 
+1. In the **Mesh Details** panel under the Clothing tab click the ![New Asset](CB_Button_AddAPEX.png) button.
+1. Select the **.APX** or **.APB** file(s) you wish to import. 
+1. Use the ![New Asset](CB_Button_clothingElement.png) drop down below each **Material Element** to select the corresponding cloth asset to assign to the material.
+
+
 ## Related
 
 [REGION:topics half]
