#UE4doc #docs - Updated images to mach latest version, updated options names that have changed, along with removing "lock actors to viewport" option since this has been removed.

diff --git a/Engine/Documentation/Source/Engine/UI/LevelEditor/Viewports/ViewportOptions/ViewportOptions.INT.udn b/Engine/Documentation/Source/Engine/UI/LevelEditor/Viewports/ViewportOptions/ViewportOptions.INT.udn
--- a/Engine/Documentation/Source/Engine/UI/LevelEditor/Viewports/ViewportOptions/ViewportOptions.INT.udn
+++ b/Engine/Documentation/Source/Engine/UI/LevelEditor/Viewports/ViewportOptions/ViewportOptions.INT.udn
@@ -58,14 +58,13 @@ The **Viewport Options** menu is accessed by clicking the ![](button_ViewportOpt
 |--------|-------------|
 | [REGION:tablesection]**Viewport Options**[/REGION] ||
 | **Realtime** | Toggles realtime rendering of the viewport. |
-| **Stats** | Toggles the ability to display statistics in the viewport. Enabling this also enables **Realtime**. |
+| **Show Stats** | Toggles the ability to display statistics in the viewport. Enabling this also enables **Realtime**. |
 | **Show FPS** | Toggles the display of frames per second (FPX) information in the viewport. Enabling this also enables **Realtime**. |
 | **Field of View**| Allows you to change the FOV of the viewport camera. This effectively zooms the lens of the camera in and out. |
 | **Far View Plane**| Allows you to choose the distance to use as your far view plane. Setting this to zero acts as an infinite far view plane. |
-| **Allow Matinee Preview** | Toggles the ability to preview Matinee cinematics in the viewport. |
+| **Toggles Cinematic Preview** | Toggles the ability to preview Matinee cinematics in the viewport. |
 | **Game View** | Toggles Game View which displays the scene as it appears in-game, i.e. Actor icons are not rendered, etc. |
 | **Immersive Mode** | Toggles Immersive mode for the viewport. When in Immersive mode, the viewport occupies the entire area of the Unreal Editor window. |
-| **Lock Viewport to Actor** | Locks an Actor to the viewport. See [Lock Viewport to Actor](Engine/UI/LevelEditor\Viewports\ActorLocking) for details.|
 | **Bookmarks** | Provides the ability to set camera locations and rotations as bookmarks, and provides access to jump to existing bookmarks. To create a bookmark, press **Ctrl** + any of the numbers 1-0 along the top of your keyboard. To recall a bookmark, simply press the number you chose for that bookmark. |
 |**Create Camera Here** | Creates a Camera Actor at the current location and rotation of the viewport.|
 |**High Resolution Screenshot...** | Opens the [High Resolution Screenshot](Engine/Basics/Screenshots) window.![](HighResScreenshotDialog.png) |
