#UE4doc #docs Added and completed Source Branch reference page.

diff --git a/Engine/Documentation/Source/GettingStarted/DownloadingUnrealEngine/SourceBranches/SourceBranches.INT.udn b/Engine/Documentation/Source/GettingStarted/DownloadingUnrealEngine/SourceBranches/SourceBranches.INT.udn
new file mode 100644
--- /dev/null
+++ b/Engine/Documentation/Source/GettingStarted/DownloadingUnrealEngine/SourceBranches/SourceBranches.INT.udn
@@ -0,0 +1,60 @@
+Availability: NoPublish
+Title: Source Branches Reference
+Crumbs: %ROOT%, GettingStarted, GettingStarted/DownloadingUnrealEngine
+Description: Reference guide for source code branch descriptions.
+Prereq: GettingStarted/RecommendedSpecifications
+Related: GettingStarted/DownloadingUnrealEngine
+Related: Programming/Development/BuildingUnrealEngine
+Related: Programming/Development/CompilingProjects
+Related: Support/Builds
+Version: 4.8
+
+## Source Code Branches
+
+Once you've accessed the repository, you'll notice that we've published UE4's source code in three rolling branches, namely:
+
+*	[Release Branch](#Bookmark1)
+*	[Promoted Branch](#Bookmark2)
+*	[Master Branch](#Bookmark3)
+
+### Release Branch
+
+(#Bookmark1)
+
+[REGION:raw]
+![](ReleaseBranch.png)
+[/REGION]
+
+The [release branch](https://github.com/EpicGames/UnrealEngine/tree/release) (pictured above) is extensively tested by our QA team and makes a great starting point for learning UE4 or making Unreal games. We work hard to stabilize this branch with the goal of releasing new builds every one to two months.
+
+### Promoted Branch
+
+(#Bookmark2)
+
+[REGION:raw]
+![](PromotedBranch.png)
+[/REGION]
+
+The [promoted branch](https://github.com/EpicGames/UnrealEngine/tree/promoted) (pictured above) is updated with builds for use by our designers and artists. We try to update this branch daily so that our subscribers have access to UE4's latest stable features. This branch is a great place to learn how feature development works in UE4.
+
+### Master Branch
+
+(#Bookmark3)
+
+[REGION:warning]
+Because the master branch contains cutting-edge technology, UE4 isn't guaranteed to compile from source.
+[/REGION]
+
+[REGION:raw]
+![](MasterBranch.png)
+[/REGION]
+
+The [master branch](https://github.com/EpicGames/UnrealEngine/tree/master) (pictured above) tracks [live changes](https://github.com/EpicGames/UnrealEngine/commits/master) (pictured below) being made by our engine team. Our experienced subscribers typically work from this branch when they want to work in tandem with UE4's developers.
+
+[REGION:raw]
+![](LiveChanges.png)
+[/REGION]
+
+[REGION:note]
+Other branches may appear from time-to-time as we work to stabilize new releases or hotfixes.
+[/REGION]
\ No newline at end of file
