#UE4doc #docs Fixed document Crumbs and chunked out Source Branch reference material.

diff --git a/Engine/Documentation/Source/GettingStarted/DownloadingUnrealEngine/DownloadingUnrealEngine4.INT.udn b/Engine/Documentation/Source/GettingStarted/DownloadingUnrealEngine/DownloadingUnrealEngine4.INT.udn
--- a/Engine/Documentation/Source/GettingStarted/DownloadingUnrealEngine/DownloadingUnrealEngine4.INT.udn
+++ b/Engine/Documentation/Source/GettingStarted/DownloadingUnrealEngine/DownloadingUnrealEngine4.INT.udn
@@ -1,6 +1,6 @@
 ﻿Availability:Public
 Title:Downloading Unreal Engine Source Code
-Crumbs: %ROOT%, Programming, Programming/Development
+Crumbs: %ROOT%, GettingStarted, GettingStarted/DownloadingUnrealEngine
 Description:Steps for connecting to the source code repository and downloading the latest build of Unreal Engine.
 Prereq: GettingStarted/RecommendedSpecifications
 Related: Programming/Development/BuildingUnrealEngine
@@ -34,55 +34,7 @@ Before you can access the [repository](https://github.com/EpicGames/UnrealEngine
 *	For instructions on associating your GitHub account with your Unreal Engine account, visit the [UE4 on GitHub](https://www.unrealengine.com/ue4-on-github) page.
 [/REGION]
 
-## Source Code Branches
-
-Once you've accessed the repository, you'll notice that we've published UE4's source code in three rolling branches, namely:
-
-*	[Release Branch](#Bookmark1)
-*	[Promoted Branch](#Bookmark2)
-*	[Master Branch](#Bookmark3)
-
-### Release Branch
-
-(#Bookmark1)
-
-[REGION:raw]
-![](ReleaseBranch.png)
-[/REGION]
-
-The [release branch](https://github.com/EpicGames/UnrealEngine/tree/release) (pictured above) is extensively tested by our QA team and makes a great starting point for learning UE4 or making Unreal games. We work hard to stabilize this branch with the goal of releasing new builds every one to two months.
-
-### Promoted Branch
-
-(#Bookmark2)
-
-[REGION:raw]
-![](PromotedBranch.png)
-[/REGION]
-
-The [promoted branch](https://github.com/EpicGames/UnrealEngine/tree/promoted) (pictured above) is updated with builds for use by our designers and artists. We try to update this branch daily so that our subscribers have access to UE4's latest stable features. This branch is a great place to learn how feature development works in UE4.
-
-### Master Branch
-
-(#Bookmark3)
-
-[REGION:warning]
-Because the master branch contains cutting-edge technology, UE4 isn't guaranteed to compile from source.
-[/REGION]
-
-[REGION:raw]
-![](MasterBranch.png)
-[/REGION]
-
-The [master branch](https://github.com/EpicGames/UnrealEngine/tree/master) (pictured above) tracks [live changes](https://github.com/EpicGames/UnrealEngine/commits/master) (pictured below) being made by our engine team. Our experienced subscribers typically work from this branch when they want to work in tandem with UE4's developers.
-
-[REGION:raw]
-![](LiveChanges.png)
-[/REGION]
-
-[REGION:note]
-Other branches may appear from time-to-time as we work to stabilize new releases or hotfixes.
-[/REGION]
+[INCLUDE:GettingStarted/DownloadingUnrealEngine/SourceBranches]
 
 [/PUBLISH:Rocket]
 
