#UE4doc #docs - updated images from UE3 to UE4. Added PSD file in the event these need to be edited in the future.

diff --git a/Engine/Documentation/Source/Engine/Rendering/Materials/Functions/Reference/Gradient/Gradient.INT.udn b/Engine/Documentation/Source/Engine/Rendering/Materials/Functions/Reference/Gradient/Gradient.INT.udn
--- a/Engine/Documentation/Source/Engine/Rendering/Materials/Functions/Reference/Gradient/Gradient.INT.udn
+++ b/Engine/Documentation/Source/Engine/Rendering/Materials/Functions/Reference/Gradient/Gradient.INT.udn
@@ -15,7 +15,7 @@ Below is a list of the gradient functions.
 
 ### Radial Gradient Exponential
 
-The RadialGradientExponential function uses UV channel 0 to produce a radial gradient, giving the user the ability to adjust the radius and offset the center point.
+The **RadialGradientExponential** function uses UV channel 0 to produce a radial gradient, giving the user the ability to adjust the radius and offset the center point.
 
 | Item | Description |
 | -------- | -------- |
@@ -43,7 +43,7 @@ The RadialGradientExponential function uses UV channel 0 to produce a radial gra
 
 ### Diamond Gradient
 
-The DiamondGradient function uses UV channel 0 to produce a radial gradient, giving the user the ability to adjust the falloff rate of the gradient.
+The **DiamondGradient** function uses UV channel 0 to produce a radial gradient, giving the user the ability to adjust the falloff rate of the gradient.
 
 | Item | Description |
 | -------- | -------- |
@@ -65,7 +65,7 @@ The DiamondGradient function uses UV channel 0 to produce a radial gradient, giv
 
 ### Linear Gradient
 
-The LinearGradient function uses UV channel 0 to generate a linear gradient in either the U or V direction, depending on which output is used.
+The **LinearGradient** function uses UV channel 0 to generate a linear gradient in either the U or V direction, depending on which output is used.
 
 
 | Item | Description |
@@ -87,7 +87,7 @@ The LinearGradient function uses UV channel 0 to generate a linear gradient in e
 
 ### Smooth Curve
 
-The SmoothCurve function takes in an existing texture channel or gradient and uses a procedural curve to control the transition from dark to light. The user can adjust the tangents of this curve to change the result.
+The **SmoothCurve** function takes in an existing texture channel or gradient and uses a procedural curve to control the transition from dark to light. The user can adjust the tangents of this curve to change the result.
 
 | Item | Description |
 | -------- | -------- |
@@ -110,7 +110,7 @@ The SmoothCurve function takes in an existing texture channel or gradient and us
 
 ### Value Step
 
-The ValueStep function takes in an existing texture channel or gradient and outputs a pure black and white image based on inputs from the user. The result is a mask that represents what portion of the gradient equals the value of the inputs.
+The **ValueStep** function takes in an existing texture channel or gradient and outputs a pure black and white image based on inputs from the user. The result is a mask that represents what portion of the gradient equals the value of the inputs.
 
 | Item | Description |
 | -------- | -------- |
