UEDOC-1023

New images have replaced older images

diff --git a/Engine/Documentation/Source/Engine/Physics/Constraints/DampingAndFriction/DampingAndFriction.INT.udn b/Engine/Documentation/Source/Engine/Physics/Constraints/DampingAndFriction/DampingAndFriction.INT.udn
--- a/Engine/Documentation/Source/Engine/Physics/Constraints/DampingAndFriction/DampingAndFriction.INT.udn
+++ b/Engine/Documentation/Source/Engine/Physics/Constraints/DampingAndFriction/DampingAndFriction.INT.udn
@@ -34,7 +34,7 @@ The Damping for Physics Constraints work a bit differently than with Physics Bod
 Physics Constraint's Motion(s) to Limited will display the **Soft Limits**:
 
 [REGION:raw]
-![](Phys_ConstraintDampProps.png)(convert:false)
+![](AngularAndDamping.png)(convert:false)
 [/REGION]
 
 _The two Angular Swing Motions share Damping values, the Angular Twist receives its own Damping values, and all the Linear axes use the same Damping values. While they may have 
