UEDOC-478

Removed the WIP warning box at the top of the page.

diff --git a/Engine/Documentation/Source/Programming/Plugins/Plugins.INT.udn b/Engine/Documentation/Source/Programming/Plugins/Plugins.INT.udn
--- a/Engine/Documentation/Source/Programming/Plugins/Plugins.INT.udn
+++ b/Engine/Documentation/Source/Programming/Plugins/Plugins.INT.udn
@@ -4,14 +4,6 @@ Crumbs:%ROOT%, Programming
 Description:How to create Unreal Engine plugins.
 Version: 4.9
 
-
-[REGION:warning]
-This document is a **work in progress**. Our plugin system is new and still undergoing significant changes 
-each month. At some point when we feel the system has cooled down a bit, we will remove this notice from the 
-document.
-[/REGION]
-
-
 [TOC(start:2)]
 
 
