UEDOC-1041

Updating images on the page

diff --git a/Engine/Documentation/Source/Engine/Rendering/Materials/IntroductionToMaterials/IntroductionToMaterials.INT.udn b/Engine/Documentation/Source/Engine/Rendering/Materials/IntroductionToMaterials/IntroductionToMaterials.INT.udn
--- a/Engine/Documentation/Source/Engine/Rendering/Materials/IntroductionToMaterials/IntroductionToMaterials.INT.udn
+++ b/Engine/Documentation/Source/Engine/Rendering/Materials/IntroductionToMaterials/IntroductionToMaterials.INT.udn
@@ -174,15 +174,15 @@ Applying a Material to a Brush surface is a simple process:
 
 1. Select a Material from the **Content Browser**.
 
-	![](MaterialSelectedInCB.png)
+	![](MaterialSelected.png)
 
 1. Click on the Brush surface to which you wish to assign the Material.
 
-	![](SurfaceSelected.png)
+	![](BoxPanelSelected.png)
 
-1. In the **Details** panel, click the ![](button_ApplyMaterials.png) button. Select the desired material from the dropdown that appears.
+1. In the **Details** panel, click the ![](ArrowButton.png) button. Pressing this button will apply the Material to the Brush surface.
 
-	![](ApplyMaterial.png)![](MaterialApplied.png)
+	![](AddingMaterialToBox.png)![](RedOnBox.png)
 
 
 %playlistnav%
