Fixing grammar error reported from the Forums

https://forums.unrealengine.com/showthread.php?66057-Reporting-Issues-for-Documentation&p=348762&viewfull=1#post348762

diff --git a/Engine/Documentation/Source/Engine/Rendering/Materials/HowTo/DetailTexturing/DetailTexturing.INT.udn b/Engine/Documentation/Source/Engine/Rendering/Materials/HowTo/DetailTexturing/DetailTexturing.INT.udn
--- a/Engine/Documentation/Source/Engine/Rendering/Materials/HowTo/DetailTexturing/DetailTexturing.INT.udn
+++ b/Engine/Documentation/Source/Engine/Rendering/Materials/HowTo/DetailTexturing/DetailTexturing.INT.udn
@@ -40,10 +40,10 @@ While this is not the only way that you can apply Detail Textures to Materials,
 |--------- | ---------- |
 | **Scale(S)** | This sets the scale of the detail Texture. Bigger numbers result in more titling smaller numbers results in less titling. |
 | **Diffuse(V3)** | This is the input for the Diffuse Texture. |
-| **DetailDiffuse(T2d)** | This is the input for the Diffuse Detail Texture. This input can only except a Texture Object. |
+| **DetailDiffuse(T2d)** | This is the input for the Diffuse Detail Texture. This input can only accept a Texture Object. |
 | **DiffuseIntensity(S)** | This controls the intensity of Detail Diffuse Texture. |
 | **Normal(V3)** | This is the input for the Normal Map Texture. |
-| **DetailNormal(T2D)** | This is the input for the Normal Map Detail Texture. This input can only except a Texture Object.|
+| **DetailNormal(T2D)** | This is the input for the Normal Map Detail Texture. This input can only accept a Texture Object.|
 | **NormalIntensity(S)** | This controls the intensity of the Detail Normal Map Texture.|
 
 ### Converting Texture Sample to Texture Objects
