Fixed error on page with paragraph mentioning left and right

https://forums.unrealengine.com/showthread.php?66057-Reporting-Issues-for-Documentation&p=348760&viewfull=1#post348760

diff --git a/Engine/Documentation/Source/Engine/Rendering/Materials/HowTo/DetailTexturing/DetailTexturing.INT.udn b/Engine/Documentation/Source/Engine/Rendering/Materials/HowTo/DetailTexturing/DetailTexturing.INT.udn
--- a/Engine/Documentation/Source/Engine/Rendering/Materials/HowTo/DetailTexturing/DetailTexturing.INT.udn
+++ b/Engine/Documentation/Source/Engine/Rendering/Materials/HowTo/DetailTexturing/DetailTexturing.INT.udn
@@ -26,7 +26,7 @@ Here is an example of Detail Texturing in action.
 
 On the Left Hand side, labeled with a number 1, you can see a Material that does make use of Detail Texturing. 
 On the Right Hand side, labeled with a number 2, you can see a Material that does not make use of Detail Texturing. 
-Notice how in the image on the left, 1, there looks to be more detail from the Normal map and that detail looks much sharper than the image on the left. 
+Notice how in the image on the left, 1, there looks to be more detail from the Normal map and that detail looks much sharper than the image on the right. 
 
 
 ## Detail Texturing Node Breakdown
