#UE4doc #docs - Updated images and corrected to meet new auto-convex method.

diff --git a/Engine/Documentation/Source/Engine/Physics/Collision/Reference/Reference.INT.udn b/Engine/Documentation/Source/Engine/Physics/Collision/Reference/Reference.INT.udn
--- a/Engine/Documentation/Source/Engine/Physics/Collision/Reference/Reference.INT.udn
+++ b/Engine/Documentation/Source/Engine/Physics/Collision/Reference/Reference.INT.udn
@@ -83,10 +83,10 @@ Selecting **Collision Menu -> Auto Convex Collision** will bring up this panel i
 
 ![Collision_CreateCollisionHullWithAutoConvex_1.jpg](Collision_CreateCollisionHullWithAutoConvex_1.png)(convert:false)
 
-* **Max Hulls** - Determines the number of hulls to create to best match the shape of the object.
+* **Accuracy** - Determines the the fewest number of hulls to create to best match the shape of the object.
 * **Max Hull Verts** - Maximum number of collision hull vertices. Increasing this value increase how complex the collision hulls are allowed to be.
 * **Apply** - Adds collision to the object based on the settings above.
-* **Defaults** - Resets the values of Max Hulls and Max Hull Verts to the values seen in the image above.
+* **Defaults** - Resets the values of Accuracy and Max Hull Verts to the values seen in the image above.
 
 ![Collision_CreateCollisionHullWithAutoConvex_2.jpg](Collision_CreateCollisionHullWithAutoConvex_2.png)(convert:false)
 
