UEDOC-934

Banner has been updated with updated image

diff --git a/Engine/Documentation/Source/Engine/Rendering/Materials/Functions/Overview/FunctionsOverview.INT.udn b/Engine/Documentation/Source/Engine/Rendering/Materials/Functions/Overview/FunctionsOverview.INT.udn
--- a/Engine/Documentation/Source/Engine/Rendering/Materials/Functions/Overview/FunctionsOverview.INT.udn
+++ b/Engine/Documentation/Source/Engine/Rendering/Materials/Functions/Overview/FunctionsOverview.INT.udn
@@ -4,7 +4,7 @@ Crumbs: %ROOT%, Engine, Engine/Rendering/Materials
 Description:An overview of how material functions work and some critical concepts for their use.
 
 [REGION:banner]
-![Material Function](matfunc_overview_banner.png)
+![](matfunc_overview_banner_NEW.jpg)
 [/REGION]
 
 [TOC(start:2)]
