UEDOC-959

Updated images to v4.9 images

diff --git a/Engine/Documentation/Source/Engine/Actors/Types/CommonActorTypes.INT.udn b/Engine/Documentation/Source/Engine/Actors/Types/CommonActorTypes.INT.udn
--- a/Engine/Documentation/Source/Engine/Actors/Types/CommonActorTypes.INT.udn
+++ b/Engine/Documentation/Source/Engine/Actors/Types/CommonActorTypes.INT.udn
@@ -35,12 +35,12 @@ Unreal Engine to familiarize you with what they are and why you would use them.
 | **Brush** | ![](Shared/Icons/Source/ActorIcons/Brush_40x.png)(convert:false) | The **Brush Actor** is a basic type of Actor that displays simple 3D geometry in the scene. These Actors can be modified using the Geometry Editing mode in the Level Editor. BrushActors (or just _Brushes_) are commonly used for quickly prototyping environments and blocking out levels for testing gameplay. |
 | **SkeletalMeshActor** | ![](Shared/Icons/Source/ActorIcons/SkeletalMesh_40x.png)(convert:false) | The **Skeletal Mesh Actor** is a type of Actor that displays an animated mesh, or a SkeletalMesh, whose geometry can be deformed - generally through the use of animation sequences created and exported from external 3D animation applications. These Actors are usually used for things like characters or other living creatures as well as complex machinery; anything that needs to deform or display other complex motion. These are also often used in conjunction with Matinee to create cinematic sequences |
 | [REGION:tablesection]Gameplay Actor Types[/REGION] |||
-| **PlayerStart** | ![](Shared/Icons/Source/ActorIcons/PlayerStart_40x.png)(convert:false) | A **Player Start** is an Actor that is placed in the level to designate where the player should start out when the game begins. |
-| **Triggers** | ![](Shared/Icons/Source/icon_class_TriggerBox_40px.png)(convert:false) ![](Shared/Icons/Source/icon_class_TriggerCapsule_40px.png)(convert:false) ![](Shared/Icons/Source/icon_class_TriggerSphere_40px.png)(convert:false) | [INCLUDE:Engine/Actors/Triggers#intro] [INCLUDE:#trigger] See [](Engine/Actors/Triggers) for more detailed usage info. |
-| **MatineeActor** | ![](Shared/Icons/Source/icon_class_MatineeActor_40px.png)(convert:false) | **Matinee Actors** enable you to animate the properties of Actors over time via the **Matinee** animation tool, to create either dynamic gameplay or cinematic in-game sequences. The system is based on the use of specialized animation tracks in which you can place keyframes to set the values of certain properties of the Actors in the level. The Matinee animation tool is similar to the non-linear editors used for video editing, making it familiar to video professionals. See [](Engine/Matinee) for more detailed usage information. |
+| **PlayerStart** | ![](PlayerStart.png)(w:43) | A **Player Start** is an Actor that is placed in the level to designate where the player should start out when the game begins. |
+| **Triggers** | ![](BoxTrigger.png)(w:43) ![](CapsuleTrigger.png)(w:43) ![](SphereTrigger.png)(w:43) | [INCLUDE:Engine/Actors/Triggers#intro] [INCLUDE:#trigger] See [](Engine/Actors/Triggers) for more detailed usage info. |
+| **MatineeActor** | ![](MatineeIcon.png)(w:43) | **Matinee Actors** enable you to animate the properties of Actors over time via the **Matinee** animation tool, to create either dynamic gameplay or cinematic in-game sequences. The system is based on the use of specialized animation tracks in which you can place keyframes to set the values of certain properties of the Actors in the level. The Matinee animation tool is similar to the non-linear editors used for video editing, making it familiar to video professionals. See [](Engine/Matinee) for more detailed usage information. |
 | [REGION:tablesection]Light Actor Types[/REGION] |||
 | **PointLight** | ![](Shared/Icons/Source/icon_class_PointLight_40px.png)(convert:false) | [INCLUDE:Engine/Rendering/LightingAndShadows/LightTypes/Point#intro] |
-| **SpotLight** | ![](Shared/Icons/Source/icon_class_SpotLight_40px.png)(convert:false) | [INCLUDE:Engine/Rendering/LightingAndShadows/LightTypes/Spot#intro] |
+| **SpotLight** | ![](SpotLight.png)(w:43) | [INCLUDE:Engine/Rendering/LightingAndShadows/LightTypes/Spot#intro] |
 | **DirectionalLight** | ![](Shared/Icons/Source/icon_class_DirectionalLight_40px.png)(convert:false) | [INCLUDE:Engine/Rendering/LightingAndShadows/LightTypes/Directional#intro] |
 | [REGION:tablesection]Effects Actor Types[/REGION] |||
 | **ParticleEmitter** | ![](Shared/Icons/Source/ActorIcons/Emitter_40x.png)(convert:false) | A **Particle Emitter** is an Actor that is used to create effects such as smoke, fire, sparks, and more by spawning particles in the form of sprites (camera-facing planes) or meshes. The actual behavior of the particles is defined in and controlled by special assets called [Particle Systems](Engine/Rendering/ParticleSystems) which are created in the [Content Browser](Engine/Content/Browser) and edited in the [Cascade Particle Editor](Engine/Rendering/ParticleSystems/Cascade) |
