UEDOC-1010

Imgaes have been updated to better reflect later versions of the engine

diff --git a/Engine/Documentation/Source/Engine/UI/LevelEditor/Viewports/Basics/ViewportBasics.INT.udn b/Engine/Documentation/Source/Engine/UI/LevelEditor/Viewports/Basics/ViewportBasics.INT.udn
--- a/Engine/Documentation/Source/Engine/UI/LevelEditor/Viewports/Basics/ViewportBasics.INT.udn
+++ b/Engine/Documentation/Source/Engine/UI/LevelEditor/Viewports/Basics/ViewportBasics.INT.udn
@@ -169,7 +169,7 @@ will look when you run it in the game.
 To enable Game View, simply press the **G** key (by default) with the Viewport focused or choose the **Game View** option from the 
 [Viewport Options menu](Engine/UI/LevelEditor/Viewports/ViewportOptions):
 
-![](viewport_options_gameview.png)
+![](GamemodeViewport.png)
 
 ## Immersive Mode
 
@@ -180,7 +180,7 @@ maximized, this enables you to have your Viewport run in full-screen for a truly
 To enable Immersive Mode, simply press the **F11** key (by default) with the Viewport focused or choose the **Immersive Mode** option 
 from the [Viewport Options menu](Engine/UI/LevelEditor/Viewports/ViewportOptions):
 
-![](viewport_options_immersive.png)
+![](ImmersiveModeOptions.png)
 
 You can restore the Viewport to normal using the same procedure above or by pressing the ![Restore Viewport](button_restore.png) 
 button located at the top right of the Viewport when in Immersive Mode.
