Release/Patching docs

diff --git a/Engine/Documentation/Source/Engine/Deployment/Patching/Patching.INT.udn b/Engine/Documentation/Source/Engine/Deployment/Patching/Patching.INT.udn
new file mode 100644
--- /dev/null
+++ b/Engine/Documentation/Source/Engine/Deployment/Patching/Patching.INT.udn
@@ -0,0 +1,83 @@
+Availability:Docs
+Title:Patching
+Crumbs:%ROOT%, Engine, Engine/Deployment
+Description:Creating updated content packages which allow you to update your project after release.
+
+
+To update your project after release, you can use the patching process available through the Project Launcher. During patching, the engine compares all the content post-cook
+to the originally released cooked content, and uses that to determine what is included in a patch.  The smallest piece of content is a single package (like a .ulevel or .uasset), 
+so if anything changes in a package then the entire package will be included in the patch. The method you use to get the patch's .pak file to the users will depend on 
+your distribution platforms, but this procedure will allow you to create a much smaller .pak file that just contains the updated content. 
+
+[REGION:warning]
+Currently, distribution of the patches hasn't been resolved for all platforms. This procedure will walk you through creating the patch package file, as well as through testing the package file
+on Windows. 
+[/REGION]
+
+You can patch a project you have previously released using a versioned release. Some things to keep in mind are:
+* Lock down the serialization code paths at the time of release.
+* Keep the released cooked content as the UnrealPak tool uses this to determine which content should be in the patch package file. 
+* At runtime, mount both pak files, with a higher priority for the patch file so any content within it is loaded first. 
+
+
+In this example, we are patching ShooterGame with a change to the weapon capacity and magazine size.
+
+![](PatchExample.png)
+
+**After creating a release version:**
+
+1. Open the **Project Launcher**, either within Unreal Editor or using Unreal Frontend.
+
+	![](1.png)
+1. Create a new Custom Launch Profile using the **+** button.
+	![](2.png)
+1. Set a name and description for your profile. 
+
+There are a number of settings for the patching process.
+
+#### Project
+1. You can set the specific project to use, or use **Any Project** to patch the current project.
+
+#### Build
+1. Set the build configuration to Shipping.
+1. Optionally, expand **Advanced Settings** if you need to build UAT as part of the patch process.
+
+
+#### Cook
+1. Select **By the Book** as the cooking method in the dropdown menu.
+1. Check the boxes for all platforms you would like to cook content for. In this example for Windows testing, we have selected **WindowsNoEditor**.
+1. Check the boxes for all cultures to cook localizations for.
+1. Check the boxes for all maps to cook. 
+[Region:tip]
+This will determine which content is cooked and then compared to the original package file. 
+[/REGION]
+1. In **Release/DLC/Patching Settings**:
+	* Enter in the version number to base the patch on.
+	* Check the **Generate Patch** checkbox.
+1. Expand **Advanced Settings** and make sure the following options are enabled, as well as any others you need for your specific project's distribution method:
+	* Compress content
+	* Save packages without versions
+	* Store all content in a single file (UnrealPak)
+1. Also under **Advanced Settings**, set the cooker configuration to **Shipping**.
+
+#### Package 
+* Set the build to  **Package & store locally**. 
+
+#### Deploy
+* Set the build to **Do Not Deploy**.
+
+1. Once you have set all the above settings, navigate back to the main profile window using the **Back** button in the top right corner. 
+	![](9.png)
+
+1. Click on the launch icon next to your **Patching** profile. 
+	![](10.png)
+1. The project launcher will go through the building, cooking, and packaging process. This may take some time depending on the complexity of your project.
+![](11.png)
+1. Once the operation is complete, close the window or click on **Done**. You can test the patch now with the steps below.
+	![](12.png)
+
+
+**To test the patch on Windows:**
+
+1. Copy the pak file from `ShooterGame\Releases\1.0\WindowsNoEditor` to `ShooterGame\Saved\StagedBuilds\WindowsNoEditor\ShooterGame\Content\Paks`
+1. Run `ShooterGame.exe` from `ShooterGame\Saved\StagedBuilds\WindowsNoEditor`.
\ No newline at end of file
diff --git a/Engine/Documentation/Source/Engine/Deployment/Releasing/Releasing.INT.udn b/Engine/Documentation/Source/Engine/Deployment/Releasing/Releasing.INT.udn
new file mode 100644
--- /dev/null
+++ b/Engine/Documentation/Source/Engine/Deployment/Releasing/Releasing.INT.udn
@@ -0,0 +1,73 @@
+Availability:Docs
+Title:Releasing Your Project
+Crumbs:%ROOT%, Engine, Engine/Deployment
+Description:Creating a release versioned project with cooked content for distribution.
+
+
+Once your game is complete, you can package a release version for distribution. Whether you go through the Project Launcher in the editor or the Project Launcher tab in Unreal Frontend,
+the steps are the same for packaging a release-versioned project. However, depending on if you are creating a Windows game and need to create your own installer, shipping a mobile project,
+or targeting another platform, the steps you take with the finished packaged content will be different. 
+
+This is an example of packaging a 1.0 release of ShooterGame, aimed at Windows 64-bit, localized in English. 
+
+1. Open the **Project Launcher**, either within Unreal Editor or using Unreal Frontend.
+
+	![](1.png)
+1. Create a new Custom Launch Profile using the **+** button.
+	![](2.png)
+1. Set a name and description for your profile. 
+
+There are a number of settings for the release process.
+
+#### Project
+1. You can set the specific project to use, or use **Any Project** to patch the current project.
+
+#### Build
+1. Set the build configuration to Shipping.
+1. Optionally, expand **Advanced Settings** if you need to build UAT as part of the release.
+
+#### Cook
+1. Select **By the Book** as the cooking method in the dropdown menu.
+1. Check the boxes for all platforms you would like to cook content for. In this example for Windows testing, we have selected **WindowsNoEditor**.
+1. Check the boxes for all cultures to cook localizations for.
+1. Check the boxes for which maps to cook. 
+
+1. In **Release/DLC/Patching Settings**:
+	* Check the **Create a release version of the game for distribution.** checkbox.
+	* Enter a version number for this release. 
+1. Expand **Advanced Settings** and make sure the following options are enabled, as well as any others you need for your specific project's distribution method:
+	* Compress content
+	* Save packages without versions
+	* Store all content in a single file (UnrealPak)
+1. Also under **Advanced Settings**, set the cooker configuration to **Shipping**.
+
+#### Package 
+* Set the build to  **Package & store locally**. 
+
+#### Deploy
+* Set the build to **Do Not Deploy**.
+
+1. Once you have set all the above settings, navigate back to the main profile window using the **Back** button in the top right corner. 
+	![](9.png)
+
+1. Click on the launch icon next to your **Release** profile. 
+	![](10.png)
+1. The project launcher will go through the building, cooking, and packaging process. This may take some time depending on the complexity of your project.
+![](11.png)
+1. Once the operation is complete, close the window or click on **Done**. You can test the patch now with the steps below.
+	![](12.png)
+
+1. Save the asset registry and pak file from `[ProjectName]\Releases\[ReleaseVersion]\[Platform]`. In this example, this is `ShooterGame\Releases\1.0\WindowsNoEditor`.
+
+[Region:tip]
+The asset registry and pak file will be needed for any future patches or DLC to check against. 
+[/REGION]
+
+1. On Windows, you can test running the project from `[ProjectName]\Saved\StagedBuilds\WindowsNoEditor`.
+
+
+
+[REGION:note]
+While Steam will allow you to upload full packages of your game and do the updating process for you, using release versions as outlined here is still the recommended practice
+when distributing through Steam. This will make the process smoother if you decide to add additional supported platforms or distribution methods later on. 
+[/REGION]
\ No newline at end of file
