UEDOC-897

Image has been updated to better capture the Material Editor

diff --git a/Engine/Documentation/Source/Programming/Rendering/MaterialExpressions/CreatingMaterialExpressions.INT.udn b/Engine/Documentation/Source/Programming/Rendering/MaterialExpressions/CreatingMaterialExpressions.INT.udn
--- a/Engine/Documentation/Source/Programming/Rendering/MaterialExpressions/CreatingMaterialExpressions.INT.udn
+++ b/Engine/Documentation/Source/Programming/Rendering/MaterialExpressions/CreatingMaterialExpressions.INT.udn
@@ -89,5 +89,5 @@ The constructor is used to set default values for the properties. In this exampl
 Your new material expression should now be in the list in the material editor.  Here you can see the Square expression in action:  
 
 [REGION:todo]
-![SquareMaterialExpression.jpg](SquareMaterialExpression.jpg)
+![](MaterialExpression.png)
 [/REGION]
