UEDOC-2192

Images have been changed from green boxes to yellow

diff --git a/Engine/Documentation/Source/Gameplay/HowTo/SetUpInput/Blueprints/SetInput_BP.INT.udn b/Engine/Documentation/Source/Gameplay/HowTo/SetUpInput/Blueprints/SetInput_BP.INT.udn
--- a/Engine/Documentation/Source/Gameplay/HowTo/SetUpInput/Blueprints/SetInput_BP.INT.udn
+++ b/Engine/Documentation/Source/Gameplay/HowTo/SetUpInput/Blueprints/SetInput_BP.INT.udn
@@ -18,7 +18,7 @@ _For this example, we are using the **Blank** project template with **Starter Co
 1.	Inside your project, under **Shapes**, **Right-click** on the **Shape_NarrowCapsule** shape and select **Asset Actions** - **Create Blueprint Using**. 
 
 	[REGION:raw]
-	![](Input_1.png)
+	![](Input2_1.png)
 	[/REGION]
 
 	This will be our simple character that we will control. 
@@ -38,7 +38,7 @@ _For this example, we are using the **Blank** project template with **Starter Co
 1.	On the **Components** tab under the **Components** window, click **Add Component** and add a **Spring Arm**. 
 
 	[REGION:raw]
-	![](Input_3.png)
+	![](Input2_3.png)
 	[/REGION]
 	
 	[REGION:note]
@@ -50,12 +50,12 @@ _For this example, we are using the **Blank** project template with **Starter Co
 1.	Under **Camera Settings**, check the **Use Pawn Controller Rotation** box. 
 
 	[REGION:raw]
-	![](Input_4.png)
+	![](Input2_4.png)
 	[/REGION]
 
 1.	Click the **Add Component** button again and add a **Camera**, then drag and drop it onto the **Spring Arm** component. 
 
-	![](Input_5.png)
+	![](Input2_5.jpg)
 
 	This will make it a child of the **Spring Arm**. 
 
@@ -83,13 +83,13 @@ Now that we have created a simple character, we need to create a **GameMode** an
 
 1.	On the **Defaults** tab, under the **Classes** section, click the **Default Pawn Class** drop-down and select the new character. 
 
-	![](Input_9.png)
+	![](Input2_9.png)
 
 1.	**Compile** and **Save** then close the Blueprint. 
 
 1.	From the Main Editor window, click the **Edit** Button from the menu bar and select **Project Settings**. 
 
-	![](Input_10.png)
+	![](Input2_10.png)
 
 	We are now going to tell our project to use the newly created **GameMode**.
 
@@ -97,7 +97,7 @@ Now that we have created a simple character, we need to create a **GameMode** an
 
 1.	Under **Maps & Modes**, for **Default Modes**, click the **Default GameMode** drop-down and select your **GameMode**. 
 
-	![](Input_11.png)
+	![](Input2_11.png)
 
 
 ## Creating New Inputs (Action/Axis Mappings)
@@ -112,7 +112,7 @@ The steps below will guide you through adding new **Action** and **Axis Mappings
 
 1.	Fill in the empty mappings as seen below. 
 
-	![](Input_13.png)(w:500)
+	![](Input2_13.png)(w:500)
 
 	Here you are naming each of the mappings and assigning an Input for that mapping. Make sure to set the **Scale** values to **-1.0** where highlighted as the **Scale** acts as a multiplier on the value of the key when summing up the Axis' value (which allows us to get the opposite direction). 
 
@@ -125,7 +125,7 @@ Once you have created Inputs via **Project Settings**, you can then call those I
 
 1.	Inside the Blueprint, **Right-click** anywhere in the graph and search for **Turn** then add the **Turn** Axis Event. 
 
-	![](Input_14.png)
+	![](Input2_14.png)
 
 1.	Drag off the **Axis Value** and search for **Yaw**, then add the **Add Controller Yaw Input** node. 
 
@@ -141,11 +141,11 @@ Once you have created Inputs via **Project Settings**, you can then call those I
 
 1.	Drag off the **Pressed** pin and search for **Jump** function. 
 
-	![](Input_17.png)
+	![](Input2_17.png)
 
 1.	**Right-click** in the graph and add the **MoveForward** Axis Event. 
 
-	![](Input_18.png)
+	![](Input2_18.png)
 
 1.	Drag off the **Axis Value** and add the **Add Movement Input** node. 
 
@@ -153,23 +153,23 @@ Once you have created Inputs via **Project Settings**, you can then call those I
 
 1.	Drag off the **Axis Value** and add the **Add Movement Input** node. 
 
-	![](Input_19.png)
+	![](Input2_19.png)
 
 	The events have been added, however we need to specify direction for each still. 
 
 1.	**Right-click** in the graph and search for **Get Control Rotation**. 
 
-	![](Input_20.png)
+	![](Input2_20.png)
 
 1.	Drag off the **Return Value** and add a **Break Rot** node (to break the rotation into Pitch, Yaw, and Roll). 
 
 1.	Drag off the **Yaw** pin and add a **Make Rot** node (to make a rotation out of Pitch, Yaw, and Roll). 
 
-	![](Input_21.png)
+	![](Input2_21.png)
 
 1.	**Right-click** on the **Make Rot** node and select **Break link to Break Rot (Yaw)**. 
 
-	![](Input_22.png)
+	![](Input2_22.png)
 
 	We want to zero out the **Pitch** and **Roll** and only keep the **Yaw**. 
 
