Updating images and text for BeginOverlap and End Overlap

diff --git a/Engine/Documentation/Source/Engine/Physics/Collision/Collision.INT.udn b/Engine/Documentation/Source/Engine/Physics/Collision/Collision.INT.udn
--- a/Engine/Documentation/Source/Engine/Physics/Collision/Collision.INT.udn
+++ b/Engine/Documentation/Source/Engine/Physics/Collision/Collision.INT.udn
@@ -182,14 +182,14 @@ Or:
 
 ### Overlap and Generate Overlap Events
 
-Unlike collisions that can fire every frame, the overlap events are **ReceiveBeginOverlap** and **ReceiveEndOverlap**, which only fire in those specific cases.
+Unlike collisions that can fire every frame, the overlap events are **BeginOverlap** and **EndOverlap**, which only fire in those specific cases.
 
 ![](COL_overlapEvent.png)(w:940 convert:false)
 
 | Sphere Collision Setup| Wall Collision Setup |
 | --- | ---- |
-| ![](COL_OverlapEvent_Sphere.png)(w:300) %spacer% | ![](COL_collideNoEvent_Box.png)(w:303) %spacer% |
-| Here the sphere is set to `overlap` **WorldDynamic** Actors (like our wall), and it will generate an event for itself when it does overlap something. | The wall is a **WorldDynamic** and is set to `block` **PhysicsBody** Actors (which is what the sphere is). As stated above, both Actors need to be set to block each other's respective object types. If they do not, they will not collide. But, an **Overlap** does occur here, and an event for the sphere is fired. |
+| ![](SphereCollisiionSetup.png) %spacer% | ![](WallCollisionSetup.png) %spacer% |
+| Here the sphere is set to `overlap` **WorldDynamic** Actors (like our wall), and it will generate an event for itself when it does overlap something. | The wall is a **WorldDynamic** and is set to `block` **PhysicsBody** Actors (which is what the sphere is). As stated above, both Actors need to be set to block each other's respective object types. If they do not, they will not collide. But, an **Overlap** does occur here, and an event is fired. |
 
  
 ## Custom Object Channels and Trace Channels
