Updating the Begin and End Overlap condiitons that need to be met

diff --git a/Engine/Documentation/Source/Engine/Blueprints/UserGuide/Events/Events.INT.udn b/Engine/Documentation/Source/Engine/Blueprints/UserGuide/Events/Events.INT.udn
--- a/Engine/Documentation/Source/Engine/Blueprints/UserGuide/Events/Events.INT.udn
+++ b/Engine/Documentation/Source/Engine/Blueprints/UserGuide/Events/Events.INT.udn
@@ -103,8 +103,8 @@ has died but the level does not need to reload.
 
 This event will execute when a number of conditions are met at the same time:
 
-* Both Actors must be set to Overlap.
-* The Actor that is to execute the event has **Generate Overlap Events** set to _true_.
+* Collision response between the actors must allow Overlaps.
+* Both Actors that are to execute the event have **Generate Overlap Events** set to _true_.
 * And finally, both Actors' collision starts overlapping; moving together or one is created overlapping the other.
 
 For more information on collision, see: [**Collision Responses**](Engine\Physics\Collision).
@@ -125,8 +125,8 @@ _When this Blueprint Actor overlaps the Actor stored in the Player Actor variabl
 
 This event will execute when a number of conditions are met at the same time:
 
-* Both Actors must be set to Overlap.
-* The Actor that is to execute the event has **Generate Overlap Events** set to _true_.
+* Collision response between the actors must allow Overlaps.
+* Both Actors that are to execute the event have **Generate Overlap Events** set to _true_.
 * And finally, both Actors' collision stop overlapping; moving apart or if one is destroyed.
 
 For more information on collision, see: [**Collision Responses**](Engine\Physics\Collision).
