#UE4doc #docs Revised document per JeffW's feedback. Restructured page and reworked overview.

diff --git a/Engine/Documentation/Source/GettingStarted/DownloadingUnrealEngine/DownloadingUnrealEngine4.INT.udn b/Engine/Documentation/Source/GettingStarted/DownloadingUnrealEngine/DownloadingUnrealEngine4.INT.udn
--- a/Engine/Documentation/Source/GettingStarted/DownloadingUnrealEngine/DownloadingUnrealEngine4.INT.udn
+++ b/Engine/Documentation/Source/GettingStarted/DownloadingUnrealEngine/DownloadingUnrealEngine4.INT.udn
@@ -3,8 +3,8 @@ Title:Downloading Unreal Engine Source Code
 Crumbs: %ROOT%, GettingStarted, GettingStarted/DownloadingUnrealEngine
 Description:Steps for connecting to the source code repository and downloading the latest build of Unreal Engine.
 Prereq: GettingStarted/RecommendedSpecifications
-Related: Programming/Development/BuildingUnrealEngine
-Related: Support/Builds
+Related: GettingStarted/DownloadingUnrealEngine/GitHub
+Related: GettingStarted/DownloadingUnrealEngine/DownloadZip
 Version: 4.8
 SkillLevel: Advanced
 
@@ -16,24 +16,12 @@ Unreal Engine 4 is available to be downloaded by licensees directly from Epic's
 [/PUBLISH:Licensee]
 [PUBLISH:Rocket]
 
-[REGION:warning]
-This page shows subscribers how to download and build Unreal Engine from our source code repository on GitHub. If you'd like to download the binary version of Unreal Engine, read our [](GettingStarted/Installation) documentation to learn how to [Get Unreal](https://www.unrealengine.com/).
-[/REGION]
-
-----
-
-Source code for Unreal Engine 4 (UE4) can be downloaded by our subscribers from the Unreal Engine GitHub repository. 
+Source code for Unreal Engine 4 (UE4) can be downloaded by our subscribers from the Unreal Engine GitHub repository. Before you can access the [repository](https://github.com/EpicGames/UnrealEngine), you must be an Unreal Engine subscriber with a GitHub account. For instructions on associating your GitHub account with your Unreal Engine account, visit the [UE4 on GitHub](https://www.unrealengine.com/ue4-on-github) page.
 
 <!-- The following is commented out per LaurenR's request during Peer Review on 23.07.15
 With access to our source code, subscribers can customize the engine and tools to meet their needs. Our subscribers are building the editor in Windows, Mac, Linux, and compiling games for both Android and iOS. Source code for all of our tools (such as Unreal Lightmass and Frontend) is included as well.
 -->
 
-[REGION:note]
-Before you can access the [repository](https://github.com/EpicGames/UnrealEngine), you must be an Unreal Engine subscriber with a GitHub account.  
-*	If you haven't created a GitHub account, go to [GitHub's](https://github.com/) page to sign-up.   
-*	For instructions on associating your GitHub account with your Unreal Engine account, visit the [UE4 on GitHub](https://www.unrealengine.com/ue4-on-github) page.
-[/REGION]
-
 [INCLUDE:GettingStarted/DownloadingUnrealEngine/SourceBranches]
 
 [/PUBLISH:Rocket]
@@ -44,24 +32,16 @@ Before you can access the [repository](https://github.com/EpicGames/UnrealEngine
 
 [PUBLISH:Rocket]
 
-## Downloading Unreal Engine
-
-Subscribers can download Unreal Engine's source code by:
+## Next Steps
 
-*	[%GettingStarted/DownloadingUnrealEngine/GitHub:title%](#Bookmark4)
-*	[%GettingStarted/DownloadingUnrealEngine/DownloadZip:title%](#Bookmark5)
+Subscribers can download Unreal Engine's source code in one of two ways:
 
-### %GettingStarted/DownloadingUnrealEngine/GitHub:title%
+*	[%GettingStarted/DownloadingUnrealEngine/GitHub:title%](GettingStarted/DownloadingUnrealEngine/GitHub)
+*	[%GettingStarted/DownloadingUnrealEngine/DownloadZip:title%](GettingStarted/DownloadingUnrealEngine/DownloadZip)
 
-(#Bookmark4)
-
-[INCLUDE:GettingStarted/DownloadingUnrealEngine/GitHub]
-
-### %GettingStarted/DownloadingUnrealEngine/DownloadZip:title%
-
-(#Bookmark5)
-
-[INCLUDE:GettingStarted/DownloadingUnrealEngine/DownloadZip]
+[REGION:note]
+This page shows subscribers how to download and build Unreal Engine from our source code repository on GitHub. If you'd like to download the binary version of Unreal Engine, read our [](GettingStarted/Installation) documentation to learn how to [Get Unreal](https://www.unrealengine.com/).
+[/REGION]
 
 [/PUBLISH:Rocket]
 
