#UE4doc #docs Restructured and revised documentation per JeffW's feedback.

diff --git a/Engine/Documentation/Source/Programming/Development/BuildingUnrealEngine/BuildingUnrealEngine.INT.udn b/Engine/Documentation/Source/Programming/Development/BuildingUnrealEngine/BuildingUnrealEngine.INT.udn
--- a/Engine/Documentation/Source/Programming/Development/BuildingUnrealEngine/BuildingUnrealEngine.INT.udn
+++ b/Engine/Documentation/Source/Programming/Development/BuildingUnrealEngine/BuildingUnrealEngine.INT.udn
@@ -3,8 +3,6 @@ Title:Building Unreal Engine from Source
 Crumbs: %ROOT%, Programming, Programming/Development
 Description:Compiling Unreal Engine from source.
 Prereq: GettingStarted/DownloadingUnrealEngine
-Prereq: GettingStarted/RecommendedSpecifications
-Related: GettingStarted/DownloadingUnrealEngine
 Related: GettingStarted
 Related: Engine/Content/QuickStart
 Related: Programming/QuickStart
