#UE4doc #docs Broke page out from parent DownloadingUnrealEngine4  documentation per JeffW's feedback.

diff --git a/Engine/Documentation/Source/GettingStarted/DownloadingUnrealEngine/GitHub/GitHub.INT.udn b/Engine/Documentation/Source/GettingStarted/DownloadingUnrealEngine/GitHub/GitHub.INT.udn
--- a/Engine/Documentation/Source/GettingStarted/DownloadingUnrealEngine/GitHub/GitHub.INT.udn
+++ b/Engine/Documentation/Source/GettingStarted/DownloadingUnrealEngine/GitHub/GitHub.INT.udn
@@ -2,6 +2,8 @@
 Title:Forking and Cloning Unreal from GitHub
 Crumbs:%ROOT%, GettingStarted, GettingStarted/DownloadingUnrealEngine
 Description:Forking and cloning the Unreal Engine source code from GitHub.
+Prereq: GettingStarted/DownloadingUnrealEngine
+Related: Programming/Development/BuildingUnrealEngine
 Version: 4.8
 SkillLevel: Advanced
 
@@ -98,6 +100,10 @@ Forking and cloning our source code from GitHub is a great way to get started wi
 	[/PARAM]
 [/OBJECT]
 
+## Next Steps
+
+Now that you've downloaded and set-up Unreal Engine, you're ready to [build the engine from source](Programming/Development/BuildingUnrealEngine).
+
 <!--
 
 There are two ways to get the source code for Unreal Engine from GitHub:
