#UE4doc #docs Broke page out from parent DownloadingUnrealEngine4  documentation per JeffW's feedback. Added Windows, Mac and Linux selection tabs.

diff --git a/Engine/Documentation/Source/GettingStarted/DownloadingUnrealEngine/DownloadZip/DownloadZip.INT.udn b/Engine/Documentation/Source/GettingStarted/DownloadingUnrealEngine/DownloadZip/DownloadZip.INT.udn
--- a/Engine/Documentation/Source/GettingStarted/DownloadingUnrealEngine/DownloadZip/DownloadZip.INT.udn
+++ b/Engine/Documentation/Source/GettingStarted/DownloadingUnrealEngine/DownloadZip/DownloadZip.INT.udn
@@ -2,10 +2,14 @@ Availability:NoPublish
 Title:Downloading Unreal as a Zip from GitHub
 Crumbs:%ROOT%, GettingStarted, GettingStarted/DownloadingUnrealEngine
 Description:Downloading the Unreal Engine source code as a Zip file.
+Prereq: GettingStarted/DownloadingUnrealEngine
+Related: Programming/Development/BuildingUnrealEngine
 Version: 4.8
 SkillLevel: Advanced
 
-If you prefer to download our source code as a Zip file from GitHub, this section is for you.
+If you prefer to download our source code as a Zip file from GitHub, this page is for you.
+
+%Globals:osselect_linux%
 
 [OBJECT:ToggleButtonContent]
 	[PARAMLITERAL:category]
