Our Source Skeleton on the left is in a T-Pose while our Target Skeleton is using an A-Pose. *Poses were reversed

diff --git a/Engine/Documentation/Source/Engine/Animation/RetargetingDifferentSkeletons/RetargetingDifferentSkeletons.INT.udn b/Engine/Documentation/Source/Engine/Animation/RetargetingDifferentSkeletons/RetargetingDifferentSkeletons.INT.udn
--- a/Engine/Documentation/Source/Engine/Animation/RetargetingDifferentSkeletons/RetargetingDifferentSkeletons.INT.udn
+++ b/Engine/Documentation/Source/Engine/Animation/RetargetingDifferentSkeletons/RetargetingDifferentSkeletons.INT.udn
@@ -234,7 +234,7 @@ Take for example the Skeleton in the image below on the right.
 
 ![](Retarget2_2.jpg)(w:820)
 
-Our Source Skeleton on the left is in a T-Pose while our Target Skeleton is using an A-Pose. Even though a Rig has been set up for them to share, if we were to retarget animations as is now, we may encounter issues as seen below. 
+Our Source Skeleton on the left is in a A-Pose while our Target Skeleton is using an T-Pose. Even though a Rig has been set up for them to share, if we were to retarget animations as is now, we may encounter issues as seen below. 
 
 ![](Retarget19.png)
 
