#UE4doc #docs Updated documenation per DmitryR's review.

diff --git a/Engine/Documentation/Source/GettingStarted/DownloadingUnrealEngine/DownloadingUnrealEngine4.INT.udn b/Engine/Documentation/Source/GettingStarted/DownloadingUnrealEngine/DownloadingUnrealEngine4.INT.udn
--- a/Engine/Documentation/Source/GettingStarted/DownloadingUnrealEngine/DownloadingUnrealEngine4.INT.udn
+++ b/Engine/Documentation/Source/GettingStarted/DownloadingUnrealEngine/DownloadingUnrealEngine4.INT.udn
@@ -12,7 +12,7 @@ SkillLevel: Advanced
 
 [EXCERPT:Main]
 [PUBLISH:Licensee]
-Unreal Engine 4 is available to be downloaded by licensees directly from Epic's Perforce server. Before you can connect to the server, a secure VPN connection must be established. The login credentials for both the VPN and Perforce server are provided to your team's technical lead by Epic Games in your UDN welcome email. Only one set of login credentials is provided to each team.
+Unreal Engine 4 is available to be downloaded by licensees directly from Epic's Perforce (P4V) server. Before you can connect to the server, a secure VPN connection must be established. The login credentials for both the VPN and Perforce server are provided to your team's technical lead by Epic Games in your UDN welcome email. Only one set of login credentials is provided to each team.
 [/PUBLISH:Licensee]
 [PUBLISH:Rocket]
 
