#UEdoc #docs Reference JIRA ticket UEDOC-1825. Revised Linux tab content per DmitryR's review.

diff --git a/Engine/Documentation/Source/GettingStarted/DownloadingUnrealEngine/DownloadZip/DownloadZip.INT.udn b/Engine/Documentation/Source/GettingStarted/DownloadingUnrealEngine/DownloadZip/DownloadZip.INT.udn
--- a/Engine/Documentation/Source/GettingStarted/DownloadingUnrealEngine/DownloadZip/DownloadZip.INT.udn
+++ b/Engine/Documentation/Source/GettingStarted/DownloadingUnrealEngine/DownloadZip/DownloadZip.INT.udn
@@ -23,7 +23,7 @@ If you prefer to download our source code as a Zip file from GitHub, this page i
 	[/PARAMLITERAL]
 	[PARAM:content]
 	1.	Navigate to our GitHub page.
-	1.	**Click** the **Download ZIP** button (pictured below).
+	1.	**Click** the **Download ZIP** button.
 		![](DownloadZipButton.png)
 		[REGION:warning]
 		The built-in Windows zip utility marks downloaded zip file contents as unsafe to execute.  
@@ -62,7 +62,7 @@ If you prefer to download our source code as a Zip file from GitHub, this page i
 	[/PARAMLITERAL]
 	[PARAM:content]
 	1.	Navigate to our GitHub page.
-	1.	**Click** the **Download ZIP** button (pictured below).
+	1.	**Click** the **Download ZIP** button.
 		![](DownloadZipButton.png)
 	1.	Decompress 'UnrealEngine-(branch_name).zip' onto your hard disk.
 	1.	If you haven't already done so, install the latest version of [Xcode](https://itunes.apple.com/us/app/xcode/id497799835).
@@ -87,19 +87,19 @@ If you prefer to download our source code as a Zip file from GitHub, this page i
 	[PARAM:content]
 
 	[REGION:warning]
-	At this time, we only provide development support for Ubuntu 15.04.
+	Our developers and support teams currently use the latest version of Ubuntu; as a result, we may not be able to provide support for other Linux distributions (including other versions of Ubuntu).
 	[/REGION]
 
 	1.	Navigate to our GitHub page.
-	1.	**Click** the **Download ZIP** button (pictured below).
+	1.	**Click** the **Download ZIP** button.
 		![](DownloadZipButton.png)
 	1.	Decompress 'UnrealEngine-(branch_name).zip' onto your hard disk.
-	1.	Inside the root folder, run **Setup.sh** from the terminal to download the engine's binary content (example pictured below).
+	1.	Inside the root folder, run **Setup.sh** from the terminal to download the engine's binary content.
 
 		![](RunSetupShellScript_Linux.png)
 
 		[REGION:note]
-		*	Our Setup batch file downloads binary content for the engine, installs prerequesites, and sets up Unreal file associations. Downloading engine binaries may take some time because they currently take up three to four gigabytes of disk space. Subsequent checkouts take less time to download.  
+		*	Our Setup shell script downloads binary content for the engine and installs prerequesites. Downloading engine binaries may take some time because they currently take up three to four gigabytes of disk space. Subsequent checkouts take less time to download.  
 		*	If you don't see a ** \*\*SUCCESS\*\* ** message after running this step, refer to the **'BuildThirdParty.log'** file for more information. **'BuildThirdParty.log'** is located in the `[UE4Root]/Engine/Build/BatchFiles/Linux` directory.
 		[/REGION]
 	[/PARAM]
