#UEdoc #docs Reference JIRA ticket UEDOC-1825. Revised linux tab content per SME (DmitryR) review.

diff --git a/Engine/Documentation/Source/GettingStarted/DownloadingUnrealEngine/GitHub/GitHub.INT.udn b/Engine/Documentation/Source/GettingStarted/DownloadingUnrealEngine/GitHub/GitHub.INT.udn
--- a/Engine/Documentation/Source/GettingStarted/DownloadingUnrealEngine/GitHub/GitHub.INT.udn
+++ b/Engine/Documentation/Source/GettingStarted/DownloadingUnrealEngine/GitHub/GitHub.INT.udn
@@ -23,7 +23,7 @@ Forking and cloning our source code from GitHub is a great way to get started wi
 	[/PARAMLITERAL]
 	[PARAM:content]
 	1.	Install [GitHub for Windows](https://windows.github.com/).
-	1.	**Fork** and **Clone** our [repository](https://github.com/EpicGames/UnrealEngine) (pictured below).
+	1.	**Fork** and **Clone** our [repository](https://github.com/EpicGames/UnrealEngine).
 		![](UnrealRepo.png)
 		[REGION:tip]
 		*	Refer to GitHub's [Forking Projects Guide](https://guides.github.com/activities/forking/) to learn how to fork and clone our repository.
@@ -55,7 +55,7 @@ Forking and cloning our source code from GitHub is a great way to get started wi
 	[/PARAMLITERAL]
 	[PARAM:content]
 	1.	Install [GitHub for Mac](https://mac.github.com/).
-	1.	**Fork** and **Clone** our [repository](https://github.com/EpicGames/UnrealEngine) (pictured below).
+	1.	**Fork** and **Clone** our [repository](https://github.com/EpicGames/UnrealEngine).
 		![](UnrealRepo.png)
 		[REGION:tip]
 		*	Refer to GitHub's [Forking Projects Guide](https://guides.github.com/activities/forking/) to learn how to fork and clone our repository.
@@ -83,18 +83,18 @@ Forking and cloning our source code from GitHub is a great way to get started wi
 	[PARAM:content]
 
 	[REGION:warning]
-	At this time, we only provide development support for Ubuntu 15.04.
+	Our developers and support teams currently use the latest version of Ubuntu; as a result, we may not be able to provide support for other Linux distributions (including other versions of Ubuntu).
 	[/REGION]
 
 	1.	Set up [Git](https://help.github.com/articles/set-up-git/).
-	1.	**Fork** and **Clone** our [repository](https://help.github.com/articles/fork-a-repo/) (pictured below).
+	1.	**Fork** and **Clone** our [repository](https://help.github.com/articles/fork-a-repo/).
 		![](UnrealRepo.png)
-	1.	Inside the root folder, run **Setup.sh** from the terminal to download the engine's binary content (example pictured below).
+	1.	Inside the root folder, run **Setup.sh** from the terminal to download the engine's binary content.
 
 		![](RunSetupShellScript_Linux.png)
 
 		[REGION:note]
-		*	Our Setup shell script downloads binary content for the engine, installs prerequesites, and sets up Unreal file associations. Downloading engine binaries may take some time because they currently take up three to four gigabytes of disk space. Subsequent checkouts generally take less time to download.  
+		*	Our Setup shell script downloads binary content for the engine and installs prerequesites. Downloading engine binaries may take some time because they currently take up three to four gigabytes of disk space. Subsequent checkouts generally take less time to download.  
 		*	If you don't see a ** \*\*SUCCESS\*\* ** message after running this step, refer to the **'BuildThirdParty.log'** file for more information. **'BuildThirdParty.log'** is located in the `[UE4Root]/Engine/Build/BatchFiles/Linux` directory.
 		[/REGION]
 	[/PARAM]
