#UE4doc #docs Reference JIRA ticket UEDOC-1823. Addressed SME (DmitryR) review comments and revised linux tab content.

diff --git a/Engine/Documentation/Source/Programming/Development/BuildingUnrealEngine/BuildingUnrealEngine.INT.udn b/Engine/Documentation/Source/Programming/Development/BuildingUnrealEngine/BuildingUnrealEngine.INT.udn
--- a/Engine/Documentation/Source/Programming/Development/BuildingUnrealEngine/BuildingUnrealEngine.INT.udn
+++ b/Engine/Documentation/Source/Programming/Development/BuildingUnrealEngine/BuildingUnrealEngine.INT.udn
@@ -41,15 +41,15 @@ SkillLevel: Advanced
 
 		1.	Load the project into Visual Studio by **double-clicking UE4.sln**.
 
-		1.	Set your solution configuration to **Development Editor** (pictured below).
+		1.	Set your solution configuration to **Development Editor**.
 
 			![](configMenu_updated.png)
 
-		1.	Set your solution platform to **Win64** (pictured below).
+		1.	Set your solution platform to **Win64**.
 
 			![](platform_menu_updated.png)
 
-		1.	**Right-click** the UE4 target and select **Build** (pictured below).
+		1.	**Right-click** the UE4 target and select **Build**.
 
 			![](BuildUE4.png)
 
@@ -57,11 +57,11 @@ SkillLevel: Advanced
 			Depending on your system's specifications, it may take between 10 and 40 minutes to compile the engine from source.
 			[/REGION]
 
-		1.	After compiling the engine, set your startup project to **UE4** (pictured below).
+		1.	After compiling the engine, set your startup project to **UE4**.
 
 			![](SetUE4_StartPrj.png)
 
-		1.	**Right-click** the **UE4** project, **hover over 'Debug'** and **click 'Start New Instance'** to launch the editor (pictured below).
+		1.	**Right-click** the **UE4** project, **hover over 'Debug'** and **click 'Start New Instance'** to launch the editor.
 
 			![](RunCompiledWindowsEditor.png)
 
@@ -69,7 +69,7 @@ SkillLevel: Advanced
 			Alternatively, you can **press the F5 key** on your keyboard to start a new instance of the editor.
 			[/REGION]
 
-		1.	Congratulations! You've compiled and launched the engine from source (result pictured below).
+		1.	Congratulations! You've compiled and launched the engine from source.
 
 			![](CompiledEditor_Windows.png)
 	[/PARAM]
@@ -92,11 +92,11 @@ SkillLevel: Advanced
 		
 		1.	Load the project into XCode by **double-clicking UE4.xcodeproj**.
 
-		1.	To set your build target, select **UE4Editor - Mac** for **My Mac** in the title bar (pictured below).
+		1.	To set your build target, select **UE4Editor - Mac** for **My Mac** in the title bar.
 
 			![](SettingMacTarget.png)
 
-		1.	To build the project, select **Product > Build** (pictured below).
+		1.	To build the project, select **Product > Build**.
 
 			![](BuildingMacEditor.png)
 
@@ -104,11 +104,11 @@ SkillLevel: Advanced
 			Depending on your system's specifications, it may take between 15 and 40 minutes to compile the engine from source.
 			[/REGION]
 		
-		1.	After compiling the engine, select **Product > Run** to launch the editor (pictured below).
+		1.	After compiling the engine, select **Product > Run** to launch the editor.
 		
 			![](RunningMacEditor.png)
 
-		1.	Congratulations! You've compiled and launched the engine from source (result pictured below).
+		1.	Congratulations! You've compiled and launched the engine from source.
 
 			![](CompiledEditor_Mac.png)
 	[/PARAM]
@@ -124,28 +124,32 @@ SkillLevel: Advanced
 	[/PARAMLITERAL]
 	[PARAM:content]
 		[REGION:warning]
-		At this time, we only provide development support for Ubuntu 15.04.
+		Our developers and support teams currently use the latest version of Ubuntu; as a result, we may not be able to provide support for other Linux distributions (including other versions of Ubuntu).
 		[/REGION]
 		
-		1.	Inside the root directory, run **GenerateProjectFiles.sh** from the terminal to set-up your project files (example pictured below). 
+		1.	Inside the root directory, run **GenerateProjectFiles.sh** from the terminal to set-up your project files. 
 
 			![](RunGenPrjFilesShellScript_Linux.png)
 
-		1.	To build the project, run **make** from the terminal (example pictured below).
+		1.	To build the project, run **make** from the terminal.
 
 			![](RunMakeFile_Linux.png)
 
+			[REGION:tip]
+			You can also build the engine with QtCreator, CodeLite, KDevelop, and CMake.
+			[/REGION]
+
 		1.	After compiling the engine, navigate to the editor's binary path by entering ** cd Engine/Binaries/Linux/ ** into the terminal.
 
-		1.	Run **UE4Editor** to launch the editor (pictured below).
+		1.	Run **UE4Editor** to launch the editor.
 
 			![](RunUE4Editor_Linux.png)
 
 			[REGION:warning]
-			If you're unable to create an OpenGL context, you may need to [update](http://www.nvidia.com/Download/index.aspx?lang=en-us) your Nvidia graphics driver. Refer to [Nvidia's website](http://www.nvidia.com/page/home.html) for more information.
-			[/REGION]
+			If you're unable to create an OpenGL context, you may need to update your graphics drivers (refer to Ubuntu's documentation). We recommend using NVidia binary drivers (versions equal to or greater than 343.xx).
+  			[/REGION]
 
-		1.	Congratulations! You've compiled and launched the engine from source (result pictured below).
+		1.	Congratulations! You've compiled and launched the engine from source.
 
 			![](CompiledEditor_Linux.png)
 	[/PARAM]
