#UE4doc #docs Updated linux tab content per SME (DmitryR) review. Reference JIRA ticket UEDOC-2122.

diff --git a/Engine/Documentation/Source/GettingStarted/Installation/InstallingUnrealEngine.INT.udn b/Engine/Documentation/Source/GettingStarted/Installation/InstallingUnrealEngine.INT.udn
--- a/Engine/Documentation/Source/GettingStarted/Installation/InstallingUnrealEngine.INT.udn
+++ b/Engine/Documentation/Source/GettingStarted/Installation/InstallingUnrealEngine.INT.udn
@@ -11,11 +11,6 @@ Related: Support/Builds
 Version: 4.8
 SkillLevel: Beginner
 
-[REGION:warning]
-We currently don't provide a binary installer of Unreal Engine for Linux.  
-Linux users; however, can compile Unreal Engine from source by referring to our [](GettingStarted/DownloadingUnrealEngine) documentation.
-[/REGION]
-
 By the end of this tutorial, you'll be able to:
 
 *	[Create an Epic Games Account _(if you don't already have one)_](#Bookmark1)
@@ -26,7 +21,7 @@ By the end of this tutorial, you'll be able to:
 
 ![](InstallationGoal.png)
 
-_The Epic Games Launcher enables our subscribers to conveniently download and install updates for UE4 (example pictured above)._
+_The Epic Games Launcher enables our subscribers to conveniently download and install updates for UE4._
 
 [REGION:warning]
 If you already have an Epic Games account, skip ahead to learn about [Downloading and Running the Installer](#Bookmark2).
@@ -55,15 +50,15 @@ With a valid Epic Games account, you're able to download our installer (Epic Gam
 	active_button_content
 	[/PARAMLITERAL]
 	[PARAM:content]
-	1.	**Click** the **Download** button to get our installer (pictured below).
+	1.	**Click** the **Download** button to get our installer.
 	
 		![](Download_Screen.png)
 	
-	1.	Once the download is complete, **Run** the installer (pictured below).
+	1.	Once the download is complete, **Run** the installer.
 	
 		![](RunInstaller.png)
 	
-	1.	**Click** the **Install** button when the **Epic Games Launcher Setup** dialog box appears (pictured below).
+	1.	**Click** the **Install** button when the **Epic Games Launcher Setup** dialog box appears.
 	
 		![](RunEpicLauncherSetup.png)
 	[/PARAM]
@@ -80,13 +75,13 @@ With a valid Epic Games account, you're able to download our installer (Epic Gam
 	active_button_content
 	[/PARAMLITERAL]
 	[PARAM:content]
-	1.	**Click** the **Download** button to get our installer (pictured below).
+	1.	**Click** the **Download** button to get our installer.
 	
 		![](Download_Screen.png)
 	
 	1.	**Double-click** the `EpicGamesLauncher.dmg` icon in **Finder** to open a new **Finder Window**.
 
-	1.	**Drag** the **Epic Games Launcher** icon into the **Applications** folder (pictured below).
+	1.	**Drag** the **Epic Games Launcher** icon into the **Applications** folder.
 	
 		![](InstallEpicLaucherMac.png)	
 	[/PARAM]
@@ -105,7 +100,7 @@ With a valid Epic Games account, you're able to download our installer (Epic Gam
 	[PARAM:content]
 	We currently don't provide a binary installer of Unreal Engine for Linux.
 	[REGION:tip]
-	Linux users; however, can compile Unreal Engine from source by referring to our [](GettingStarted/DownloadingUnrealEngine) documentation.
+	Linux users can compile Unreal Engine from source by referring to our [](GettingStarted/DownloadingUnrealEngine) documentation.
 	[/REGION]
 	[/PARAM]
 [/OBJECT]
@@ -114,7 +109,7 @@ With a valid Epic Games account, you're able to download our installer (Epic Gam
 
 (#Bookmark3)
 
-After the set-up program installs the Epic Games Launcher on your machine, use your Epic Game account credentials to login (see below).
+After the set-up program installs the Epic Games Launcher on your machine, use your Epic Game account credentials to login.
 
 ![](EpicSignIn.png)
 
@@ -131,15 +126,15 @@ Space Requirements
 Unreal Engine uses approximately eight gigabytes of disk space for every version of the engine that is installed on a machine. Make sure you have plenty of disk space prior to installing UE4.
 [/REGION]
 
-1.	**Click** the **Unreal Engine** tab inside the Epic Games Launcher Loading Screen (pictured below).
+1.	**Click** the **Unreal Engine** tab inside the Epic Games Launcher Loading Screen.
 
 	![](LauncherLoadScreen.png)
 
-1.	**Click** the **Install Engine** button to download and install the latest version of Unreal Engine (pictured below).
+1.	**Click** the **Install Engine** button to download and install the latest version of Unreal Engine.
 
 	![](InstallEngine.png)
 
-1.	**Read** our [Unreal Engine End User License Agreement](https://www.unrealengine.com/eula) (EULA) prior to **clicking** the **Accept** button (pictured below)
+1.	**Read** our [Unreal Engine End User License Agreement](https://www.unrealengine.com/eula) (EULA) prior to **clicking** the **Accept** button.
 
 	![](AcceptEULALauncher.png)
 
@@ -153,7 +148,7 @@ Unreal Engine uses approximately eight gigabytes of disk space for every version
 
 (#Bookmark5)
 
-After the Epic Games Launcher successfully downloads and installs Unreal Engine, you're ready to **click** the **Launch** button (see below).
+After the Epic Games Launcher successfully downloads and installs Unreal Engine, you're ready to **click** the **Launch** button.
 
 ![](LaunchUE4.png)
 
