#UE4doc #docs:
Removed unnecessary italics from a term

diff --git a/Engine/Documentation/Source/Shared/Glossary/A/Glossary_A.INT.udn b/Engine/Documentation/Source/Shared/Glossary/A/Glossary_A.INT.udn
--- a/Engine/Documentation/Source/Shared/Glossary/A/Glossary_A.INT.udn
+++ b/Engine/Documentation/Source/Shared/Glossary/A/Glossary_A.INT.udn
@@ -15,7 +15,7 @@ Actor
 -->
 
 [EXCERPT:Glossary_Actor]
-An **Actor** is any object that can be placed into a level. Actors are a generic Class that supports 3D transformations such as translation, rotation, and scale. Actors can be created (spawned) and destroyed through gameplay code (C++ or Blueprints). In C++, `AActor` is the base class of all Actors. 
+An **Actor** is any object that can be placed into a level. Actors are a generic Class that supports 3D transformations such as translation, rotation, and scale. Actors can be created (spawned) and destroyed through gameplay code (C++ or Blueprints). In C++, AActor is the base class of all Actors. 
 
 There are several different types of Actors, some examples include: StaticMeshActor, CameraActor, and PlayerStartActor. 
 [/EXCERPT:Glossary_Actor]
