UEDOC-1396

Input and Output have been switched to match the image

diff --git a/Engine/Documentation/Source/Engine/Rendering/Materials/ExpressionReference/ExpressionReference.INT.udn b/Engine/Documentation/Source/Engine/Rendering/Materials/ExpressionReference/ExpressionReference.INT.udn
--- a/Engine/Documentation/Source/Engine/Rendering/Materials/ExpressionReference/ExpressionReference.INT.udn
+++ b/Engine/Documentation/Source/Engine/Rendering/Materials/ExpressionReference/ExpressionReference.INT.udn
@@ -24,9 +24,9 @@ Certain material expressions are parameters, meaning their values can be modifie
 
 1. **Description** - All materials expressions have a common **Desc** property. Text entered in this property will be displayed in the Material Editor just above the expression in the workspace. It can be used for any purpose, but usually serves as a good way to leave short notes about the purpose or function of the expression.
 1. **Title bar** - Displays the name and/or pertinent information about properties of the material expression.
-1. **Outputs** - Links that output the results of the material expression operation.
-1. **Preview** - Displays a preview of the value(s) that are output by the material expression. Updates automatically when realtime update is enabled. Can be manually updated using the Spacebar.
 1. **Inputs** - Links that take in the value(s) to be used by the material expression.
+1. **Preview** - Displays a preview of the value(s) that are output by the material expression. Updates automatically when realtime update is enabled. Can be manually updated using the Spacebar.
+1. **Outputs** - Links that output the results of the material expression operation.
 
 
 ### Expression Types
