UEDOC-2226

The Table of Contents now matches the Table on the page

diff --git a/Engine/Documentation/Source/Engine/Rendering/Materials/HowTo/MaterialHowTo.INT.udn b/Engine/Documentation/Source/Engine/Rendering/Materials/HowTo/MaterialHowTo.INT.udn
--- a/Engine/Documentation/Source/Engine/Rendering/Materials/HowTo/MaterialHowTo.INT.udn
+++ b/Engine/Documentation/Source/Engine/Rendering/Materials/HowTo/MaterialHowTo.INT.udn
@@ -83,6 +83,8 @@ Read over each description to get an understanding of what the How-To will cover
 | **[Use Layered Materials](Engine/Rendering/Materials/HowTo/CreatingLayeredMaterials)** | %Engine/Rendering/Materials/HowTo/CreatingLayeredMaterials:description% |
 | **[Make Materials Shiny](Engine/Rendering/Materials/HowTo/ShinyMaterials)** | Making Materials that are shiny is something can be accomplished in just a few clicks of the mouse. In this How-To, we will cover how you can turn any Material into a Material that shines.|
 | **[Colored Translucent Shadows](Engine/Rendering/Materials/HowTo/ColoredTransluscentShadows)** |Getting shadows to look like they are casting colored shadows into the world can be achieved by using Colored Translucent Shadows. In this How-To, we will cover how you can create Materials that make use of Colored Translucent Shadows. |
+| **[Using a Subsurface Profile in your Materials](Engine/Rendering/Materials/HowTo/Subsurface_Profile)** |The ability to rendering realistic looking human skin is a must for any modern Video Game engine. In this How-To, we will cover what a Subsurface Profile is and how you can use them in your work. |
+| **[Using Subsurface Scattering in your Materials](Engine/Rendering/Materials/HowTo/Subsurface_Scattering)** |The term used to describe the lighting phenomenon where light scatters as it passes through a translucent / semi translucent surface. In this How-To, we will walk you through all you need to know about using Subsurface Scattering in your Materials. |
 
 
 
