UEDOC-1407

Note has been added to say that Target Point Actors is now located under the all classes section in the modes panel

diff --git a/Engine/Documentation/Source/Engine/Actors/TargetPoint/TargetPoint.INT.udn b/Engine/Documentation/Source/Engine/Actors/TargetPoint/TargetPoint.INT.udn
--- a/Engine/Documentation/Source/Engine/Actors/TargetPoint/TargetPoint.INT.udn
+++ b/Engine/Documentation/Source/Engine/Actors/TargetPoint/TargetPoint.INT.udn
@@ -24,6 +24,10 @@ Target Point Actors can be found in the **Modes** panel under the **All Classes*
 [/PARAMLITERAL]
 [/OBJECT]
 
+[region:note]
+Please note that in later version of Unreal Engine 4, **Target Point Actors** will be located under **All Classes** in the **Modes** panel.  
+[/region]
+
 ## Using Target Points
 Target Point Actors can be used for all kinds of things inside Unreal Engine 4. Here is a list of some of the things that you can use Target Point Actors for.
 
