#UE4doc #docs - removed overview header and corrected mirror flag not showing per UEDOC-882

diff --git a/Engine/Documentation/Source/Engine/Basics/Distributions/Distributions.INT.udn b/Engine/Documentation/Source/Engine/Basics/Distributions/Distributions.INT.udn
--- a/Engine/Documentation/Source/Engine/Basics/Distributions/Distributions.INT.udn
+++ b/Engine/Documentation/Source/Engine/Basics/Distributions/Distributions.INT.udn
@@ -8,7 +8,6 @@ Description:Data types that provide flexibility by allowing for constant values,
 
 
 
-## Overview
 
 Distributions are a group of data types that provide flexibility by allowing for constant values, random values within a range, values interpolated along a curve, and values driven by parameters. These are generally used in properties of Particle Systems and nodes within SoundCues.
 
@@ -368,6 +367,15 @@ When selected as the type, the following dialog will be provided for editing the
 |EDVLF_YZ|The Z-axis is locked to the Y-axis value.|
 |EDVLF_XYZ|The Y-and-Z-axes are locked to the X-axis value.  |
 
+The Mirror Flags allow for mirroring the Min/Max values for each component of the value. The following flags are supported for mirroring:
+
+
+|  **MirrorFlags**|**Description**|
+| --- | --- |
+|EDVMF_Same|Use the Max value for Min as well.|
+|EDVMF_Different|Use each value as set.|
+|EDVMF_Mirror|The Min value will be the inverse of the Max (Max * -1).  |
+
 
 As with other curve types, it is recommended that this distribution type be edited via the curve editor.
 
