UEDOC-920

Image has been updated to show UE4 instead of UE3 in Perforce

diff --git a/Engine/Documentation/Source/GettingStarted/DownloadingUnrealEngine/Perforce/Setup/Setup.INT.udn b/Engine/Documentation/Source/GettingStarted/DownloadingUnrealEngine/Perforce/Setup/Setup.INT.udn
--- a/Engine/Documentation/Source/GettingStarted/DownloadingUnrealEngine/Perforce/Setup/Setup.INT.udn
+++ b/Engine/Documentation/Source/GettingStarted/DownloadingUnrealEngine/Perforce/Setup/Setup.INT.udn
@@ -66,7 +66,7 @@ will be used to download and store the engine builds.
 The P4V client is the current client provided by Perforce. It provides access to versioned files through a 
 graphical interface and also includes tools for merging and visualizing code evolution.
 
-![p4v.jpg](p4v.jpg)(w:900 h:591)
+![](P4VUpdated.png)(w:900 h:591)
 
 The full P4V installer can be downloaded from the [Perforce Downloads: The Essentials](http://www.perforce.com/perforce/downloads/index.html) page.
 
