#UE4doc #docs Removed bolded links and updated documentation to meet Instructional Guide v1.1 requirements.

diff --git a/Engine/Documentation/Source/Engine/QuickStart/4/LDQS_4.INT.udn b/Engine/Documentation/Source/Engine/QuickStart/4/LDQS_4.INT.udn
--- a/Engine/Documentation/Source/Engine/QuickStart/4/LDQS_4.INT.udn
+++ b/Engine/Documentation/Source/Engine/QuickStart/4/LDQS_4.INT.udn
@@ -3,7 +3,7 @@ Title:4 - Placing Actors in the Level
 Crumbs:%ROOT%, GettingStarted, Engine/QuickStart
 Description:Here we will begin placing and manipulating Actors in our level to build out our environment. 
 SkillLevel:Beginner
-Version:4.9 
+Version:4.8 
 Related:Engine/UI/LevelEditor/Modes/PlaceMode
 Related:GettingStarted/HowTo/ManipulatingActors
 Related:Engine/UI/LevelEditor/Details
@@ -29,13 +29,13 @@ Related:Engine/Rendering/LightingAndShadows
 
 %Steps%
 
-In this step we will begin placing [**Actors**](Engine/Actors) (which are things like lights or geometry for example) into our empty level. 
-We will cover the two most common ways of adding Actors to the level, one of which is through [**Place Mode**](Engine/UI/LevelEditor/Modes/PlaceMode) and the other is through the [**Content Browser**](Engine/Content/Browser). 
+In this step we will begin placing [Actors](Engine/Actors) (which are things like lights or geometry for example) into our empty level. 
+We will cover the two most common ways of adding Actors to the level, one of which is through [Place Mode](Engine/UI/LevelEditor/Modes/PlaceMode) and the other is through the [Content Browser](Engine/Content/Browser). 
 After completing this step, you will know how to place Actors inside your own levels and can begin manipulating those Actors to create your environment. 
 
 ## Steps
 
-1. In the **Modes Panel**, with **Place Mode** enabled, click the [**BSP**](Engine/Actors/Brushes) category and select the **Box**. 
+1. In the **Modes Panel**, with **Place Mode** enabled, click the [BSP](Engine/Actors/Brushes) category and select the **Box**. 
 
 	![](ModesPanel.png)
 
@@ -47,7 +47,7 @@ After completing this step, you will know how to place Actors inside your own le
 
 	![](place2.png)
 
-1. With the **Box** still selected, in the [**Details**](Engine/UI/LevelEditor/Details) panel (lower-right window of the editor), set **Location** and **Rotation** all to **0**. 
+1. With the **Box** still selected, in the [Details](Engine/UI/LevelEditor/Details) panel (lower-right window of the editor), set **Location** and **Rotation** all to **0**. 
 
 1. Set the **Scale** to **4 x 4 x 0.1**. 
 
@@ -59,29 +59,29 @@ After completing this step, you will know how to place Actors inside your own le
 
 	![](placeLight.png)
 	
-1. On the on the [**Translation Tool**](Engine/UI/LevelEditor/Viewports/ViewportControls), click and drag the Z-Axis (Blue) gizmo up away from the surface of the floor.
+1. On the on the [Translation Tool](Engine/UI/LevelEditor/Viewports/ViewportControls), click and drag the Z-Axis (Blue) gizmo up away from the surface of the floor.
 
 	![](moveLight.png)
 
 	If the **Directional Light** becomes unselected, you can re-select it by **Left-clicking** on it in the **Level Viewport**. 
 
-1. In the **Modes Panel**, select the **Visual Effects** tab and drag-and-drop an [**Atmospheric Fog**](Engine/Actors/FogEffects/AtmosphericFog) to the level. 
+1. In the **Modes Panel**, select the **Visual Effects** tab and drag-and-drop an [Atmospheric Fog](Engine/Actors/FogEffects/AtmosphericFog) to the level. 
 
 	![](placeFog.png)
 	
 	The **Atmospheric Fog** Actor will add a basic sky to the level and the level will become illuminated instead of dark. 
 
-1. From the **Modes Panel** on the **Basic** tab, drag-and-drop a [**Player Start**](Engine/Actors/PlayerStart) into the level. 
+1. From the **Modes Panel** on the **Basic** tab, drag-and-drop a [Player Start](Engine/Actors/PlayerStart) into the level. 
 
 	![](placeStart.png)
 
-1. From the **Modes Panel** on the **Volumes** tab, drag-and-drop a [**Lightmass Importance Volume**](Engine/Rendering/LightingAndShadows/Lightmass/Basics) into the level. 
+1. From the **Modes Panel** on the **Volumes** tab, drag-and-drop a [Lightmass Importance Volume](Engine/Rendering/LightingAndShadows/Lightmass/Basics) into the level. 
 
 	![](placeVolume.png)
 
 	The **Lightmass Importance Volume** is used to control and concentrate lighting and shadowing affects within the volume. When placing the Lightmass Importance Volume in the level, the default size of the volume does not cover our playable area so we will need to scale it up. 
 	
-1. Inside the **Level Viewport**, press **R** to switch to the [**Scale Tool**](Engine/UI/LevelEditor/Viewports/ViewportControls#transformcontrols). 
+1. Inside the **Level Viewport**, press **R** to switch to the [Scale Tool](Engine/UI/LevelEditor/Viewports/ViewportControls#transformcontrols). 
 
 1. Click and drag the white box in the center of the **Scale Tool** so that the **Lightmass Importance Volume** encapsulates the floor. 
 
@@ -97,7 +97,7 @@ After completing this step, you will know how to place Actors inside your own le
 
 	![](addChair.png)
 
-1. With the **SM_Chair** selected inside the **Level Viewport**, press **E** to access the [**Rotation Tool**](Engine/UI/LevelEditor/Viewports/ViewportControls#transformcontrols). 
+1. With the **SM_Chair** selected inside the **Level Viewport**, press **E** to access the [Rotation Tool](Engine/UI/LevelEditor/Viewports/ViewportControls#transformcontrols). 
 	
 	![](Rotate1.png)(w:440)
 
diff --git a/Engine/Documentation/Source/Engine/QuickStart/5/LDQS_5.INT.udn b/Engine/Documentation/Source/Engine/QuickStart/5/LDQS_5.INT.udn
--- a/Engine/Documentation/Source/Engine/QuickStart/5/LDQS_5.INT.udn
+++ b/Engine/Documentation/Source/Engine/QuickStart/5/LDQS_5.INT.udn
@@ -3,7 +3,7 @@ Title:5 - Editing Placed Actors
 Crumbs:%ROOT%, GettingStarted, Engine/QuickStart
 Description:This this step we will take the Actors we placed and edit their properties to give our level a more customized look. 
 SkillLevel:Beginner
-Version:4.9
+Version:4.8
 Related:Engine/Rendering/Materials
 Related:Engine/Rendering/LightingAndShadows/LightTypes
 Related:Engine/Content/FBX/ImportOptions
@@ -27,8 +27,8 @@ Related:
 
 %Steps%
 
-With several different **Actors** placed inside our level, the next step involves [**Editing Actor Properites**](GettingStarted/HowTo/DetailsPanel) which can change the look or the way the Actor functions in the level, giving us a more customized looking level. 
-We will start with editing the properties of our **Directional Light Actor** then shift our focus to applying [**Materials**](Engine/Rendering/Materials) to some of the [**Static Mesh Actors**](Engine/Actors/StaticMeshActor) that you have placed in your level. 
+With several different **Actors** placed inside our level, the next step involves [Editing Actor Properites](GettingStarted/HowTo/DetailsPanel) which can change the look or the way the Actor functions in the level, giving us a more customized looking level. 
+We will start with editing the properties of our **Directional Light Actor** then shift our focus to applying [Materials](Engine/Rendering/Materials) to some of the [Static Mesh Actors](Engine/Actors/StaticMeshActor) that you have placed in your level. 
 
 Once you have finished this step, you will have seen where to access and modify the properties of Actors so that you can begin editing and experimenting with different settings inside your own levels. 
 
diff --git a/Engine/Documentation/Source/Engine/QuickStart/6/LDQS_6.INT.udn b/Engine/Documentation/Source/Engine/QuickStart/6/LDQS_6.INT.udn
--- a/Engine/Documentation/Source/Engine/QuickStart/6/LDQS_6.INT.udn
+++ b/Engine/Documentation/Source/Engine/QuickStart/6/LDQS_6.INT.udn
@@ -3,7 +3,7 @@ Title:6 - Running the Build Process
 Crumbs:%ROOT%, GettingStarted, Engine/QuickStart
 Description:With a number of Actors placed and customized, one of the final steps for creating any level is the Build Process.
 SkillLevel:Beginner
-Version:4.9
+Version:4.8
 
 [VAR:Steps]
 [OBJECT:Navigation]
diff --git a/Engine/Documentation/Source/Engine/QuickStart/7/LDQS_7.INT.udn b/Engine/Documentation/Source/Engine/QuickStart/7/LDQS_7.INT.udn
--- a/Engine/Documentation/Source/Engine/QuickStart/7/LDQS_7.INT.udn
+++ b/Engine/Documentation/Source/Engine/QuickStart/7/LDQS_7.INT.udn
@@ -3,7 +3,7 @@ Title:7 - On Your Own!
 Crumbs:%ROOT%, GettingStarted, Engine/QuickStart
 Description:Using the methods provided during this guide, polish up your level with additional content and read up on related topics. 
 SkillLevel:Beginner
-Version:4.9
+Version:4.8
 
 [VAR:Steps]
 [OBJECT:Navigation]
@@ -39,23 +39,23 @@ Using the methods that were provided during the course of this guide, try do to
 * Add different kinds of lights with different colors. 
 * Use different Materials on some of your Actors. 
 
-For more information on the topics covered in this quick start guide, and across the entire editor, see the [**Unreal Editor Manual**](Engine).
+For more information on the topics covered in this quick start guide, and across the entire editor, see the [Unreal Editor Manual](Engine).
 
 As for specifics covered in this quick start:
 
-* For more information on the Level Editor, see: [**Level Editor**](Engine\UI\LevelEditor)
+* For more information on the Level Editor, see: [Level Editor](Engine\UI\LevelEditor)
 
-* For more information on Viewports, see: [**Viewports**](Engine\UI\LevelEditor\Viewports)
+* For more information on Viewports, see: [Viewports](Engine\UI\LevelEditor\Viewports)
 
-* For more information on the Editing Modes available in Unreal Editor, 4 see: [**Level Editor Modes**](Engine\UI\LevelEditor\Modes)
+* For more information on the Editing Modes available in Unreal Editor, 4 see: [Level Editor Modes](Engine\UI\LevelEditor\Modes)
 
-* For more information on the **Content Browser**, see: [**Content Browser**](Engine\Content\Browser)
+* For more information on the **Content Browser**, see: [Content Browser](Engine\Content\Browser)
 
-* For more information on the **Details** Panel, see: [**Details Panel**](Engine\UI\LevelEditor\Details)
+* For more information on the **Details** Panel, see: [Details Panel](Engine\UI\LevelEditor\Details)
 
-* For more information on Building, see: [**Lightmass**](Engine\Rendering\LightingAndShadows\Lightmass)
+* For more information on Building, see: [Lightmass](Engine\Rendering\LightingAndShadows\Lightmass)
 
-* For more information on Lighting, see: [**Lighting Quick Start Guide**](Engine/Rendering/LightingAndShadows/QuickStart)
+* For more information on Lighting, see: [Lighting Quick Start Guide](Engine/Rendering/LightingAndShadows/QuickStart)
 
 %Steps%
 
