#UE4doc #docs Updated document to meet Instructional Guide v1.1 requirements. Updated image callouts and removed horizontal content dividers.

diff --git a/Engine/Documentation/Source/Programming/QuickStart/1/QuickStart_1.INT.udn b/Engine/Documentation/Source/Programming/QuickStart/1/QuickStart_1.INT.udn
--- a/Engine/Documentation/Source/Programming/QuickStart/1/QuickStart_1.INT.udn
+++ b/Engine/Documentation/Source/Programming/QuickStart/1/QuickStart_1.INT.udn
@@ -3,7 +3,7 @@ Title:1 - Essential Project Setup
 Crumbs:%ROOT%, Programming, Programming/QuickStart
 Description:Creating a new project with the specifications we need
 SkillLevel: Beginner
-Version: 4.9
+Version: 4.8
 
 [VAR:Steps]
 [OBJECT:Navigation]
@@ -23,7 +23,7 @@ Version: 4.9
 %Steps%
 
 [REGION:note]
-You should have a copy of **Visual Studio 2013** for Deskstop (Preferably Community or Professional) installed before starting this tutorial.
+You should have a copy of **Visual Studio 2013** for Desktop (Preferably Community or Professional Editions) installed before starting this tutorial.
 [/REGION]
 
 1. Open **Unreal Engine** from the Launcher. The [](Engine/Basics/Projects/Browser) will appear.
@@ -32,8 +32,6 @@ You should have a copy of **Visual Studio 2013** for Deskstop (Preferably Commun
 
 	[![](CreateProject.png)](CreateProject.png)
 
-----
-
 The **Unreal Editor** will now open our new project. **Visual Studio** will also open and load the solution file that our project has created.
 
 %Steps%
diff --git a/Engine/Documentation/Source/Programming/QuickStart/2/QuickStart_2.INT.udn b/Engine/Documentation/Source/Programming/QuickStart/2/QuickStart_2.INT.udn
--- a/Engine/Documentation/Source/Programming/QuickStart/2/QuickStart_2.INT.udn
+++ b/Engine/Documentation/Source/Programming/QuickStart/2/QuickStart_2.INT.udn
@@ -3,7 +3,7 @@ Title:2 - Create a C++ Class
 Crumbs:%ROOT%, Programming, Programming/QuickStart
 Description:Creating a C++ class
 SkillLevel: Beginner
-Version: 4.9
+Version: 4.8
 
 [VAR:Steps]
 [OBJECT:Navigation]
@@ -35,8 +35,6 @@ Version: 4.9
  
 	[![](NameActorClass.png)](NameActorClass.png)
 
-----
-
 Now that we have created a C++ class, we can switch to **Visual Studio** to program it. `FloatingActor.cpp` will be opened for us automatically, and **Unreal Engine** will automatically compile and reload the code with our new class.
 
 %Steps%
diff --git a/Engine/Documentation/Source/Programming/QuickStart/3/QuickStart_3.INT.udn b/Engine/Documentation/Source/Programming/QuickStart/3/QuickStart_3.INT.udn
--- a/Engine/Documentation/Source/Programming/QuickStart/3/QuickStart_3.INT.udn
+++ b/Engine/Documentation/Source/Programming/QuickStart/3/QuickStart_3.INT.udn
@@ -3,7 +3,7 @@ Title:3 - Write and Compile C++ Code
 Crumbs:%ROOT%, Programming, Programming/QuickStart
 Description:Writing a C++ class and compiling it with Unreal Engine
 SkillLevel: Beginner
-Version: 4.9
+Version: 4.8
 
 [VAR:Steps]
 [OBJECT:Navigation]
@@ -60,11 +60,11 @@ Version: 4.9
 
 			[![](CompileFromVS.png)](CompileFromVS.png)
 	
-			(compiling from **Visual Studio** pictured above)
+			_Compiling from **Visual Studio**._
 
 			[![](CompileFromEditor.png)](CompileFromEditor.png)
 
-			(compiling from the **Unreal Editor** pictured above)
+			_Compiling from the **Unreal Editor**._
 	[/PARAM]
 [/OBJECT]
 
@@ -87,17 +87,17 @@ Version: 4.9
 		
 			![](ProjectBuild_mac.png)
 			
-			(compiling using XCode's Debug Configuration pictured above)
+			_Compiling using XCode's Debug Configuration._
 
 		*	We can also compile using XCode's Development Configuration by clicking the **Product > Build For > Profiling** option.
 
 			![](ProjectBuildProfile_mac.png)
 
-			(compiling using XCode's Development Configuration pictured above)
+			_Compiling using XCode's Development Configuration._
 
 			[![](CompileFromEditor.png)](CompileFromEditor.png)
 
-			(compiling from the **Unreal Editor** pictured above)
+			_Compiling from the **Unreal Editor**._
 
 		[REGION:note]
 		When running the binary Editor, it is important to add the -game flag if you rebuilt your project in any **Uncooked** configuration, and the -debug flag if you rebuilt your project in any **Debug** configuration.
@@ -105,8 +105,6 @@ Version: 4.9
 	[/PARAM]
 [/OBJECT]
 
-----
-
 We're now ready to create objects in **Unreal Editor** based on our code. For reference, all code used on this page is included below.
 
 [OBJECT:Section]
diff --git a/Engine/Documentation/Source/Programming/QuickStart/4/QuickStart_4.INT.udn b/Engine/Documentation/Source/Programming/QuickStart/4/QuickStart_4.INT.udn
--- a/Engine/Documentation/Source/Programming/QuickStart/4/QuickStart_4.INT.udn
+++ b/Engine/Documentation/Source/Programming/QuickStart/4/QuickStart_4.INT.udn
@@ -3,7 +3,7 @@ Title:4 - Test Your Code
 Crumbs:%ROOT%, Programming, Programming/QuickStart
 Description:Creating a test object in your level
 SkillLevel: Beginner
-Version: 4.9
+Version: 4.8
 
 [VAR:Steps]
 [OBJECT:Navigation]
@@ -43,8 +43,6 @@ Version: 4.9
 
 	[![](MovingCone.png)](MovingCone.png)
 
-----
-
 [OBJECT:Section]
     [PARAMLITERAL:id]
     code
diff --git a/Engine/Documentation/Source/Programming/QuickStart/5/QuickStart_5.INT.udn b/Engine/Documentation/Source/Programming/QuickStart/5/QuickStart_5.INT.udn
--- a/Engine/Documentation/Source/Programming/QuickStart/5/QuickStart_5.INT.udn
+++ b/Engine/Documentation/Source/Programming/QuickStart/5/QuickStart_5.INT.udn
@@ -3,7 +3,7 @@ Title:5 - On Your Own!
 Crumbs:%ROOT%, Programming, Programming/QuickStart
 Description:Ideas for further exploration!
 SkillLevel: Beginner
-Version: 4.9
+Version: 4.8
 
 [VAR:Steps]
 [OBJECT:Navigation]
@@ -22,8 +22,6 @@ Version: 4.9
 
 %Steps%
 
-----
-
 Using what you have learned, try to do the following:
 
 * Add a [Particle System Component](Engine/Rendering/ParticleSystems) to your **FloatingActor**. Some pre-built Particle Systems are already included in your project.
@@ -35,6 +33,4 @@ As for the specifics covered in this tutorial:
 * For more information on **Actors**, see the [Actor](API:AActor) manual page.
 * For further tutorials, see the [%Programming/Tutorials:title%](Programming/Tutorials) page.
 
-----
-
 %Steps%
diff --git a/Engine/Documentation/Source/Programming/QuickStart/QuickStart.INT.udn b/Engine/Documentation/Source/Programming/QuickStart/QuickStart.INT.udn
--- a/Engine/Documentation/Source/Programming/QuickStart/QuickStart.INT.udn
+++ b/Engine/Documentation/Source/Programming/QuickStart/QuickStart.INT.udn
@@ -41,8 +41,6 @@ Version: 4.8
 [/OBJECT]
 [/VAR]
 
-----
-
 [OBJECT:EmbeddedVideo]
 	[PARAMLITERAL:width]
 	640
