#UE4doc #docs - updated images by removing embedded text

diff --git a/Engine/Documentation/Source/Engine/Landscape/Landscape.INT.udn b/Engine/Documentation/Source/Engine/Landscape/Landscape.INT.udn
--- a/Engine/Documentation/Source/Engine/Landscape/Landscape.INT.udn
+++ b/Engine/Documentation/Source/Engine/Landscape/Landscape.INT.udn
@@ -191,14 +191,26 @@ Since the data is being stored in Textures, the standard Texture streaming syste
 
 Since the X and Y slopes of the Landscape are being stored, the entire high-resolution (non-LOD'd) normal data is available for lighting  calculations. 
 
+**Left:** Landscape LODs
+
+**Right:** Full Res Normals
+
 ![Landscape_HighRes_Normals.jpg](Landscape_HighRes_Normals.jpg)(w:900 h:272)
 
 This makes it possible to always use the highest resolution of the terrain for per-pixel lighting even on distant components that have been LOD'd out.
 
+**Left:** Simple Vertex Lighting
+
+**Right:** 
+
 ![Landscape_HighRes_Lighting.jpg](Landscape_HighRes_Lighting.jpg)(w:900 h:272)
 
 When this high-resolution normal data is combined with detail normal maps, Landscape terrains can achieve extremely detailed lighting with very little overhead.
 
+**Left:** Geometry Normals Only
+
+**Right:** With Detail Normals
+
 ![Landscape_Normals_Comparison.jpg](Landscape_Normals_Comparison.jpg)(w:900 h:272)
 
 ### PhysX Collision
