#UE4doc #docs - Updated images in collision properties and instructions to correspond with images. fixes for UEDOC-1345

diff --git a/Engine/Documentation/Source/GettingStarted/HowTo/DetailsPanel/DetailsPanel.INT.udn b/Engine/Documentation/Source/GettingStarted/HowTo/DetailsPanel/DetailsPanel.INT.udn
--- a/Engine/Documentation/Source/GettingStarted/HowTo/DetailsPanel/DetailsPanel.INT.udn
+++ b/Engine/Documentation/Source/GettingStarted/HowTo/DetailsPanel/DetailsPanel.INT.udn
@@ -139,17 +139,17 @@ In this step, we will update the **Collision** for the Actor so that it generate
 
 	This will now fire an event when the sphere hits something. 
 
-1.	Off the **OnActorHit** node, drag off the **Other Actor** pin and search for **cast to mycharacter** then select it in the window.  
+1.	Off the **OnActorHit** node, drag off the **Other Actor** pin and search for **cast to thirdpersoncharacter** then select it in the window.  
 
 	![](Details21.png)
 
-	Here we are saying, when the Sphere hits something, is the **Other Actor** we hit **MyCharacter** (the player character). 
+	Here we are saying, when the Sphere hits something, is the **Other Actor** we hit **ThirdPersonCharacter** (the player character). 
 
-1.	Drag off the **As My Character** pin of the **Cast** node and search for and add the **Destroy Actor** node. 
+1.	Drag off the **As Third Person Character** pin of the **Cast** node and search for and add the **Destroy Actor** node. 
 
 	![](Details22.png)
 
-	If **MyCharacter** was the **Other Actor** hit, then destroy **MyCharacter** (killing the player character). 
+	If **ThirdPersonCharacter** was the **Other Actor** hit, then destroy **ThirdPersonCharacter** (killing the player character). 
 
 1.	Click the **Compile** button in the upper-left corner of the window to finalize the Blueprint, then close the window. 
 
