#UE4Doc #docs - Updated images with correct box colors (green > yellow) and update the gif to match latest UI changes per UEDOC-1407.

diff --git a/Engine/Documentation/Source/Engine/Actors/TargetPoint/TargetPoint.INT.udn b/Engine/Documentation/Source/Engine/Actors/TargetPoint/TargetPoint.INT.udn
--- a/Engine/Documentation/Source/Engine/Actors/TargetPoint/TargetPoint.INT.udn
+++ b/Engine/Documentation/Source/Engine/Actors/TargetPoint/TargetPoint.INT.udn
@@ -6,7 +6,7 @@ Related: Gameplay/HowTo/SpawnAndDestroyActors/Blueprints
 
 ![](Target_Point_Actors.png)
 
-## Overview
+
 When creating worlds for games you sometimes need the ability to place and spawn items in the world for the player to interact with. **Target Point Actors** do just this; they give you a generic point in the world that you can spawn items from. If you are familiar with other 3D applications like 3Ds Max or Maya, Target Point Actors are very similar to dummy Actors that you find in those programs. 
 
 ## Placing Target Point Actors
