#UE4doc #docs:
Removed italics from a few terms
Bolded a term
Replaced a picture to reflect UI changes

diff --git a/Engine/Documentation/Source/Engine/Actors/Brushes/Brushes.INT.udn b/Engine/Documentation/Source/Engine/Actors/Brushes/Brushes.INT.udn
--- a/Engine/Documentation/Source/Engine/Actors/Brushes/Brushes.INT.udn
+++ b/Engine/Documentation/Source/Engine/Actors/Brushes/Brushes.INT.udn
@@ -225,7 +225,7 @@ Brushes are created using the **Geometry** tab in the **Place Mode** panel:
 |**Step Height**|Height of each stair from top to bottom.|
 |**Step Thickness**|Increases the thickness of the staircase as a whole.|
 |**Num Steps Per 360**|Determines the number of steps required to make one complete revolution.|
-|**Num Steps**|Number of steps in the staircase. Each time this reaches the value (or multiple) of _Num Steps per 360_, the staircase will wrap around itself again.|
+|**Num Steps**|Number of steps in the staircase. Each time this reaches the value (or multiple) of Num Steps per 360, the staircase will wrap around itself again.|
 |**Add to First Step**|Adds the provided height to the first step, effectively raising the height of the entire staircase.|
 |**Sloped Ceiling**|This causes the ceiling to be sloped rather than tiered.|
 |**Sloped Floor**|Effectively turns the floor into a spiraling ramp, instead of a traditional staircase.|
@@ -296,7 +296,7 @@ A **Subtractive** Brush is a hollow, carved-out space. This is the type of Brush
 
 ### Advanced Properties
 
-Clicking the ![Advanced](button_advanced.png) button at the bottom of the Brush Settings exposes the advanced Brush properties:
+Clicking the ![Advanced](button_advanced.png) button at the bottom of the **Brush Settings** exposes the advanced Brush properties:
 
 ![](BrushDetailsAdvanced.png)
 
@@ -400,8 +400,8 @@ The **Lighting** section allows you to change a variety of important light-centr
 |---|---|
 |**Lightmap Resolution**|Essentially allows for adjustment of the shadows across this surface. The lower the number, the tighter the shadows.|
 |[REGION:tablesection]Lightmass Settings[/REGION]||
-|**Use Two Sided Lighting**|If _true_, this surface will receive light on both the positve and negative side of each polygon.|
-|**Shadow Indirect Only**|If _true_, this surface will only create shadows from indirect lighting.|
+|**Use Two Sided Lighting**|If true, this surface will receive light on both the positve and negative side of each polygon.|
+|**Shadow Indirect Only**|If true, this surface will only create shadows from indirect lighting.|
 |**Diffuse Boost**|Scales the amount of influnce diffuse color will have on indirect lighting.|
 |**Fully Occluded Samples Fraction**|Fraction of samples along this surface that must be occluded before it is considered to be occluded from indirect lighting calculations.|
 
