#UE4doc #docs:
Removed Overview
Added engine version field
Changed text to reflect that Camera is now in the 'All Classes' part of the Modes panel instead of the Basic part

diff --git a/Engine/Documentation/Source/Engine/Actors/CameraActors/CameraActors.INT.udn b/Engine/Documentation/Source/Engine/Actors/CameraActors/CameraActors.INT.udn
--- a/Engine/Documentation/Source/Engine/Actors/CameraActors/CameraActors.INT.udn
+++ b/Engine/Documentation/Source/Engine/Actors/CameraActors/CameraActors.INT.udn
@@ -2,10 +2,10 @@ Availability: Public
 Title:Camera Actors
 Crumbs: %ROOT%, Engine, Engine/Actors
 Description:Guide to placing and using Camera Actors within the Editor.
+version: 4.9
 
 ![](CameraActorHeader.png)
 
-## Overview
 
 Just about every project will require a way to view it, and the most basic way to do this is through **Camera Actors**. These are primarily used in combination with **Matinee**, but can be used with Blueprints or C++ to create a fixed camera game, security cameras, or even short cut scenes using Camera Anims.
 
@@ -13,7 +13,7 @@ If you select a Camera Actor (or a Blueprint that contains a Camera Component),
 
 ## Placing a Camera Actor
 
-The Camera Actor can be found in the **Modes** panel under the **Basic** category. To place one into the world, simply drag it from the **Modes** panel into the game world.
+The Camera Actor can be found in the **Modes** panel under the **All Classes** category. To place one into the world, drag it from the **Modes** panel into the game world.
 
 [OBJECT:EmbeddedVideo]
 [PARAMLITERAL:width]
