#UE4doc #docs:
Added engine version field

diff --git a/Engine/Documentation/Source/Engine/Actors/Brushes/Brushes.INT.udn b/Engine/Documentation/Source/Engine/Actors/Brushes/Brushes.INT.udn
--- a/Engine/Documentation/Source/Engine/Actors/Brushes/Brushes.INT.udn
+++ b/Engine/Documentation/Source/Engine/Actors/Brushes/Brushes.INT.udn
@@ -5,6 +5,8 @@ Description:Guide to using Brushes to create level geometry in Unreal Editor.
 Related: Engine/UI/LevelEditor
 Related: Engine/UI/LevelEditor/Modes
 Related: Engine/Actors/Volumes
+version: 4.9
+
 
 [VAR:playlistnav]
 [OBJECT:PlaylistNavigation]
