#UE4doc #docs:
Added engine version field
Changed the font settings for multiple terms
Changed text to reflect that the atmospheric fog actor can be accessed from the Modes panel
Replaced picture reflect getting the atmospheric fog actor from the Modes panel

diff --git a/Engine/Documentation/Source/Engine/Actors/FogEffects/AtmosphericFog/AtmosphericFog.INT.udn b/Engine/Documentation/Source/Engine/Actors/FogEffects/AtmosphericFog/AtmosphericFog.INT.udn
--- a/Engine/Documentation/Source/Engine/Actors/FogEffects/AtmosphericFog/AtmosphericFog.INT.udn
+++ b/Engine/Documentation/Source/Engine/Actors/FogEffects/AtmosphericFog/AtmosphericFog.INT.udn
@@ -2,6 +2,7 @@ Availability: Public
 Title:Atmospheric Fog User Guide
 Crumbs: %ROOT%, Engine, Engine/Actors, Engine/Actors/FogEffects
 Description:Atmospheric Fog Actors provide a realistic sense of atmosphere, air density, and light scattering through atmospheric media.
+version: 4.9
 
 [VAR:Topic]
 [OBJECT:Topic]
@@ -49,9 +50,9 @@ Description:Atmospheric Fog Actors provide a realistic sense of atmosphere, air
 
 [TOC (start:2 end:3)]
 
-Atmospheric Fog gives an approximation of light scattering through a planetary atmosphere. This can give your outdoor levels a much more realistic look. This total effect includes the following:
+**Atmospheric Fog** gives an approximation of light scattering through a planetary atmosphere. This can give your outdoor levels a much more realistic look. This total effect includes the following:
 
-* The dominant directional light in your level will receive a "sun disc" effect in the sky. This will be placed infinitely far away, opposite the light's direction.
+* The dominant directional light in your level will receive a sun disc effect in the sky. This will be placed infinitely far away, opposite the light's direction.
 * A sky color that will vary depending on the altitude of the sun (or put another way, how close the dominant directional light's vector gets to being parallel with the ground).
 * Control over scattering and decay settings, allowing for full control of your atmospheric density.
 
@@ -62,7 +63,7 @@ Best of all, the settings for this Actor are easy to use and quick to set up.
 
 ## Adding Atmospheric Fog to a Level
 
-Adding Atmospheric Fog to your level is as easy as adding its Actor from the Class Viewer:
+Adding Atmospheric Fog to your level is as easy as adding its Actor from the **Modes** panel under the **Visual Effects** section:
 
 This may be done via drag-and-drop:
 
@@ -79,7 +80,7 @@ The Atmospheric Fog Actor contains the following properties:
 | **Sun Multiplier** | This is an overall multiplier for the directional light brightness. It brightens both the sky and fog color. |
 | **Fog Multiplier** | This multiplier only affects the fog color, not the directional light. |
 | **Density Multiplier** | This control factor only affects the fog density, not the directional light. |
-| **Density Offset** | This offset value controls fog opacity; its effective range is from -1...1. |
+| **Density Offset** | This offset value controls fog opacity; its effective range is from -1 to 1. |
 | **Distance Scale** | Controls a distance factor. The default value of 1 assumes a 1:1 scale of Unreal Units to centimeters. This makes for a very small world. As your world size increases, you will want to increase this value accordingly. Larger values cause changes in the fog attenuation to occur more quickly. |
 | **Altitude Scale** | Controls the scale along the z axis only. The default value is 1. |
 | **Ground Color Additive** | Currently, the irradiance color is based on the scene color. This can cause blackening in shadowed areas. To combat this effect, this parameter allows you to add a small boost (i.e. ~0.05) to make the irradiance color more visible in dark areas. |
@@ -108,9 +109,9 @@ The Atmospheric Fog Actor contains the following properties:
 [REGION:simpletable]
 	|**Property** | **Description** |
 	| --- | --- |
-	| *Density Height* |This value controls fog density decay height. Lower values cause the fog to be denser, while higher values thin the fog, causing less scatter. Values range from [0..1] |
-	| *Max Scattering Num* | This sets a limit on the number of scattering calculations, from 1 (single scattering) to 4 (quadruple scattering). |
-	| *Inscatter Altitude Sample* | Number of different altitudes at which to sample inscatter color (Size of 3d texture dimension; 1-32). |
+	| **Density Height** |This value controls fog density decay height. Lower values cause the fog to be denser, while higher values thin the fog, causing less scatter. Values range from [0..1] |
+	| **Max Scattering Num** | This sets a limit on the number of scattering calculations, from 1 (single scattering) to 4 (quadruple scattering). |
+	| **Inscatter Altitude Sample** | Number of different altitudes at which to sample inscatter color (Size of 3d texture dimension; 1-32). |
 [/REGION]
 [/EXCERPT:PrecomputeParams]
 End Includes----->
