#UE4doc #docs:
Changed font settings for a few terms
Changed text to refer the user to the Modes panel instead of the Class Viewer for creating actors.

diff --git a/Engine/Documentation/Source/Engine/Actors/FogEffects/HeightFog/HeightFog.INT.udn b/Engine/Documentation/Source/Engine/Actors/FogEffects/HeightFog/HeightFog.INT.udn
--- a/Engine/Documentation/Source/Engine/Actors/FogEffects/HeightFog/HeightFog.INT.udn
+++ b/Engine/Documentation/Source/Engine/Actors/FogEffects/HeightFog/HeightFog.INT.udn
@@ -53,7 +53,7 @@ Description:Guide to creating height-based, distance fog in levels.
 
 
 
-This document covers Exponential Height Fog and its uses. It assumes that you have a level already set up in which you would like to add Exponential Height Fog.
+This document covers **Exponential Height Fog** and its uses. It assumes that you have a level already set up in which you would like to add Exponential Height Fog.
 
 [EXCERPT:intro]
 Exponential Height Fog creates more density in low places of a map and less density in high places. The transition is smooth so you never get a hard cutoff as you increase altitude. Exponential Height Fog also provides two fog colors, one for the hemisphere facing the dominant directional light (or straight up if none exists), and another color for the opposite hemisphere. 
@@ -64,14 +64,14 @@ Exponential Height Fog creates more density in low places of a map and less dens
 
 ## Using Exponential Height Fog
 
-In the Class Viewer, select the **Exponential Height Fog** Actor, which is under _Fog_. **Right-click** in the world and select _Add ExponentialHeightFog Here_. 
+In the **Modes** panel, select the **Exponential Height Fog** Actor, which is under **Visual Effects**. **Left-click** and drag to place it in the world. 
 
 Positioning the Exponential Height Fog Actor will determine the height of the fog.
 
 
 ### Exponential Height Fog Properties
 
-In the _Height Fog_ section, you can edit the following properties related to the component:
+In the Height Fog section, you can edit the following properties related to the component:
 
 |**Property** | **Description** |
 |---|---|
@@ -94,7 +94,7 @@ In the _Height Fog_ section, you can edit the following properties related to th
 
 ## Animating Exponential Height Fog
 
-To animate Exponential Height Fog, use Matinee to animate the ExponentialHeightFog Actor values. For more information, see the [Matinee User Guide](Engine/Matinee/UserGuide) page.
+To animate Exponential Height Fog, use Matinee to animate the Exponential Height Fog Actor values. For more information, see the [Matinee User Guide](Engine/Matinee/UserGuide) page.
 
 You can also animate Exponential Height Fog directly in Level Blueprint, by way of a timeline or by any network of nodes driving the functions and properties of the fog Actor.
 
@@ -107,7 +107,7 @@ Height Fog Color is now able to be changed at runtime using the Linear Color Pro
 1. Create a group in Matinee with that Actor attached.
 1. Create a linear color property track in the above group.
 1. A dialog should appear letting you choose which color property to interpolate. There should be only two options: FogInscatteringColor and DirectionalInscatteringColor.
-1. Once that track is set up, you can place keyframes and use the context menu (**right-click**) or the [Curve Editor](Engine/UI/CurveEditor) to set the color for each key.
+1. Once that track is set up, you can place keyframes and use the context menu (**Right-click**) or the [Curve Editor](Engine/UI/CurveEditor) to set the color for each key.
 
 
 
