UEDOC-2228

Video has been updated and will show live after the next publishing

diff --git a/Engine/Documentation/Source/Engine/Actors/CameraActors/CameraActors.INT.udn b/Engine/Documentation/Source/Engine/Actors/CameraActors/CameraActors.INT.udn
--- a/Engine/Documentation/Source/Engine/Actors/CameraActors/CameraActors.INT.udn
+++ b/Engine/Documentation/Source/Engine/Actors/CameraActors/CameraActors.INT.udn
@@ -23,7 +23,7 @@ The Camera Actor can be found in the **Modes** panel under the **All Classes** c
 360
 [/PARAMLITERAL]
 [PARAMLITERAL:videoid]
-HMx63-uevs4
+Rt_L_KI4G3U
 [/PARAMLITERAL]
 [/OBJECT]
 
