UEDOC-2229

Link has been updated and will show love after the next publishing

diff --git a/Engine/Documentation/Source/Engine/Actors/CameraActors/CameraActors.INT.udn b/Engine/Documentation/Source/Engine/Actors/CameraActors/CameraActors.INT.udn
--- a/Engine/Documentation/Source/Engine/Actors/CameraActors/CameraActors.INT.udn
+++ b/Engine/Documentation/Source/Engine/Actors/CameraActors/CameraActors.INT.udn
@@ -40,7 +40,7 @@ Using a Camera Actor is as easy as selecting it from the **Modes** panel and the
 	[/PARAM]
 	[PARAM:links]
 		* [Using Cameras](Gameplay/HowTo/UsingCameras "%Gameplay/HowTo/UsingCameras%")
-		* [Using Cameras in Blueprints](Gameplay/HowTo/UsingCameras/Blueprints "Gameplay/HowTo/UsingCameras/Blueprints:description)
+		* [Using Cameras in Blueprints](Gameplay/HowTo/UsingCameras/Blueprints "%Gameplay/HowTo/UsingCameras/Blueprints%")
 	[/PARAM]
 [/OBJECT]
 
