#UE4doc #docs:
Changed font settings for multiple actors
Added engine version field
Changed mentions of 'Quick Settings' to 'Settings' to be accurate according to the engine's UI

diff --git a/Engine/Documentation/Source/Engine/Actors/Grouping/Grouping.INT.udn b/Engine/Documentation/Source/Engine/Actors/Grouping/Grouping.INT.udn
--- a/Engine/Documentation/Source/Engine/Actors/Grouping/Grouping.INT.udn
+++ b/Engine/Documentation/Source/Engine/Actors/Grouping/Grouping.INT.udn
@@ -2,6 +2,7 @@ Availability: Public
 Title:Grouping Actors
 Crumbs:%ROOT%, Engine, Engine/Actors
 Description:Guide to creating and working with groups of actors within Unreal Editor.
+version: 4.9
 
 
 [VAR:playlistnav]
@@ -83,11 +84,11 @@ An Actor can only belong to one group at a time, unless that group itself is nes
 
 	![Group Menu Option](LE_GroupMenu.png)
 
-If you do not see the **Group** option on the context menu, make sure that **Allow Group Selection** is selected on the **Quick Settings** menu. You can access the **Quick Settings** menu from the toolbar at the top of the Unreal Editor.
+If you do not see the **Group** option on the context menu, make sure that **Allow Group Selection** is selected on the **Settings** menu. You can access the **Settings** menu from the toolbar at the top of the Unreal Editor.
 
 ![](LE_ViewMenuGroup.png)
 
-When you group Actors together, you create a Group Actor in the World Outliner. To select all Actors in that group, you can select either the Group Actor or any member of the group.
+When you group Actors together, you create a Group Actor in the **World Outliner**. To select all Actors in that group, you can select either the Group Actor or any member of the group.
 
 ![Group Actor](LE_GroupActor.png)
 
@@ -96,7 +97,7 @@ For more information about Actor selection, see [](Engine/Actors/Selection).
 
 ## Options for Grouping
 
-Options related to groups can be found in the context menu available by **right-clicking** on a selection of Actors (minimum of 2) in the Viewport or in the World Outliner.
+Options related to groups can be found in the context menu available by **Right-clicking** on a selection of Actors (minimum of 2) in the **Viewport** or in the **World Outliner**.
 
 Which options you see on the context menu depends on what you have selected:
 * Objects in a group.
@@ -115,7 +116,7 @@ Which options you see on the context menu depends on what you have selected:
 
 ### Group
 
-The **Group** option is used to create a new group from selected Actors. If none of the selected Actors are already in a group, only **Group** will be available. Choose **Group** to group the selected Actors.   
+The **Group** option is used to create a new group from selected Actors. If none of the selected Actors are already in a group, only Group will be available. Choose Group to group the selected Actors.   
 
 Grouping is non-destructive in that it never destroys existing selected groups. Any number of Actors or groups can be included in a selection when creating a group. For example, if two Lights and an existing locked group of two Static Meshes are selected and then grouped, a new group will be created with the two Lights and the sub-group of Static Meshes under it. If any selected Actors are in unlocked groups, they will be removed from their old group and placed in the new one.
 
@@ -171,7 +172,7 @@ Locking groups works in the opposite order. Selecting **Lock Groups** on a curre
 
 ### Remove from Group
 
-The **Remove from Group** option is only available when selecting Actors in an unlocked group.  This option will remove the selected Actors from the group.  If a nested group is removed from a larger group, it still remains a group itself.  Once a group contains no Actors or other groups, it is automatically deleted, and its GroupActor is removed from the World Outliner.
+The **Remove from Group** option is only available when selecting Actors in an unlocked group.  This option will remove the selected Actors from the group.  If a nested group is removed from a larger group, it still remains a group itself.  Once a group contains no Actors or other groups, it is automatically deleted, and its GroupActor is removed from the **World Outliner**.
 
 
 ### Add Actors to Group
@@ -183,9 +184,9 @@ The **Add Actors to Group** option is only available when one group is selected
 ## Notes
 
 
-* Groups can be globally disabled (or re-enabled) in the Editor for selection purposes by selecting **Allow Group Selection** from the **Quick Settings** menu in the toolbar.
+* Groups can be globally disabled (or re-enabled) in the Editor for selection purposes by selecting **Allow Group Selection** from the **Settings** menu in the toolbar.
 * On creation of a new group, pivot points are automatically centered. Group pivot points will set any Editor widgets when the group is locked. When a group is unlocked, the pivot points of its Actors will function as normal.
-* Group pivots can be set by selecting the desired group and then holding down **Alt** and clicking the **middle mouse button** in the viewport for the pivot location.
+* Group pivots can be set by selecting the desired group and then holding down **Alt** and clicking the **Middle Mouse Button** in the viewport for the pivot location.
 * It is not possible to group Actors that are in different levels.  Moving an Actor that is currently in a group from one level to another will remove it from the existing group.  However, entire groups can be moved across levels.
 
 
