#UE4doc #docs:
Changed font settings for multiple terms
Added engine version field

diff --git a/Engine/Documentation/Source/Engine/Actors/Mobility/ActorMobility.INT.udn b/Engine/Documentation/Source/Engine/Actors/Mobility/ActorMobility.INT.udn
--- a/Engine/Documentation/Source/Engine/Actors/Mobility/ActorMobility.INT.udn
+++ b/Engine/Documentation/Source/Engine/Actors/Mobility/ActorMobility.INT.udn
@@ -2,6 +2,7 @@ Availability:Public
 Title:Actor Mobility
 Crumbs: %ROOT%, Engine, Engine/Actors
 Description:Setting that controls whether an Actor will be allowed to move or change in some way during gameplay.
+version: 4.9
 
 
 [VAR:playlistnav]
@@ -27,37 +28,37 @@ Description:Setting that controls whether an Actor will be allowed to move or ch
 [EXCERPT:Intro]
 ![](TransformMobility.png)
 
-The Mobility setting controls whether an Actor will be allowed to move or change in some way during gameplay. This primarily applies to Static Mesh Actors and Light Actors.
+The **Mobility** setting controls whether an Actor will be allowed to move or change in some way during gameplay. This primarily applies to Static Mesh Actors and Light Actors.
 [/EXCERPT:Intro]
 
 When available, the Mobility property has 3 states:
 
 | Mobility State | Description |
 |--- | --- |
-| **Static** | _Static_ is reserved for Actors that are not intended to move or update in any way during gameplay. [INCLUDE:#Static] |
+| **Static** | Static is reserved for Actors that are not intended to move or update in any way during gameplay. [INCLUDE:#Static] |
 | **Stationary** | [INCLUDE:#Stationary]  |
-| **Movable** | _Movable_ should be used only when an Actor needs to move during gameplay. [INCLUDE:#Movable] |
+| **Movable** | Movable should be used only when an Actor needs to move during gameplay. [INCLUDE:#Movable] |
 
 
 
 <!-----
 
 [EXCERPT:Static]
-* **Static Mesh Actors** with a Mobility of _Static_ will have their shadows contribute to pre-calculated light maps (Lightmass baked lighting). This makes them ideal for structural or decorative meshes that never need to relocate during gameplay. Note, however, that their Materials can still be animated.
-* **Light Actors** with a Mobility of _Static_ will contribute to pre-calculated light maps (Lightmass baked lighting). They will still illuminate dynamic objects due to the [Indirect Lighting Cache](Engine/Rendering/LightingAndShadows/IndirectLightingCache). This is the ideal lighting path for mobile workflows, since it is essentially free in terms of gameplay performance.
+* **Static Mesh Actors** with a Mobility of Static will have their shadows contribute to pre-calculated light maps (Lightmass baked lighting). This makes them ideal for structural or decorative meshes that never need to relocate during gameplay. Note, however, that their Materials can still be animated.
+* **Light Actors** with a Mobility of Static will contribute to pre-calculated light maps (Lightmass baked lighting). They will still illuminate dynamic objects due to the [Indirect Lighting Cache](Engine/Rendering/LightingAndShadows/IndirectLightingCache). This is the ideal lighting path for mobile workflows, since it is essentially free in terms of gameplay performance.
 
 
 [/EXCERPT:Static]
 
 
 [EXCERPT:Stationary]
-_Stationary_ is used for **Light Actors** that do not move, but may update in some way during gameplay, such as being switched on/off, having their color change, etc. Lights set up in this way still contribute to the Lightmass pre-calculated light map, but can also cast dynamic shadows for moving objects. However, care must be used not to have too many affecting a given Actor. See the [Stationary Lights documentation](Engine/Rendering/LightingAndShadows/LightMobility/StationaryLights) for more details. **Static Mesh Actors** cannot have a Mobility of _Stationary_.
+Stationary is used for **Light Actors** that do not move, but may update in some way during gameplay, such as being switched on/off, having their color change, etc. Lights set up in this way still contribute to the Lightmass pre-calculated light map, but can also cast dynamic shadows for moving objects. However, care must be used not to have too many affecting a given Actor. See the [Stationary Lights documentation](Engine/Rendering/LightingAndShadows/LightMobility/StationaryLights) for more details. **Static Mesh Actors** cannot have a Mobility of Stationary.
 [/EXCERPT:Stationary]
 
 
 [EXCERPT:Movable]
-* **Static Mesh Actors** with a Mobility of _Movable_ will not cast pre-calculated shadows into the light map. They can be lit by Static Light Actors due to the [Indirect Lighting Cache](Engine/Rendering/LightingAndShadows/IndirectLightingCache). They will cast dynamic shadows if lit by Stationary or Movable lights. This is the typical setting for any non-deforming mesh element that needs to move in a scene, such as a door, a lift, and so on.
-* **Light Actors** with a Mobility of _Movable_ can _only_ cast dynamic shadows. Because of this, care must be taken when making these lights cast shadows, as their shadowing method is the most performance intensive. 	It should be noted, however, that _non-shadowing_ Movable lights are very inexpensive to calculate due to Unreal Engine's deferred rendering system.
+* **Static Mesh Actors** with a Mobility of Movable will not cast pre-calculated shadows into the light map. They can be lit by Static Light Actors due to the [Indirect Lighting Cache](Engine/Rendering/LightingAndShadows/IndirectLightingCache). They will cast dynamic shadows if lit by Stationary or Movable lights. This is the typical setting for any non-deforming mesh element that needs to move in a scene, such as a door, a lift, and so on.
+* **Light Actors** with a Mobility of Movable can only cast dynamic shadows. Because of this, care must be taken when making these lights cast shadows, as their shadowing method is the most performance intensive. 	It should be noted, however, that non-shadowing Movable lights are very inexpensive to calculate due to Unreal Engine's deferred rendering system.
 [/EXCERPT:Movable]
 
 -->
