#UE4doc #docs:
Changed font settings for multiple terms
Added engine version field

diff --git a/Engine/Documentation/Source/Engine/Actors/Placement/PlacingActors.INT.udn b/Engine/Documentation/Source/Engine/Actors/Placement/PlacingActors.INT.udn
--- a/Engine/Documentation/Source/Engine/Actors/Placement/PlacingActors.INT.udn
+++ b/Engine/Documentation/Source/Engine/Actors/Placement/PlacingActors.INT.udn
@@ -4,6 +4,7 @@ Crumbs: %ROOT%, Engine, Engine/Actors
 Description:Shows how you can place Actors such as props, lights, cameras, etc. in your levels. 
 Related: Engine/UI/ClassViewer
 Related: GettingStarted\HowTo
+version: 4.9
 
 [VAR:playlistnav]
 [OBJECT:PlaylistNavigation]
@@ -74,23 +75,23 @@ When you create a new project, the level that is loaded when you open your proje
 [/EXCERPT:Place01]
 
 [REGION:note]
-Adding Actors to your level automatically also adds them to the list in the **World Outliner**, which is a panel that is by default located in the upper-right corner of the Unreal Editor. For more information about the World Outliner, see [World Outliner](Engine/UI/LevelEditor\SceneOutliner).
+Adding Actors to your level automatically also adds them to the list in the **World Outliner**, which is a panel that is by default located in the upper-right corner of the Unreal Editor. For more information about the **World Outliner**, see [World Outliner](Engine/UI/LevelEditor\SceneOutliner).
 [/REGION]
 
 ## Placing Actors in Place Mode
 
 [EXCERPT:PlaceMode]
-This section will show you how to place Actors from the **Modes** menu while in [Place Mode](Engine/UI/LevelEditor/Modes/PlaceMode). 
+This section will show you how to place Actors from the **Modes** menu while in [Place mode](Engine/UI/LevelEditor/Modes/PlaceMode). 
 
 1.  With your project open, in the upper-left corner of the screen you will see the **Modes** window. 
 
 	![](Place0.png)
 
-1.	While in **Place** mode (1), click the **Basic** tab (2) then **Left-click** on and drag an Actor (3) into the level viewport. 
+1.	While in **Place** mode (1), click the **Basic** tab (2) then **Left-click** on and drag an Actor (3) into the **Level Viewport**. 
 
 	![](Place1.png)
 
-	Each of Actor displayed on each tab of Place Mode can be dragged-and-dropped into a level. 
+	Each of Actor displayed on each tab of **Place** mode can be dragged-and-dropped into a level. 
 [/EXCERPT:PlaceMode]
 
 ## Placing Actors Using the Context Menu
@@ -107,30 +108,30 @@ This section will demonstrate how to place Actors from the [Content Browser](Eng
 
 	![](Place3.png)
 
-1.	Inside the Level Viewport, **Right-click** anywhere to bring up the **Context Menu** and select the asset under the _Place Actor_ section.
+1.	Inside the **Level Viewport**, **Right-click** anywhere to bring up the **Context Menu** and select the asset under the **Place Actor** section.
 
 	![](Place4.png)
 
-	After selecting the asset from the Context Menu, you will see it appear inside your level. 
+	After selecting the asset from the **Context Menu**, you will see it appear inside your level. 
 
-You also have the option of adding a different type of asset even when you have a specific asset selected in the **Content Browser**. You can see the types of assets you can add by opening the flyout menu in the **Add Actor** section of the Context Menu as seen in the image below. 
+You also have the option of adding a different type of asset even when you have a specific asset selected in the **Content Browser**. You can see the types of assets you can add by opening the flyout menu in the **Add Actor** section of the **Context Menu** as seen in the image below. 
 
 ![](PlaceActorSubMenu.png)
 
-If you have one or more Actors selected in the Viewport, and an asset selected in the **Content Browser**, you can opt to replace the selected Actors in the level with the selected asset by clicking the asset's name in the **Replace Selected Actors with** section of the Context Menu as seen in the image below. This may be useful if you want to do sweeping changes and swapping multiple assets out at once. 
+If you have one or more Actors selected in the **Viewport**, and an asset selected in the **Content Browser**, you can opt to replace the selected Actors in the level with the selected asset by clicking the asset's name in the **Replace Selected Actors with** section of the **Context Menu** as seen in the image below. This may be useful if you want to do sweeping changes and swapping multiple assets out at once. 
 
 ![](ReplaceMenu.png)
 
 
 ## Placing Actors Using Drag and Drop
 
-In addition to being able to add specific types of Actors from the **Content Browser** to a level through the Viewport Context Menu, you can also add an asset to your level by dragging it from the **Content Browser** and dropping it into one of the Viewports, in the location where you would like to place the Actor. When you start to drag the asset into the Viewport, the cursor will change so that you know that type of asset will be dropped into the Viewport.
+In addition to being able to add specific types of Actors from the **Content Browser** to a level through the **Viewport Context Menu**, you can also add an asset to your level by dragging it from the **Content Browser** and dropping it into one of the **Viewports**, in the location where you would like to place the Actor. When you start to drag the asset into the **Viewport**, the cursor will change so that you know that type of asset will be dropped into the **Viewport**.
 
 1.  Inside your project, in the lower-left corner of the screen you will see the **Content Browser**. 
 
 	![](Place2.png)
 
-1.	Find the asset you want to add, then **Left-click** on and **drag-and-drop** it into the Level Viewport. 
+1.	Find the asset you want to add, then **Left-click** on and **drag-and-drop** it into the **Level Viewport**. 
 
 	![](Place5.png)
 
@@ -148,7 +149,7 @@ When you drag and drop an asset from the **Content Browser**, the following type
 
 A more advanced method of placing Actors is by using the **Class Viewer**, which is a hierarchical list of classes used by the Editor.
 
-1.  To open the **Class Viewer**, click **Window** then **Developer Tools** and under _Miscellaneous_, select **Class Viewer**. 
+1.  To open the **Class Viewer**, click **Window** then **Developer Tools** and under **Miscellaneous**, select **Class Viewer**. 
 
 	![](Place6.png)
 
@@ -158,11 +159,11 @@ A more advanced method of placing Actors is by using the **Class Viewer**, which
 
 	Any Actor that can be placed in the level will be tinted in **Blue** in this window. 
 
-1.	Select the asset you wish to place, then **drag-and-drop** it in the Level Viewport. 
+1.	Select the asset you wish to place, then **drag-and-drop** it in the **Level Viewport**. 
 
 	![](Place8.png)
 
-For more information on working with the Class Viewer, refer to the [](Engine/UI/ClassViewer) documentation. 
+For more information on working with the **Class Viewer**, refer to the [](Engine/UI/ClassViewer) documentation. 
 
 
 %playlistnav%
