#UE4doc #docs:
Added engine version field
Changed font settings for multiple terms
Fixed some spelling errors

diff --git a/Engine/Documentation/Source/Engine/Actors/PlayerStart/PlayerStart.INT.udn b/Engine/Documentation/Source/Engine/Actors/PlayerStart/PlayerStart.INT.udn
--- a/Engine/Documentation/Source/Engine/Actors/PlayerStart/PlayerStart.INT.udn
+++ b/Engine/Documentation/Source/Engine/Actors/PlayerStart/PlayerStart.INT.udn
@@ -2,6 +2,7 @@ Availability: Public
 Title:Player Start
 Crumbs: %ROOT%, Engine, Engine/Actors
 Description:Guide to using Player Starts.
+version: 4.9
 
 ![](Player_Start_Actor.png)
 
@@ -10,7 +11,7 @@ The ability to spawn a player into the world at any location is a very useful fe
 
 ## Placing a Player Start Actor
 
-Player Start Actors can be found in the Mode panel under the Basic category. To place one into the world, simply drag it from the **Modes** panel into the game world.
+Player Start Actors can be found in the **Modes** panel under the **Basic** category. To place one into the world, drag it from the **Modes** panel into the game world.
 
 [OBJECT:EmbeddedVideo]
 [PARAMLITERAL:width]
@@ -38,11 +39,11 @@ While using Player Starts is quite simple, there are a few tips and tricks about
 
  **No Player Start:** If you try and play your game without adding a Player Start to the world, the player will start from 0,0,0 in the world. Because of this, always make sure that you have a Player Start in your world.
 
- **Right-Click Player From Here:** Sometimes you might want to play from a certain location and not from the Player Start. You can do this by using the **Right-click** play from here option that is available by **Right-clicking** in the editor viewport and selecting the **Play From Here** options. 
+ **Play From Here:** Sometimes you might want to play from a certain location and not from the Player Start. You can do this by using the **Right-click** play from here option that is available by **Right-clicking** in the editor viewport and selecting the **Play From Here** options. 
 
 ![](Right_Click_Play.png)
 
-** Bad Size:** Sometimes your Player Start might change from an icon of a controller to an icon that says Bad Size. If this happens, all you need to do is move your Player Start around in the world until it is not intersecting anything in the scene.
+** Bad Size:** Sometimes your Player Start might change from an icon of a controller to an icon that says "BADsize". If this happens, all you need to do is move your Player Start around in the world until it is not intersecting anything in the scene.
 
 [OBJECT:EmbeddedVideo]
 [PARAMLITERAL:width]
