#UE4doc #docs:
Removed Overview header
Changed font settings for multiple terms
Added engine version field

diff --git a/Engine/Documentation/Source/Engine/Actors/Selection/ActorSelection.INT.udn b/Engine/Documentation/Source/Engine/Actors/Selection/ActorSelection.INT.udn
--- a/Engine/Documentation/Source/Engine/Actors/Selection/ActorSelection.INT.udn
+++ b/Engine/Documentation/Source/Engine/Actors/Selection/ActorSelection.INT.udn
@@ -2,6 +2,7 @@ Availability:Public
 Title:Selecting Actors
 Crumbs: %ROOT%, Engine, Engine/Actors
 Description:Overview of methods available for selecting Actors in the Level Editor viewports.
+version: 4.9
 
 [VAR:playlistnav]
 [OBJECT:PlaylistNavigation]
@@ -66,7 +67,7 @@ Description:Overview of methods available for selecting Actors in the Level Edit
 
 [TOC(start:2)]
 
-## Overview
+
 
 Selecting Actors, while simple in nature, is an important part of the level-editing process. The ability to quickly and easily select the Actors that you want to work with enhances productivity and speeds up the design process.
 
@@ -75,7 +76,7 @@ There are many different ways to select Actors, or groups of Actors. Each of the
 
 ## Simple Selection
 
-The most basic method of selecting Actors is simply to **left-click** on them in the viewport. Clicking on an Actor will deselect any currently selected Actors and select the new one instead. If you hold down the **Ctrl** key while you click on a new (unselected) Actor, the new Actor is added to the selection. If you hold down the **Ctrl** key while clicking on a selected Actor, the Actor is removed from the selection.
+The most basic method of selecting Actors is simply to **Left-click** on them in the viewport. Clicking on an Actor will deselect any currently selected Actors and select the new one instead. If you hold down the **Ctrl** key while you click on a new (unselected) Actor, the new Actor is added to the selection. If you hold down the **Ctrl** key while clicking on a selected Actor, the Actor is removed from the selection.
 
 This method is good for selecting small numbers of Actors or several isolated Actors spread out across the map, but it can be slow and tedious for selecting large numbers of Actors.
 
@@ -109,13 +110,13 @@ The following table shows the possible keyboard combinations for marquee selecti
 
 | **Keyboard Shortcut** | **Action** |
 | ---------| ---------|
-| **LMB** | Replaces the current selection with the Actors contained in the box. |
+| **(Left Mouse Button) LMB** | Replaces the current selection with the Actors contained in the box. |
 | **Shift+LMB** | Adds the Actors contained in the box to the current selection. |
 | **Ctrl+Alt+LMB** | Removes any select Actors in the box from the current selection. |
 
 ## Advanced Selection Techniques
 
-You can select Actors that share specific things in common, such as class, asset type, or Material.  To see your selection options, **right-click** a selected Actor, click **Select**, and then review the list of options on the context menu.
+You can select Actors that share specific things in common, such as class, asset type, or Material.  To see your selection options, **Right-click** a selected Actor, click **Select**, and then review the list of options on the context menu.
 
 [REGION:tip]
 The options available on the **Select** context menu vary depending on what is selected. For example, if you have a Static Mesh selected, you will see different options than if you have a Blueprint selected.
