#UE4doc #docs:
Changed font settings for multiple terms
Added engine version field

diff --git a/Engine/Documentation/Source/Engine/Actors/SkeletalMeshActors/SkeletalMeshActors.INT.udn b/Engine/Documentation/Source/Engine/Actors/SkeletalMeshActors/SkeletalMeshActors.INT.udn
--- a/Engine/Documentation/Source/Engine/Actors/SkeletalMeshActors/SkeletalMeshActors.INT.udn
+++ b/Engine/Documentation/Source/Engine/Actors/SkeletalMeshActors/SkeletalMeshActors.INT.udn
@@ -4,6 +4,7 @@ Crumbs: %ROOT%, Engine, Engine/Actors
 Description:Creating and using Skeletal Mesh Actors
 Related: Engine/Actors
 Related: Engine/Animation
+version: 4.9
 
 
 
@@ -15,18 +16,18 @@ Related: Engine/Animation
 When playing a video game that uses Unreal Engine 4 (UE4), the character that you are controlling on screen is often made up of a special kind of mesh called a **Skeletal Mesh Actor**. Skeletal Mesh Actors are special because they can be used to display complex animation data that was made inside of an external 3D application and then imported into UE4. The following guide aims to show you the ins and outs of how to use Skeletal Mesh actors inside UE4.
 
 [region:tip]
-All content for this tutorial was taken from the **Matinee** demo which you can download for free from the Learn Tab. 
+All content for this tutorial was taken from the Matinee demo which you can download for free from the Learn Tab. 
 [/region]
 
 ## What is a Skeletal Mesh Actor
 A Skeletal Mesh Actor is a special type of Actor that is used to display complex animation data that was created in an external 3D animation program. Skeletal Mesh Actors are different than Static Mesh Actors because Skeletal Mesh Actors have a complex skeleton inside of them the helps to move the individual vertices of the Skeletal Mesh to match the current animation that is playing. Because of a Skeletal Mesh Actor's ability to use a complex Skeleton, they are generally used for things like characters, creatures, complex machinery, or for anything that needs to deform or display complex motion. Inside of UE4, you will find that Skeletal Mesh Actors are not only used for characters but also used with Matinee to create cinematic sequences. 
 
 ## Placing a Skeletal Mesh Actor in a Level
-Skeletal Mesh Actors are placed in maps using standard Actor placement methods; either through the **Right-click** context menu of the viewports or by **Drag and Dropping** from the [Content Browser](Engine/Content/Browser).
+Skeletal Mesh Actors are placed in maps using standard Actor placement methods; either through the **Right-click** context menu of the viewports or by dragging and dropping from the [Content Browser](Engine/Content/Browser).
 
 **Drag & Drop**
 
-You can drag & drop Skeletal Mesh Actors right from the **Content Browser** right into the level you are working with by following these few simple steps.
+You can drag and drop Skeletal Mesh Actors right from the **Content Browser** right into the level you are working with by following these few simple steps.
 1. In the **Content Browser**, locate the Skeletal Mesh you want to add to the map as a Skeletal Mesh Actor.
 
 1. **Left-click** on the Skeletal Mesh in the **Content Browser** and drag the mouse (while holding the **Left Mouse Button** down) from the **Content Browser** to the location in the viewport where you want to place the mesh. The location does not need to be exact. You can always reposition and rotate the mesh afterwards.   
@@ -53,7 +54,7 @@ Getting a Skeletal Mesh Actor so that it will animate while the game is running
 ### Skeletal Mesh Actor Setup
 The setup for getting a Skeletal Mesh Actor ready to play an animation is actually quite simple and can be done in the following steps.
 1. First find the **Skeletal Mesh** you want in the **Content Browser**
-1. Once you have found the **Skeletal Mesh** you wish to use, add it to the level either via **Drag & Drop** or via the **Context Menus**.
+1. Once you have found the **Skeletal Mesh** you wish to use, add it to the level either via dragging and dropping or via the **Context Menus**.
 1. When completed, you should have something that looks like this.  
 	![](SkeletalMeshActor_Adding_Mesh_To_Scene.png)
 
@@ -95,7 +96,7 @@ To create a Physics Asset for your mesh, all you have to do is simply **Right-cl
 ![](SkeletalMeshActor_Create_Physics_Asset.png)
 
 ### Assigning a Physics Asset
-Once your Physics Asset has been created, it is not time to assign it to your mesh. Physics Assets can be assigned to Skeletal Meshes by simply opening up the Skeletal Mesh and then assigning the Physics Asset to the Physics Asset slot that is under the Physics section. 
+Once your Physics Asset has been created, it is not time to assign it to your mesh. Physics Assets can be assigned to Skeletal Meshes by simply opening up the Skeletal Mesh and then assigning the Physics Asset to the **Physics Asset** slot that is under the **Physics** section. 
 
 ![](SkeletalMeshActor_Add_Physics_Assets.png)
 
