#UE4doc #docs:
Added engine version field
Changed font settings for multiple terms
Replaced a picture and a couple terms to be accurate with changes made to the Editor's UI

diff --git a/Engine/Documentation/Source/Engine/Actors/StaticMeshActor/UserGuide.INT.udn b/Engine/Documentation/Source/Engine/Actors/StaticMeshActor/UserGuide.INT.udn
--- a/Engine/Documentation/Source/Engine/Actors/StaticMeshActor/UserGuide.INT.udn
+++ b/Engine/Documentation/Source/Engine/Actors/StaticMeshActor/UserGuide.INT.udn
@@ -2,6 +2,7 @@ Availability: Public
 Title:Static Mesh Actors
 Crumbs:%ROOT%, Engine, Engine/Actors
 Description:Guide to placing and working with StaticMeshActors for creating world geometry in your environments.
+version: 4.9
 
 [VAR:playlistnav]
 [OBJECT:PlaylistNavigation]
@@ -32,7 +33,7 @@ to populate your worlds, the **StaticMeshActor** is used. StaticMeshes dragged i
 automatically converted to StaticMeshActors.
 
 [REGION:note]
-Though they are called StaticMeshActors, this simply refers to the **mesh** of the StaticMeshActor being static. A StaticMeshActor can be made movable so it can be an elevator or made to simulate physics so it can crush the player. 
+Though they are called StaticMeshActors, this simply refers to the mesh of the StaticMeshActor being static. A StaticMeshActor can be made movable so it can be an elevator or made to simulate physics so it can crush the player. 
 See [Making a StaticMeshActor Movable](#makingastaticmeshactormovable) for details.
 [/REGION]
 
@@ -44,7 +45,7 @@ See [Making a StaticMeshActor Movable](#makingastaticmeshactormovable) for detai
 
 ## Placement
 
-StaticMeshActors are placed in maps using standard Actor placement methods; either through the **right-click** context 
+StaticMeshActors are placed in maps using standard Actor placement methods; either through the **Right-click** context 
 menu of the viewports or drag and drop from the [Content Browser](Engine/Content/Browser).
 
 **Drag and Drop**
@@ -55,7 +56,7 @@ menu of the viewports or drag and drop from the [Content Browser](Engine/Content
 
 1. **Left-Click** on the StaticMesh in the **Content Browser** and drag the mouse (while holding the **Left Mouse Button** down) from the **Content Browser** to the location in the **Viewport** where you want to place the mesh. The location does not need to be exact. You can always reposition, rotate, scale the mesh afterwards.   
 
-1. Release the **Left Mouse Button** to place the mesh in the map as a StaticMeshActor, as seen in the Property Window.   
+1. Release the **Left Mouse Button** to place the mesh in the map as a StaticMeshActor, as seen in the **Property** window.   
 
 
 
@@ -66,15 +67,15 @@ menu of the viewports or drag and drop from the [Content Browser](Engine/Content
 
 1. In the **Content Browser**, select the StaticMesh you want to add to the map as a StaticMeshActor.
    
-1. **Right-click** in the viewport where you want to add the mesh and choose **Add Actor: Selection** from the context menu. The location does not need to be exact. You can always reposition, rotate, scale the mesh afterwards. 
+1. **Right-click** in the viewport where you want to add the mesh and choose **Place Actor: Selection** from the **Context Menu**. The location does not need to be exact. You can always reposition, rotate, scale the mesh afterwards. 
 
-1. The StaticMesh has been placed in the map as a StaticMeshActor, as seen in the Property Window. 
+1. The StaticMesh has been placed in the map as a StaticMeshActor, as seen in the **Property** window. 
   
 	
 
 ## Making a StaticMeshActor Movable
 
-In order to move, rotate, or scale a StaticMeshActor during play, you must first make it movable.
+In order to move, rotate, or scale a StaticMeshActor during play, you must first make it **Movable** by selecting **Moveable** from the top of the **Details** panel under **Mobility** when the StaticMeshActor is selected.
 
 
 ## Enabling a StaticMeshActor to Simulate Physics
@@ -86,7 +87,7 @@ In order to move, rotate, or scale a StaticMeshActor during play, you must first
 
 The Material(s) applied to a StaticMesh can be overridden on a per-Actor basis. This allows a single StaticMesh 
 asset to be used multiple times in a single map or throughout your game while giving each a unique appearance. 
-The **Materials** property located in the **Rendering** category of the _StaticMeshComponent_ of the StaticMeshActor 
+The **Materials** property located in the **Rendering** category of the StaticMeshComponent of the StaticMeshActor 
 is an array of Materials that directly maps to those applied to the StaticMesh asset through the [Static Mesh Editor](Engine/Content/Types/StaticMeshes/Editor). 
 Materials can be assigned to the array manually or they can be applied directly to the mesh in the **Viewport** by 
 dragging and dropping from the **Content Browser**.
@@ -97,7 +98,7 @@ dragging and dropping from the **Content Browser**.
 
 1. In the **Viewport**, select the StaticMeshActor you want to assign.
 
-1. In the **Details** Pane, under the **Materials** category, you will see all the materials assigned to the StaticMeshActor.
+1. In the **Details** panel, under the **Materials** category, you will see all the materials assigned to the StaticMeshActor.
 
 1. In the **Content Browser**, select the Material you want to apply to StaticMeshActor in the map, and either:   
 
