#UE4doc #docs:
Added engine version field

diff --git a/Engine/Documentation/Source/Engine/Actors/TargetPoint/TargetPoint.INT.udn b/Engine/Documentation/Source/Engine/Actors/TargetPoint/TargetPoint.INT.udn
--- a/Engine/Documentation/Source/Engine/Actors/TargetPoint/TargetPoint.INT.udn
+++ b/Engine/Documentation/Source/Engine/Actors/TargetPoint/TargetPoint.INT.udn
@@ -3,6 +3,7 @@ Title:Target Point Actors
 Crumbs: %ROOT%, Engine, Engine/Actors
 Description:Guide to creating and using Target Actors.
 Related: Gameplay/HowTo/SpawnAndDestroyActors/Blueprints
+version: 4.9
 
 ![](Target_Point_Actors.png)
 
