UEDOC-884

Updated image color to match the instructional guide

diff --git a/Engine/Documentation/Source/Resources/Showcases/BlueprintExamples/SecurityDoor/SecurityDoor.INT.udn b/Engine/Documentation/Source/Resources/Showcases/BlueprintExamples/SecurityDoor/SecurityDoor.INT.udn
--- a/Engine/Documentation/Source/Resources/Showcases/BlueprintExamples/SecurityDoor/SecurityDoor.INT.udn
+++ b/Engine/Documentation/Source/Resources/Showcases/BlueprintExamples/SecurityDoor/SecurityDoor.INT.udn
@@ -22,7 +22,7 @@ Of particular importance in this Blueprint is the use of a [Blueprint Interface]
 
 In this example, the Security Alarm Interface acts as a data conduit, allowing us to share functionality with both the Alarm Beacon and the Door Blueprints, specifically, whether the alarm is on or off.
 
-[REGION:raw]![](InterfaceConceptDiagram.png)[/REGION]
+[REGION:raw]![](InterfaceConceptDiagramNew.jpg)[/REGION]
 
 
 ## Camera Blueprint
