#UE4doc #docs - updated images to match 4.9 UI and updated terms in reflect changes. Addresses UEDOC-1028.

diff --git a/Engine/Documentation/Source/Programming/Slate/Overview/SlateOverview.INT.udn b/Engine/Documentation/Source/Programming/Slate/Overview/SlateOverview.INT.udn
--- a/Engine/Documentation/Source/Programming/Slate/Overview/SlateOverview.INT.udn
+++ b/Engine/Documentation/Source/Programming/Slate/Overview/SlateOverview.INT.udn
@@ -194,7 +194,7 @@ Docked tab:
 
 ## Developer tools
 
-The **Slate Widget Inspector** provides a means of debugging and analyzing the UI and associated code. This helps track down bugs and undesirable behavior as well as profile and optimize your user interface.
+The **Slate Widget Reflector** provides a means of debugging and analyzing the UI and associated code. This helps track down bugs and undesirable behavior as well as profile and optimize your user interface.
 
 [REGION:fullwidth]
 ![Widget Inspector](widget_inspector.png)
