#UE4doc #docs - Updated notes in page to indicate the correct axis that the source length follows. Addresses UEDOC-850.

diff --git a/Engine/Documentation/Source/Resources/ContentExamples/Lighting/5_3/5_3.INT.udn b/Engine/Documentation/Source/Resources/ContentExamples/Lighting/5_3/5_3.INT.udn
--- a/Engine/Documentation/Source/Resources/ContentExamples/Lighting/5_3/5_3.INT.udn
+++ b/Engine/Documentation/Source/Resources/ContentExamples/Lighting/5_3/5_3.INT.udn
@@ -40,8 +40,7 @@ Related: Engine/Rendering/LightingAndShadows/LightMobility
 
 This example demonstrates the use of **Source Length** on reflective surfaces. The purpose of Source Length is to give lights an elongated or tube-like reflective shape, as if the light were coming form a fluorescent light fixture. It should be noted that light emission is still calculated as if it is coming from a single point. The effect exists only within reflections.
 
-The Source Length runs down the Y-Axis of the light, and can be rotated by rotating the light. Also, while the light is selected, Source Length can be visualized by a yellow line stretching out from the light itself.
-
+The Source Length runs down the X-Axis of the light, and can be rotated by rotating the light. Also, while the light is selected, Source Length can be visualized by a yellow line stretching out from the light itself. When rotating the light the Source Length will continue along the Local Space X-Axis, but when checking the World Space, the light may not follow a direct axis.
 
 ![](SourceLength.png)(w:400)
 
