UEDOC-835

Embedded images have changed into tables.

diff --git a/Engine/Documentation/Source/Engine/Landscape/Editing/Brushes/Brushes.INT.udn b/Engine/Documentation/Source/Engine/Landscape/Editing/Brushes/Brushes.INT.udn
--- a/Engine/Documentation/Source/Engine/Landscape/Editing/Brushes/Brushes.INT.udn
+++ b/Engine/Documentation/Source/Engine/Landscape/Editing/Brushes/Brushes.INT.udn
@@ -41,7 +41,11 @@ You can set the brush type and falloff in either **Sculpt** or **Paint** mode.
 
 The size and falloff of the current brush is displayed in the viewport as a flashing, light-purple filled circle, or pair of concentric circles.
 
-![Landscape_Brush_Viewport.jpg](Landscape_Brush_Viewport.jpg)(w:900 h:184)
+[REGION:imagetable]
+| ![](Falloff00.jpg) | ![](Falloff05.jpg) | ![](Falloff10.jpg) |
+| -------- | ----------- | ------ |
+| Falloff: 0.0 | Falloff: 0.5 | Falloff: 1.0 |
+[/REGION]
 
 
 ## Circle
@@ -63,8 +67,11 @@ The Circle brush applies the current tool in a circular area, with a falloff def
 
 Below is an example of each of these falloff types when operating on the heightmap, with the same radius and falloff:
 
-![landscape_brush_falloff.png](landscape_brush_falloff.png)(w:789 h:416)
-
+[REGION:imagetable]
+| ![](Smooth.jpg)(w:200 h:300) | ![](Linear.jpg)(w:200 h:300) | ![](Sphere.jpg)(w:200 h:300) | ![](Tip.jpg)(w:200 h:300) |
+| -------- | ----------- | ------ | ----- |
+| ** Smooth **| ** Linear ** | ** Spherical** | ** Tip ** | 
+[/REGION]
 
 ## Alpha
 
