#UE4doc #docs - Updated Images to reflect 4.9 UI changes, added additional settings to the table, and added section for setting Game Default Map. Addresses UEDOC-1413.

diff --git a/Engine/Documentation/Source/Engine/Basics/Projects/Packaging/PackagingProjects.INT.udn b/Engine/Documentation/Source/Engine/Basics/Projects/Packaging/PackagingProjects.INT.udn
--- a/Engine/Documentation/Source/Engine/Basics/Projects/Packaging/PackagingProjects.INT.udn
+++ b/Engine/Documentation/Source/Engine/Basics/Projects/Packaging/PackagingProjects.INT.udn
@@ -6,7 +6,6 @@ Description:Packaging Unreal game projects for distribution.
 [TOC(start:2)]
 
 
-## Overview
 
 [EXCERPT:Intro]
 Before an Unreal project can be distributed to users, it must be properly packaged. Packaging ensures that all code and content is up to date and in the proper format to run on the desired target platform. 
@@ -26,6 +25,15 @@ Once you select the platform, the editor will compile your game if it contains c
 Starter Content in your project, or if you have made a lot of test/temporary content and maps.
 
 
+## Setting A Game Default Map
+
+Before you package your game you will first need to set a **Game Default Map** to load when your packaged game starts. If you do not set a map and are using a blank project you will only see a black screen when the packaged game starts. If you have used one of the template maps, like the First Person template or Third Person template the starting map will be loaded.
+
+To set the Game Default Map click on **Edit > Project Settings > Maps & Modes** in the Editor's main menu:
+
+![](Project_Settings_MapsNModes.png)
+
+
 
 
 ## Creating Packages
@@ -40,11 +48,11 @@ Confirming the target directory will then initiate the actual process that packa
 
 ![](progress.png)
 
-The status indicator provides a way to cancel the active packaging process by clicking the Cancel button. It is also possible to display extended output log information by clicking the Show Log link. The output log is useful for figuring out what went wrong in case a packaging process did not complete successfully:
+The status indicator provides a way to cancel the active packaging process by clicking the **Cancel** button. It is also possible to display extended output log information by clicking the **Show Log** link. The output log is useful for figuring out what went wrong in case a packaging process did not complete successfully:
 
 ![](log.png)
 
-For less experienced users, the most important log messages, such as errors and warnings, are logged to the regular Message Log window:
+For less experienced users, the most important log messages, such as errors and warnings, are logged to the regular **Message Log** window:
 
 ![](message_log.png)
  
@@ -86,7 +94,7 @@ will be used for just that project. Directions for how to generate the key and f
 
 ## Advanced Settings
 
-Clicking **File > Package Project > Packaging Settings...** in the main menu will present you with a number of advanced configuration options for the packaging feature.
+Clicking **File > Package Project > Packaging Settings...** or **Edit > Project Settings > Packaging** in the main menu will present you with a number of advanced configuration options for the packaging feature.
 
 ![](settings.png)
 
@@ -98,6 +106,11 @@ Clicking **File > Package Project > Packaging Settings...** in the main menu wil
 | Staging Directory | The directory that will contain your packaged build. This will be updated automatically when you pick a different directory in the target directory selection. |
 | Full Rebuild | Whether all of your code should be compiled. If disabled, only the modified code will be compiled. This may speed up the packaging process. For shipping builds you should always do a full rebuild to make sure nothing is missing or outdated. This option is enabled by default. |
 | Use Pak File | Whether to package your project's assets as individual files or a single package. If enabled, all assets will be put into a single .pak file instead of copying out all the individual files. If your project has a lot of asset files, using a Pak may make it easier to distribute it, because you do not have to transfer so many files. This option is enabled by default. |
+| Generate Chunks | Whether to generate .pak file chunks that can be used for streaming installs. |
+| Build Http Chunk Install Data | Whether to generate data for HTTP chunk installer. This allows this data to be hosted on a webserver to be installed at runtime. |
+| Http Chunk Install Data Directory | This is the directory where data will be built to. |
+| Http Chunk Install Data Version | This is the version name for HTTP chunk install data. |
+| Include Prerequisites | This specifies whether to include prerequisites of packaged games, such as redistributable operating system components. |
 [PUBLISH:Licensee]
 | Directories to Always Cook | The list of directories with files that should always be cooked. By default, the packaging feature automatically detects and cooks all content that is referenced by your game. Some required content may not be referenced directly, such as textures for a custom Slate UI. The directories in which such content resides should be listed here in order to also include it in the package. |
 [/PUBLISH:Licensee]
