#UE4doc #docs - Updated images and separated into imagetables. Removed embedded text from images.

diff --git a/Engine/Documentation/Source/Engine/Landscape/Landscape.INT.udn b/Engine/Documentation/Source/Engine/Landscape/Landscape.INT.udn
--- a/Engine/Documentation/Source/Engine/Landscape/Landscape.INT.udn
+++ b/Engine/Documentation/Source/Engine/Landscape/Landscape.INT.udn
@@ -159,7 +159,7 @@ The Landscape system paves the way for terrains that are orders of magnitude lar
 [/EXCERPT:landscape001]
 
 [REGION:raw]
-![Landscape_Terrain_Size.jpg](Landscape_Terrain_Size.jpg)
+![](Landscape_Terrain_Size.png)
 [/REGION]
 
 ### Landscape Memory Usage
@@ -179,8 +179,11 @@ The Landscape system stores the render data for the terrain in Textures in GPU m
 
 LODs for Landscape terrains are handled using standard Texture mipmaps. Each mipmap is a level of detail and the mipmap to sample can be specified using the `text2Dlod` HLSL instruction. This allows for a large number of LODs and also allows for smooth LOD transition as both mip levels can be sampled and then the heights and X and Y offsets can be interpolated in the vertex shader creating a clean morphing effect.
 
+
 [REGION:fullwidth raw]
-![Landscape_LOD_Morph.jpg](Landscape_LOD_Morph.jpg)
+| ![](Landscape_LOD1.png) | ![](Landscape_LOD1toLOD2.png) |  ![](Landscape_LOD2.png) |
+| -------- | ----------- | -------|
+| ** Fully LOD 1 **| ** Morphing from LOD 1 to LOD 2 ** | ** Fully LOD 2**|
 [/REGION]
 
 ### Heightmap and Weight Data Streaming
@@ -191,27 +194,32 @@ Since the data is being stored in Textures, the standard Texture streaming syste
 
 Since the X and Y slopes of the Landscape are being stored, the entire high-resolution (non-LOD'd) normal data is available for lighting  calculations. 
 
-**Left:** Landscape LODs
-
-**Right:** Full Res Normals
 
-![Landscape_HighRes_Normals.jpg](Landscape_HighRes_Normals.jpg)(w:900 h:272)
+[REGION:imagetable]
+| ![](Landscape_LODs.png) | ![](Landscape_FullResNormals.png) |
+| -------- |-------- |
+| ** Landscape LODs **| ** Full Resolution Normals ** |
+[/REGION]
 
 This makes it possible to always use the highest resolution of the terrain for per-pixel lighting even on distant components that have been LOD'd out.
 
-**Left:** Simple Vertex Lighting
 
-**Right:** 
+[REGION:imagetable]
+| ![](Landscape_SimpleVertexLighting.png) | ![](Landscape_HRPerPixel.png) | 
+| -------- | ----------- | 
+| ** Simple Vertex Lighting **| ** High Res Per-Pixel Lighting ** | 
+[/REGION]
 
-![Landscape_HighRes_Lighting.jpg](Landscape_HighRes_Lighting.jpg)(w:900 h:272)
 
 When this high-resolution normal data is combined with detail normal maps, Landscape terrains can achieve extremely detailed lighting with very little overhead.
 
-**Left:** Geometry Normals Only
 
-**Right:** With Detail Normals
 
-![Landscape_Normals_Comparison.jpg](Landscape_Normals_Comparison.jpg)(w:900 h:272)
+[REGION:imagetable]
+| ![](Landscape_GeometryNormals.png) | ![](Landscape_DetailNormals.png) | 
+| -------- | ----------- | 
+| ** Geometry Normals Only **| ** With Detail Normals ** | 
+[/REGION]
 
 ### PhysX Collision
 
