UEDOC-856

Fixed the image table on this page.

diff --git a/Engine/Documentation/Source/Engine/Landscape/Landscape.INT.udn b/Engine/Documentation/Source/Engine/Landscape/Landscape.INT.udn
--- a/Engine/Documentation/Source/Engine/Landscape/Landscape.INT.udn
+++ b/Engine/Documentation/Source/Engine/Landscape/Landscape.INT.udn
@@ -180,9 +180,9 @@ The Landscape system stores the render data for the terrain in Textures in GPU m
 LODs for Landscape terrains are handled using standard Texture mipmaps. Each mipmap is a level of detail and the mipmap to sample can be specified using the `text2Dlod` HLSL instruction. This allows for a large number of LODs and also allows for smooth LOD transition as both mip levels can be sampled and then the heights and X and Y offsets can be interpolated in the vertex shader creating a clean morphing effect.
 
 
-[REGION:fullwidth raw]
+[REGION:imagetable]
 | ![](Landscape_LOD1.png) | ![](Landscape_LOD1toLOD2.png) |  ![](Landscape_LOD2.png) |
-| -------- | ----------- | -------|
+| -------- | ----------- | ------- |
 | ** Fully LOD 1 **| ** Morphing from LOD 1 to LOD 2 ** | ** Fully LOD 2**|
 [/REGION]
 
