Fixing capitalization on this page while verifying UEDOC-1345

diff --git a/Engine/Documentation/Source/GettingStarted/HowTo/DetailsPanel/DetailsPanel.INT.udn b/Engine/Documentation/Source/GettingStarted/HowTo/DetailsPanel/DetailsPanel.INT.udn
--- a/Engine/Documentation/Source/GettingStarted/HowTo/DetailsPanel/DetailsPanel.INT.udn
+++ b/Engine/Documentation/Source/GettingStarted/HowTo/DetailsPanel/DetailsPanel.INT.udn
@@ -139,7 +139,7 @@ In this step, we will update the **Collision** for the Actor so that it generate
 
 	This will now fire an event when the sphere hits something. 
 
-1.	Off the **OnActorHit** node, drag off the **Other Actor** pin and search for **cast to thirdpersoncharacter** then select it in the window.  
+1.	Off the **OnActorHit** node, drag off the **Other Actor** pin and search for **Cast To ThirdPersonCharacter** then select it in the window.  
 
 	![](Details21.png)
 
