#UE4doc #docs:
Added engine version field
Removed italics from a few terms

diff --git a/Engine/Documentation/Source/Engine/Actors/Triggers/TriggerActors.INT.udn b/Engine/Documentation/Source/Engine/Actors/Triggers/TriggerActors.INT.udn
--- a/Engine/Documentation/Source/Engine/Actors/Triggers/TriggerActors.INT.udn
+++ b/Engine/Documentation/Source/Engine/Actors/Triggers/TriggerActors.INT.udn
@@ -2,6 +2,7 @@ Availability:Public
 Title: Trigger Actors
 Crumbs: %ROOT%, Engine, Engine/Actors
 Description:Actor that can be activated and cause events to occur in the level.
+version: 4.9
 
 [VAR:playlistnav]
 [OBJECT:PlaylistNavigation]
@@ -27,7 +28,7 @@ Description:Actor that can be activated and cause events to occur in the level.
 
 [EXCERPT:Intro]
 **Triggers** are Actors that are used to cause an event to occur when they are interacted with by some other object in the 
-level. In other words, they are used to _trigger_ events in response to some other action in the level. All of the default 
+level. In other words, they are used to trigger events in response to some other action in the level. All of the default 
 Triggers are generally the same, differing only in the shape of the area of influence - box, capsule, and sphere - used by 
 the Trigger to detect if another object has activated it. 
 [/EXCERPT:Intro]
@@ -64,7 +65,7 @@ Level Blueprint for the current level:
 ![](trigger_event.png)
 
 The exec pin of the new event node will fire each time the particular event occurs - in the example above, any time an Actor 
-overlaps (or _runs through_) the Trigger:
+overlaps (or runs through) the Trigger:
 
 <!-- Add GIF/slider showing overlap during play and pin firing in graph -->
 
