#UE4doc #docs:
Removed italics from a few terms

diff --git a/Engine/Documentation/Source/Engine/Actors/Types/CommonActorTypes.INT.udn b/Engine/Documentation/Source/Engine/Actors/Types/CommonActorTypes.INT.udn
--- a/Engine/Documentation/Source/Engine/Actors/Types/CommonActorTypes.INT.udn
+++ b/Engine/Documentation/Source/Engine/Actors/Types/CommonActorTypes.INT.udn
@@ -31,13 +31,13 @@ Unreal Engine to familiarize you with what they are and why you would use them.
 | Type | Icon | Description |
 | ---- | ---- | ----------- |
 | [REGION:tablesection]Mesh & Geometry Actor Types[/REGION] |||
-| **StaticMeshActor** | ![](Shared/Icons/Source/ActorIcons/StaticMesh_40x.png)(convert:false) | The **Static Mesh Actor** is a simple type of Actor that displays a mesh in the scene. Even though the name implies the Actor is static - or cannot move - the _static_ part of the name refers to the type of mesh, a StaticMesh, that is used. These meshes are static in that their geometry does not change. The actual Actor can move or be modified in other ways during play, however. These Actors are generally used as world geometry and decorative meshes to create the environments of the level. See [](Engine/Actors/StaticMeshActor) for more detailed usage information. |
-| **Brush** | ![](Shared/Icons/Source/ActorIcons/Brush_40x.png)(convert:false) | The **Brush Actor** is a basic type of Actor that displays simple 3D geometry in the scene. These Actors can be modified using the Geometry Editing mode in the Level Editor. BrushActors (or just _Brushes_) are commonly used for quickly prototyping environments and blocking out levels for testing gameplay. |
+| **StaticMeshActor** | ![](Shared/Icons/Source/ActorIcons/StaticMesh_40x.png)(convert:false) | The **Static Mesh Actor** is a simple type of Actor that displays a mesh in the scene. Even though the name implies the Actor is static - or cannot move - the static part of the name refers to the type of mesh, a StaticMesh, that is used. These meshes are static in that their geometry does not change. The actual Actor can move or be modified in other ways during play, however. These Actors are generally used as world geometry and decorative meshes to create the environments of the level. See [](Engine/Actors/StaticMeshActor) for more detailed usage information. |
+| **Brush** | ![](Shared/Icons/Source/ActorIcons/Brush_40x.png)(convert:false) | The **Brush Actor** is a basic type of Actor that displays simple 3D geometry in the scene. These Actors can be modified using the Geometry Editing mode in the Level Editor. BrushActors (or just Brushes) are commonly used for quickly prototyping environments and blocking out levels for testing gameplay. |
 | **SkeletalMeshActor** | ![](Shared/Icons/Source/ActorIcons/SkeletalMesh_40x.png)(convert:false) | The **Skeletal Mesh Actor** is a type of Actor that displays an animated mesh, or a SkeletalMesh, whose geometry can be deformed - generally through the use of animation sequences created and exported from external 3D animation applications. These Actors are usually used for things like characters or other living creatures as well as complex machinery; anything that needs to deform or display other complex motion. These are also often used in conjunction with Matinee to create cinematic sequences |
 | [REGION:tablesection]Gameplay Actor Types[/REGION] |||
 | **PlayerStart** | ![](PlayerStart.png)(w:43) | A **Player Start** is an Actor that is placed in the level to designate where the player should start out when the game begins. |
 | **Triggers** | ![](BoxTrigger.png)(w:43) ![](CapsuleTrigger.png)(w:43) ![](SphereTrigger.png)(w:43) | [INCLUDE:Engine/Actors/Triggers#intro] [INCLUDE:#trigger] See [](Engine/Actors/Triggers) for more detailed usage info. |
-| **MatineeActor** | ![](MatineeIcon.png)(w:43) | **Matinee Actors** enable you to animate the properties of Actors over time via the **Matinee** animation tool, to create either dynamic gameplay or cinematic in-game sequences. The system is based on the use of specialized animation tracks in which you can place keyframes to set the values of certain properties of the Actors in the level. The Matinee animation tool is similar to the non-linear editors used for video editing, making it familiar to video professionals. See [](Engine/Matinee) for more detailed usage information. |
+| **MatineeActor** | ![](MatineeIcon.png)(w:43) | **Matinee Actors** enable you to animate the properties of Actors over time via the Matinee animation tool, to create either dynamic gameplay or cinematic in-game sequences. The system is based on the use of specialized animation tracks in which you can place keyframes to set the values of certain properties of the Actors in the level. The Matinee animation tool is similar to the non-linear editors used for video editing, making it familiar to video professionals. See [](Engine/Matinee) for more detailed usage information. |
 | [REGION:tablesection]Light Actor Types[/REGION] |||
 | **PointLight** | ![](Shared/Icons/Source/icon_class_PointLight_40px.png)(convert:false) | [INCLUDE:Engine/Rendering/LightingAndShadows/LightTypes/Point#intro] |
 | **SpotLight** | ![](SpotLight.png)(w:43) | [INCLUDE:Engine/Rendering/LightingAndShadows/LightTypes/Spot#intro] |
diff --git a/Engine/Documentation/Source/Engine/Rendering/LightingAndShadows/LightTypes/Spot/Spot.INT.udn b/Engine/Documentation/Source/Engine/Rendering/LightingAndShadows/LightTypes/Spot/Spot.INT.udn
--- a/Engine/Documentation/Source/Engine/Rendering/LightingAndShadows/LightTypes/Spot/Spot.INT.udn
+++ b/Engine/Documentation/Source/Engine/Rendering/LightingAndShadows/LightTypes/Spot/Spot.INT.udn
@@ -50,7 +50,7 @@ Description:The basics of Spot Lights.
 
 
 [EXCERPT:Intro]
-A **Spot Light** emits light from a single point in a cone shape. Users are given two cones to shape the light - the **Inner Cone Angle** and **Outer Cone Angle**. Within the _Inner Cone Angle_, the light achieves full brightness. As you go from the extent of the inner radius to the extents of the _Outer Cone Angle_, a falloff takes place, creating a penumbra, or softening around the Spot Light's disc of illumination. The Radius of the light defines the length of the cones. More simply, this will work like a flash light or stage can light. 
+A **Spot Light** emits light from a single point in a cone shape. Users are given two cones to shape the light - the **Inner Cone Angle** and **Outer Cone Angle**. Within the Inner Cone Angle, the light achieves full brightness. As you go from the extent of the inner radius to the extents of the Outer Cone Angle, a falloff takes place, creating a penumbra, or softening around the Spot Light's disc of illumination. The Radius of the light defines the length of the cones. More simply, this will work like a flash light or stage can light. 
 
 Like other lights, a Spot Light can be set to one of three Mobility settings: 
 
