Adding more info from DanielL's feedback

diff --git a/Engine/Documentation/Source/Engine/Deployment/Patching/Patching.INT.udn b/Engine/Documentation/Source/Engine/Deployment/Patching/Patching.INT.udn
--- a/Engine/Documentation/Source/Engine/Deployment/Patching/Patching.INT.udn
+++ b/Engine/Documentation/Source/Engine/Deployment/Patching/Patching.INT.udn
@@ -73,10 +73,15 @@ This will determine which content is cooked and then compared to the original pa
 	![](10.png)
 1. The project launcher will go through the building, cooking, and packaging process. This may take some time depending on the complexity of your project.
 ![](11.png)
-1. Once the operation is complete, close the window or click on **Done**. You can test the patch now with the steps below.
+1. Once the operation is complete, close the window or click on **Done**. 
 	![](12.png)
 
 
+The patching process creates a pak file in `[ProjectName]\Saved\StagedBuilds\[PlatformName]\[ProjectName]\Content\Paks`. This pak file is the patch that should be distributed to users.
+For example, on Windows, you would create an installer to copy that pak file into the user's `[ProjectName]\Releases\[VersionNumber]\[PlatformName]` folder, next to the original
+content pak file.  
+
+
 **To test the patch on Windows:**
 
 1. Copy the pak file from `ShooterGame\Releases\1.0\WindowsNoEditor` to `ShooterGame\Saved\StagedBuilds\WindowsNoEditor\ShooterGame\Content\Paks`
