#UE4doc #docs:
Changed font settings for multiple terms
Changed picture for Blocking Volumes to represent new properties

diff --git a/Engine/Documentation/Source/Engine/Actors/Volumes/Volumes.INT.udn b/Engine/Documentation/Source/Engine/Actors/Volumes/Volumes.INT.udn
--- a/Engine/Documentation/Source/Engine/Actors/Volumes/Volumes.INT.udn
+++ b/Engine/Documentation/Source/Engine/Actors/Volumes/Volumes.INT.udn
@@ -99,7 +99,7 @@ The properites of the Camera Blocking Volume can be adjusted in the same manner
 
 ### Trigger Volume 
 
-**Trigger Volumes** are used as a way to cause events when a Player or other object enters or exits them. They can be used quickly with the **Level Blueprint** to test out events and gameplay scenarios or functionality without the need for additional _Blueprints_.
+**Trigger Volumes** are used as a way to cause events when a Player or other object enters or exits them. They can be used quickly with the **Level Blueprint** to test out events and gameplay scenarios or functionality without the need for additional Blueprints.
 
 For example: you could place a TriggerVolume in your level and in your **Level Blueprint**, create an overlap event for the volume which can play a SoundCue, start a Cinematic or even open a door. 
 
@@ -130,11 +130,11 @@ There are some properties that are designated for this volume which can be adjus
 
 |Property| Description |
 | ------- | ------------ |
-| Terminal Velocity | Determines the Terminal Velocity of Pawns using CharacterMovement when falling. |
-| Priority | Determines which PhysicsVolume takes precedence if they overlap.  |
-| Fluid Friction | Determines the amount of friction applied by the volume as Pawns using CharacterMovement move through it. <br> The higher this value, the harder it will feel to move through the volume. |
-| Water Volume | Determines if the volume contains a fluid, like water. |
-| Physics on Contact | Determines if the Actor is affected by the volume by touching it (by Default, an Actor must be inside the volume for it to affect them). |
+| **Terminal Velocity** | Determines the Terminal Velocity of Pawns using CharacterMovement when falling. |
+| **Priority** | Determines which PhysicsVolume takes precedence if they overlap.  |
+| **Fluid Friction** | Determines the amount of friction applied by the volume as Pawns using CharacterMovement move through it. <br> The higher this value, the harder it will feel to move through the volume. |
+| **Water Volume** | Determines if the volume contains a fluid, like water. |
+| **Physics on Contact** | Determines if the Actor is affected by the volume by touching it (by Default, an Actor must be inside the volume for it to affect them). |
 
 
 ### Pain Causing Volume
@@ -147,11 +147,11 @@ In addition to the properties that can be assigned from a Physics Volume, the Pa
 
 |Property| Description |
 | ------- | ------------ |
-| Pain Causing | Whether the volume currently causes damage or not.  |
-| Damage Per Sec | Damage done per second to the Actors in the volume when Pain Causing is enabled. |
-| Damage Type | This determines the type of damage done to the Actor. |
-| Pain Interval | This is the amount of time, in seconds, between applied damage when Pain Causing is enabled.  |
-| Entry Pain | Determines if damage is applied when entering the volume if Pain Causing is enabled in addition to damage applied based on the Pain Interval. |
+| **Pain Causing** | Whether the volume currently causes damage or not.  |
+| **Damage Per Sec** | Damage done per second to the Actors in the volume when Pain Causing is enabled. |
+| **Damage Type** | This determines the type of damage done to the Actor. |
+| **Pain Interval** | This is the amount of time, in seconds, between applied damage when Pain Causing is enabled.  |
+| **Entry Pain** | Determines if damage is applied when entering the volume if Pain Causing is enabled in addition to damage applied based on the Pain Interval. |
 
 
 
@@ -171,10 +171,10 @@ You can adjust how a Level Streaming Volume handles level streaming by adjusting
 
 |Property| Description |
 | ------- | ------------ |
-| Streaming Levels | Displays the levels affected by the volume. |
-| Editor Pre Vis Only | Determines if the streaming volume should only be used for editor streaming level previs. |
-| Disabled | If _true_, the streaming volume is ignored by the streaming volume code. <br> Also used to either disable a Level Streaming Volume without disassociating it from the level, or to toggle the control of a level's streaming between Blueprints and Volume Streaming. |
-| Streaming Usage | Determines what the volume is used for, e.g. whether to control loading, loading and visibility, or just visibility (blocking on load). |
+| **Streaming Levels** | Displays the levels affected by the volume. |
+| **Editor Pre Vis Only** | Determines if the streaming volume should only be used for editor streaming level previs. |
+| **Disabled** | If true, the streaming volume is ignored by the streaming volume code. <br> Also used to either disable a Level Streaming Volume without disassociating it from the level, or to toggle the control of a level's streaming between Blueprints and Volume Streaming. |
+| **Streaming Usage** | Determines what the volume is used for, e.g. whether to control loading, loading and visibility, or just visibility (blocking on load). |
 
 
 ### Cull Distance Volume
@@ -189,7 +189,7 @@ In the image above, the properties define the following behavior:
 
 * Objects within the volume that have a size less than 50 units will be culled (disappear) when they are 500 units away from the camera or farther.
 * Objects within the volume that have a size less than 120 units will be culled (disappear) when they are 1000 units away from the camera or farther.
-* Objects within the volume that have a size less than 300 units **will never be culled**, since 0 in this case is considered to be infinity, meaning that the camera can never get far enough away.
+* Objects within the volume that have a size less than 300 units will never be culled, since 0 in this case is considered to be infinity, meaning that the camera can never get far enough away.
 
 Setup starts with adding a new entry into the Cull Distances array, which is done by clicking the ![](button_Plus.png) button. Next, simply fill in the size of a given object and the desired distance at which you want objects of that size or smaller to be culled. These properties do not have to be created in a specific order. 
 
@@ -202,25 +202,25 @@ Several properties can be adjusted from the **Details** panel on this volume to
 
 |Property| Description |
 | ------- | ------------ |
-| Priority | In the event of overlapping volumes, the highest priority volume will be used. The order is undefined if two or more overlapping volumes have the same priority. |
-| Apply Reverb | Determines if the reverb settings should be used. |
-| Reverb Effect | This is the reverb asset to use for the volume. |
-| Volume | This is the overall volume level of the reverb effect. |
-| Fade Time | This is the time (in seconds) to fade from the current reverb settings into the volumes setting.  |
-| Enabled | Determines whether the volume is currently enabled and is able to affect sounds. |
+| **Priority** | In the event of overlapping volumes, the highest priority volume will be used. The order is undefined if two or more overlapping volumes have the same priority. |
+| **Apply Reverb** | Determines if the reverb settings should be used. |
+| **Reverb Effect** | This is the reverb asset to use for the volume. |
+| **Volume** | This is the overall volume level of the reverb effect. |
+| **Fade Time** | This is the time (in seconds) to fade from the current reverb settings into the volumes setting.  |
+| **Enabled** | Determines whether the volume is currently enabled and is able to affect sounds. |
 
 **Ambient Zones** define how Sound Actors located inside the associated Audio Volume will be altered by the Player's location. The settings for the Ambient Zone can also be adjusted from the **Details** panel. 
 
 |Property| Description |
 | ------- | ------------ |
-| Exterior Volume | The final volume of exterior sounds when the player is inside the volume.  |
-| Exterior Time | Time to fade to new exterior volme in seconds. |
-| Exterior LPF | Lowpass Filter multiplier applied to exterior sounds when inside (1.0 to apply the maximum LPF). |
-| Exterior LPFTime | Time to fade to new Lowpass Filter level in seconds. |
-| Interior Volume | The final volume of interior sounds when the player is outside the volume. |
-| Interior Time | Time to fade to new interior volume in seconds. |
-| Interior LPF | Lowpass Filter multiplier applied to interior sounds when outside (1.0 to apply the maximum LPF). |
-| Interior LPFTime | Time to fade to new Lowpass Filter level in seconds. |
+| **Exterior Volume** | The final volume of exterior sounds when the player is inside the volume.  |
+| **Exterior Time** | Time to fade to new exterior volme in seconds. |
+| **Exterior LPF** | Lowpass Filter multiplier applied to exterior sounds when inside (1.0 to apply the maximum LPF). |
+| **Exterior LPFTime** | Time to fade to new Lowpass Filter level in seconds. |
+| **Interior Volume** | The final volume of interior sounds when the player is outside the volume. |
+| **Interior Time** | Time to fade to new interior volume in seconds. |
+| **Interior LPF** | Lowpass Filter multiplier applied to interior sounds when outside (1.0 to apply the maximum LPF). |
+| **Interior LPFTime** | Time to fade to new Lowpass Filter level in seconds. |
 
 <!--For more information on Ambient Zones, please refer to the following documentation: [Ambient Zones Documentation](Engine/Audio/AmbientZones).-->
 
