UEDOC-1400

Added Next/Previous button to Content Examples

diff --git a/Engine/Documentation/Source/Resources/ContentExamples/Blueprints/1_1/1_1.INT.udn b/Engine/Documentation/Source/Resources/ContentExamples/Blueprints/1_1/1_1.INT.udn
--- a/Engine/Documentation/Source/Resources/ContentExamples/Blueprints/1_1/1_1.INT.udn
+++ b/Engine/Documentation/Source/Resources/ContentExamples/Blueprints/1_1/1_1.INT.udn
@@ -10,6 +10,22 @@ Related: Engine/Animation/AnimBlueprints/EventGraph
 Related: Engine/Blueprints/UserGuide/UserConstructionScript
 Related: Gameplay/HowTo/SetUpInput/Blueprints
 
+[VAR:Steps]
+[OBJECT:Navigation]
+	[PARAM:previous]
+	[/PARAM]
+	[PARAM:current]
+	[/PARAM]
+	[PARAM:home]
+		[Blueprint Content Examples Home](Resources/ContentExamples/Blueprints)
+	[/PARAM]
+	[PARAM:next]
+		[Next Example](Resources/ContentExamples/Blueprints/1_2 "%Resources/ContentExamples/Blueprints/1_2:title%")
+	[/PARAM]
+[/OBJECT]
+[/VAR]
+
+%Steps%
 
 [VAR:TopicCompact]
 [OBJECT:TopicCompact]
diff --git a/Engine/Documentation/Source/Resources/ContentExamples/Blueprints/1_2/1_2.INT.udn b/Engine/Documentation/Source/Resources/ContentExamples/Blueprints/1_2/1_2.INT.udn
--- a/Engine/Documentation/Source/Resources/ContentExamples/Blueprints/1_2/1_2.INT.udn
+++ b/Engine/Documentation/Source/Resources/ContentExamples/Blueprints/1_2/1_2.INT.udn
@@ -11,6 +11,23 @@ Related: Engine/Animation/AnimBlueprints/EventGraph
 Related: Engine/Blueprints/UserGuide/UserConstructionScript
 Related: Gameplay/HowTo/SetUpInput/Blueprints
 
+[VAR:Steps]
+[OBJECT:Navigation]
+	[PARAM:previous]
+		[Previous Example](Resources/ContentExamples/Blueprints/1_1 "%Resources/ContentExamples/Blueprints/1_1:title%")
+	[/PARAM]
+	[PARAM:current]
+	[/PARAM]
+	[PARAM:home]
+		[Blueprint Content Examples Home](Resources/ContentExamples/Blueprints)
+	[/PARAM]
+	[PARAM:next]
+		[Next Example](Resources/ContentExamples/Blueprints/1_3 "%Resources/ContentExamples/Blueprints/1_3:title%")
+	[/PARAM]
+[/OBJECT]
+[/VAR]
+
+%Steps%
 
 [VAR:TopicCompact]
 [OBJECT:TopicCompact]
diff --git a/Engine/Documentation/Source/Resources/ContentExamples/Blueprints/1_3/1_3.INT.udn b/Engine/Documentation/Source/Resources/ContentExamples/Blueprints/1_3/1_3.INT.udn
--- a/Engine/Documentation/Source/Resources/ContentExamples/Blueprints/1_3/1_3.INT.udn
+++ b/Engine/Documentation/Source/Resources/ContentExamples/Blueprints/1_3/1_3.INT.udn
@@ -11,7 +11,23 @@ Related: Engine/Animation/AnimBlueprints/EventGraph
 Related: Engine/Blueprints/UserGuide/UserConstructionScript
 Related: Gameplay/HowTo/SetUpInput/Blueprints
 
+[VAR:Steps]
+[OBJECT:Navigation]
+	[PARAM:previous]
+		[Previous Example](Resources/ContentExamples/Blueprints/1_2 "%Resources/ContentExamples/Blueprints/1_2:title%")
+	[/PARAM]
+	[PARAM:current]
+	[/PARAM]
+	[PARAM:home]
+		[Blueprint Content Examples Home](Resources/ContentExamples/Blueprints)
+	[/PARAM]
+	[PARAM:next]
+		[Next Example](Resources/ContentExamples/Blueprints/1_4 "%Resources/ContentExamples/Blueprints/1_4:title%")
+	[/PARAM]
+[/OBJECT]
+[/VAR]
 
+%Steps%
 
 [VAR:TopicCompact]
 [OBJECT:TopicCompact]
diff --git a/Engine/Documentation/Source/Resources/ContentExamples/Blueprints/1_4/1_4.INT.udn b/Engine/Documentation/Source/Resources/ContentExamples/Blueprints/1_4/1_4.INT.udn
--- a/Engine/Documentation/Source/Resources/ContentExamples/Blueprints/1_4/1_4.INT.udn
+++ b/Engine/Documentation/Source/Resources/ContentExamples/Blueprints/1_4/1_4.INT.udn
@@ -11,7 +11,23 @@ Related: Engine/Animation/AnimBlueprints/EventGraph
 Related: Engine/Blueprints/UserGuide/UserConstructionScript
 Related: Gameplay/HowTo/SetUpInput/Blueprints
 
+[VAR:Steps]
+[OBJECT:Navigation]
+	[PARAM:previous]
+		[Previous Example](Resources/ContentExamples/Blueprints/1_3 "%Resources/ContentExamples/Blueprints/1_3:title%")
+	[/PARAM]
+	[PARAM:current]
+	[/PARAM]
+	[PARAM:home]
+		[Blueprint Content Examples Home](Resources/ContentExamples/Blueprints)
+	[/PARAM]
+	[PARAM:next]
+		[Next Example](Resources/ContentExamples/Blueprints/1_5 "%Resources/ContentExamples/Blueprints/1_5:title%")
+	[/PARAM]
+[/OBJECT]
+[/VAR]
 
+%Steps%
 
 [VAR:TopicCompact]
 [OBJECT:TopicCompact]
diff --git a/Engine/Documentation/Source/Resources/ContentExamples/Blueprints/1_5/1_5.INT.udn b/Engine/Documentation/Source/Resources/ContentExamples/Blueprints/1_5/1_5.INT.udn
--- a/Engine/Documentation/Source/Resources/ContentExamples/Blueprints/1_5/1_5.INT.udn
+++ b/Engine/Documentation/Source/Resources/ContentExamples/Blueprints/1_5/1_5.INT.udn
@@ -11,7 +11,23 @@ Related: Engine/Animation/AnimBlueprints/EventGraph
 Related: Engine/Blueprints/UserGuide/UserConstructionScript
 Related: Gameplay/HowTo/SetUpInput/Blueprints
 
+[VAR:Steps]
+[OBJECT:Navigation]
+	[PARAM:previous]
+		[Previous Example](Resources/ContentExamples/Blueprints/1_4 "%Resources/ContentExamples/Blueprints/1_4:title%")
+	[/PARAM]
+	[PARAM:current]
+	[/PARAM]
+	[PARAM:home]
+		[Blueprint Content Examples Home](Resources/ContentExamples/Blueprints)
+	[/PARAM]
+	[PARAM:next]
+		[Next Example](Resources/ContentExamples/Blueprints/1_6 "%Resources/ContentExamples/Blueprints/1_6:title%")
+	[/PARAM]
+[/OBJECT]
+[/VAR]
 
+%Steps%
 
 [VAR:TopicCompact]
 [OBJECT:TopicCompact]
diff --git a/Engine/Documentation/Source/Resources/ContentExamples/Blueprints/1_6/1_6.INT.udn b/Engine/Documentation/Source/Resources/ContentExamples/Blueprints/1_6/1_6.INT.udn
--- a/Engine/Documentation/Source/Resources/ContentExamples/Blueprints/1_6/1_6.INT.udn
+++ b/Engine/Documentation/Source/Resources/ContentExamples/Blueprints/1_6/1_6.INT.udn
@@ -11,7 +11,23 @@ Related: Engine/Animation/AnimBlueprints/EventGraph
 Related: Engine/Blueprints/UserGuide/UserConstructionScript
 Related: Gameplay/HowTo/SetUpInput/Blueprints
 
+[VAR:Steps]
+[OBJECT:Navigation]
+	[PARAM:previous]
+		[Previous Example](Resources/ContentExamples/Blueprints/1_5 "%Resources/ContentExamples/Blueprints/1_5:title%")
+	[/PARAM]
+	[PARAM:current]
+	[/PARAM]
+	[PARAM:home]
+		[Blueprint Content Examples Home](Resources/ContentExamples/Blueprints)
+	[/PARAM]
+	[PARAM:next]
+		[Next Example](Resources/ContentExamples/Blueprints/1_7 "%Resources/ContentExamples/Blueprints/1_7:title%")
+	[/PARAM]
+[/OBJECT]
+[/VAR]
 
+%Steps%
 
 [VAR:TopicCompact]
 [OBJECT:TopicCompact]
diff --git a/Engine/Documentation/Source/Resources/ContentExamples/Blueprints/1_7/1_7.INT.udn b/Engine/Documentation/Source/Resources/ContentExamples/Blueprints/1_7/1_7.INT.udn
--- a/Engine/Documentation/Source/Resources/ContentExamples/Blueprints/1_7/1_7.INT.udn
+++ b/Engine/Documentation/Source/Resources/ContentExamples/Blueprints/1_7/1_7.INT.udn
@@ -11,6 +11,22 @@ Related: Engine/Animation/AnimBlueprints/EventGraph
 Related: Engine/Blueprints/UserGuide/UserConstructionScript
 Related: Gameplay/HowTo/SetUpInput/Blueprints
 
+[VAR:Steps]
+[OBJECT:Navigation]
+	[PARAM:previous]
+		[Previous Example](Resources/ContentExamples/Blueprints/1_6 "%Resources/ContentExamples/Blueprints/1_6:title%")
+	[/PARAM]
+	[PARAM:current]
+	[/PARAM]
+	[PARAM:home]
+		[Blueprint Content Examples Home](Resources/ContentExamples/Blueprints)
+	[/PARAM]
+	[PARAM:next]
+	[/PARAM]
+[/OBJECT]
+[/VAR]
+
+%Steps%
 
 
 [VAR:TopicCompact]
