#UE4doc #docs:
Changed font settings for multiple terms
Fixed grammar for a few lines
Removed a smiley face

diff --git a/Engine/Documentation/Source/Engine/AI/BehaviorTrees/HowUE4BehaviorTreesDiffer/HowUE4BehaviorTreesDiffer.INT.udn b/Engine/Documentation/Source/Engine/AI/BehaviorTrees/HowUE4BehaviorTreesDiffer/HowUE4BehaviorTreesDiffer.INT.udn
--- a/Engine/Documentation/Source/Engine/AI/BehaviorTrees/HowUE4BehaviorTreesDiffer/HowUE4BehaviorTreesDiffer.INT.udn
+++ b/Engine/Documentation/Source/Engine/AI/BehaviorTrees/HowUE4BehaviorTreesDiffer/HowUE4BehaviorTreesDiffer.INT.udn
@@ -5,7 +5,7 @@ Description:How Unreal Engine 4 Behavior Trees Differ From Traditional Behavior
 
 [TOC(start:2 end:3)]
 
-This section is for people who are familiar with Behavior Trees generally and would like to dive into the Unreal Engine 4 implementation as quickly as possible.  For those who have not used behavior trees before, you may find some of the explanation here confusing. <!-- You may want to check out the user guide first. You can return to this section if you're interested in what we've done differently than others. -->
+This section is for people who are familiar with **Behavior Trees** generally and would like to dive into the Unreal Engine 4 implementation as quickly as possible.  For those who have not used behavior trees before, you may find some of the explanation here confusing. <!-- You may want to check out the user guide first. You can return to this section if you're interested in what we've done differently than others. -->
 
 There are three critical ways in which the UE4 implementation of Behavior Trees differs from "standard" behavior trees:
 
@@ -17,11 +17,11 @@ Having an event-driven architecture grants improvements to both performance and
 
 Since the code does not have to iterate through the entire tree every tick, performance is much better!  Conceptually, instead of constantly asking "Are we there yet?", we can just rest until we are prodded and told "We are there!"
 
-When stepping forward and backward through the behavior tree's execution history to visually debug the behavior, it is ideal to make the history show relevant changes and NOT show irrelevant ones.  In the event-driven implementation, it is not necessary to filter out irrelevant steps that iterated over the tree and chose the same behavior as before, because that additional iteration never had to happen in the first place!  Instead, only changes to the location of execution in the tree or to blackboard values matter, and it is easy to just show those differences.  <!-- (See the debugging section for more information on how to step through the execution history and other debugging tricks.) -->
+When stepping forward and backward through the behavior tree's execution history to visually debug the behavior, it is ideal to make the history show relevant changes and not show irrelevant ones.  In the event-driven implementation, it is not necessary to filter out irrelevant steps that iterated over the tree and chose the same behavior as before, because that additional iteration never had to happen in the first place!  Instead, only changes to the location of execution in the tree or to blackboard values matter, and it is easy to just show those differences.  <!-- (See the debugging section for more information on how to step through the execution history and other debugging tricks.) -->
 
 ## Conditionals Are Not Leaf Nodes
 
-In the standard model for behavior trees, conditionals are "Task" leaf nodes, which simply do not do anything other than succeed or fail.  Although nothing prevents you from making traditional conditional tasks, it is highly recommended that you use our Decorator system for conditionals instead.
+In the standard model for behavior trees, conditionals are Task leaf nodes, which simply do not do anything other than succeed or fail.  Although nothing prevents you from making traditional conditional tasks, it is highly recommended that you use our Decorator system for conditionals instead.
 
 Making conditionals be decorators rather than tasks has a couple of significant advantages.
 
@@ -29,7 +29,7 @@ First, conditional decorators make the behavior tree UI more intuitive and easie
 
 ![](decorator.png)
 
-_In this section of a Behavior Tree, the Decorators **Close Enough** and **Blackboard** can prevent the execution of the Sequence node's children._
+In this section of a Behavior Tree, the Decorators **Close Enough** and **Blackboard** can prevent the execution of the Sequence node's children.
 
 Another advantage of conditional decorators is that it is easy to make those decorators act as observers (waiting for events) at critical nodes in the tree.  This feature is critical to gaining full advantage from the event-driven nature of the trees.  <!-- Read more about how to use Conditional Decorators and their special properties here. -->
 
@@ -38,7 +38,7 @@ Another advantage of conditional decorators is that it is easy to make those dec
  Standard behavior trees often use a **Parallel** composite node to handle concurrent behaviors.  The Parallel node begins execution on all of its children simultaneously.  Special rules determine how to act if one or more of those child trees finish (depending on the desired behavior).
 
 [REGION:note]
-Parallel nodes are not necessarily multi-threading (executing tasks at "truly" the same time).  They are just a way to conceptually perform several tasks at once.  Often they still run on the same thread and begin in some sequence.  That sequence should be irrelevant since they will all happen in the same frame, but it is still sometimes important.
+Parallel nodes are not necessarily multi-threading (executing tasks at truly the same time).  They are just a way to conceptually perform several tasks at once.  Often they still run on the same thread and begin in some sequence.  That sequence should be irrelevant since they will all happen in the same frame, but it is still sometimes important.
 [/REGION]
 
 Instead of complex Parallel nodes, UE4 Behavior Trees use **Simple Parallel** nodes and our own special node type which we call Services to accomplish the same sorts of behaviors.
@@ -47,7 +47,7 @@ Instead of complex Parallel nodes, UE4 Behavior Trees use **Simple Parallel** no
 
 1.	Parallel nodes can be very confusing, even for relatively simple behaviors. 
 
-	Effectively Parallel nodes are simultaneously running a bunch of separate sub-trees, but any or all of those sub-trees may need to abort if one of them fails.  Or they may succeed when the others finish (whether successful or failing).  Or... etc.!  Parallel behaviors can be confusing even in simple cases, and with the number of options potentially available, it can become highly confusing.
+	Effectively Parallel nodes are simultaneously running a bunch of separate sub-trees, but any or all of those sub-trees may need to abort if one of them fails, or they may succeed when the others finish (whether successful or failing).  Parallel behaviors can be confusing even in simple cases, and with the number of options potentially available, it can become highly confusing.
 
 1.	Parallel nodes make it harder to optimize performance, especially in terms of making event-driven trees.
 
@@ -69,7 +69,7 @@ Simple Parallel nodes allow easy usage of some of our event-driven optimizations
 
 #### Services
 
-Services are special nodes associated with any composite node (Selector, Sequence, or Simple Parallel), which can register for callbacks every X seconds and perform updates of various sorts that need to occur periodically.
+**Services** are special nodes associated with any composite node (Selector, Sequence, or Simple Parallel), which can register for callbacks every X seconds and perform updates of various sorts that need to occur periodically.
 
 For example, a service can be used to determine which enemy is the best choice for the AI pawn to pursue while the pawn continues to act normally in its behavior tree toward its current enemy.
 
@@ -79,7 +79,7 @@ Services are active only as long as execution remains in the subtree rooted at t
 
 #### Decorator "Observer Aborts" Property
 
-One common usage case for standard Parallel nodes is to constantly check conditions so that a task can abort if the conditions it requires becomes _false_.  For example, if you have a cat that performs a sequence "Shake Rear End", "Pounce", you may want to give up immediately if the mouse escapes into its mouse hole.  With Parallel nodes, you would have a child that checks if the mouse can be pounced on, and then another child that's the sequence to perform.  Since our Behavior Trees are event-driven, we instead handle this by having our conditional decorators observe their values and abort when necessary.  (In this example, you would just have the "Mouse Can Be Pounced On?" decorator on the sequence itself, with "Observer Aborts" set to "Self". )
+One common usage case for standard Parallel nodes is to constantly check conditions so that a task can abort if the conditions it requires becomes false.  For example, if you have a cat that performs a sequence, such as "Shake Rear End" and "Pounce", you may want to give up immediately if the mouse escapes into its mouse hole.  With Parallel nodes, you would have a child that checks if the mouse can be pounced on, and then another child that's the sequence to perform.  Since our Behavior Trees are event-driven, we instead handle this by having our conditional decorators observe their values and abort when necessary.  (In this example, you would just have the "Mouse Can Be Pounced On?" decorator on the sequence itself, with "Observer Aborts" set to "Self". )
 
 <!-- For more information, see the <LINK>"Observer Aborts" explanation page. -->
 
@@ -98,16 +98,16 @@ Clearer graphs are easier to debug.  In addition, having fewer simultaneous exec
 
 ### Easier Optimizations
 
-Event-driven graphs are easier to optimize if they do not have a lot of subtrees "simultaneously" executing.
+Event-driven graphs are easier to optimize if they do not have a lot of subtrees simultaneously executing.
 
 
 ## FAQ
 
 "Can you really do everything you can do with Parallel nodes?"
 
-* We believe you can do everything necessary with the nodes we are providing, with a better interface!  Certainly the nodes above handle the most common cases.  If we find any edge cases that cannot be handled or are less than ideal, we will consider additional fixes to handle those cases.
+* We believe you can do everything necessary with the nodes we are providing, with a better interface. Certainly the nodes above handle the most common cases.  If we find any edge cases that cannot be handled or are less than ideal, we will consider additional fixes to handle those cases.
 
 
 "Are these the only differences between UE4 Behavior Trees and a "standard" behavior tree?"
 
-* "Standard" is in quotes for a reason.  ;)  There is really no such thing as "standard", so there could be any number of differences between UE4's implementation and whatever implementation you are most familiar with.  If you are familiar with an unusual implementation, it may have other critical differences, and there are likely more subtle differences regardless.  Hopefully these notes give you an idea of the most important differences relevant to how you will need to build your trees.  For more information about our special types of nodes, read the sections about those nodes.
+* "Standard" is in quotes for a reason. There is really no such thing as "standard", so there could be any number of differences between UE4's implementation and whatever implementation you are most familiar with.  If you are familiar with an unusual implementation, it may have other critical differences, and there are likely more subtle differences regardless.  Hopefully these notes give you an idea of the most important differences relevant to how you will need to build your trees.  For more information about our special types of nodes, read the sections about those nodes.
