Minor fixes

diff --git a/Engine/Documentation/Source/Engine/Deployment/Patching/Patching.INT.udn b/Engine/Documentation/Source/Engine/Deployment/Patching/Patching.INT.udn
--- a/Engine/Documentation/Source/Engine/Deployment/Patching/Patching.INT.udn
+++ b/Engine/Documentation/Source/Engine/Deployment/Patching/Patching.INT.udn
@@ -30,6 +30,7 @@ In this example, we are patching ShooterGame with a change to the weapon capacit
 
 	![](1.png)
 1. Create a new Custom Launch Profile using the **+** button.
+	
 	![](2.png)
 1. Set a name and description for your profile. 
 
@@ -67,20 +68,25 @@ This will determine which content is cooked and then compared to the original pa
 * Set the build to **Do Not Deploy**.
 
 1. Once you have set all the above settings, navigate back to the main profile window using the **Back** button in the top right corner. 
+	
 	![](9.png)
 
 1. Click on the launch icon next to your **Patching** profile. 
+	
 	![](10.png)
 1. The project launcher will go through the building, cooking, and packaging process. This may take some time depending on the complexity of your project.
-![](11.png)
+	
+	![](11.png)
 1. Once the operation is complete, close the window or click on **Done**. 
+	
 	![](12.png)
 
 
+[REGION:tip]
 The patching process creates a pak file in `[ProjectName]\Saved\StagedBuilds\[PlatformName]\[ProjectName]\Content\Paks`. This pak file is the patch that should be distributed to users.
 For example, on Windows, you would create an installer to copy that pak file into the user's `[ProjectName]\Releases\[VersionNumber]\[PlatformName]` folder, next to the original
 content pak file.  
-
+[/REGION]
 
 **To test the patch on Windows:**
 
