#UE4doc #docs:
Changed font settings for multiple terms

diff --git a/Engine/Documentation/Source/Engine/AI/BehaviorTrees/NodeReference/Composites/Composites.INT.udn b/Engine/Documentation/Source/Engine/AI/BehaviorTrees/NodeReference/Composites/Composites.INT.udn
--- a/Engine/Documentation/Source/Engine/AI/BehaviorTrees/NodeReference/Composites/Composites.INT.udn
+++ b/Engine/Documentation/Source/Engine/AI/BehaviorTrees/NodeReference/Composites/Composites.INT.udn
@@ -11,7 +11,7 @@ Availability:Public
 
 ![](Selector.png)
 
-**Selector** Nodes execute their children from left to right, and will stop executing its children when one of their children **Succeeds**. If a Selector's child succeeds, the Selector succeeds. If all the Selector's children fail, the Selector fails.
+**Selector** Nodes execute their children from left to right, and will stop executing its children when one of their children succeeds. If a Selector's child succeeds, the Selector succeeds. If all the Selector's children fail, the Selector fails.
 
 | Property | Description |
 | -------- | ----------- |
@@ -21,7 +21,7 @@ Availability:Public
 
 ![](Sequence.png)
 
-**Sequence** Nodes execute their children from left to right, and will stop executing its children when one of their children **Fails**. If a child fails, then the Sequence fails. If all the Sequence's children succeed, then the Sequence succeeds.
+**Sequence** Nodes execute their children from left to right, and will stop executing its children when one of their children fails. If a child fails, then the Sequence fails. If all the Sequence's children succeed, then the Sequence succeeds.
 
 | Property | Description |
 | -------- | ----------- |
@@ -31,7 +31,7 @@ Availability:Public
 
 ![](SimpleParallel.png)
 
-The Simple Parallel node allows a single main task node to be executed along side of a full tree. When the main task finishes, the setting in **Finish Mode** dictates if the node should finish **Immediately**, aborting the secondary tree, or if it should **Delay** for the secondary tree to finish.
+The Simple Parallel node allows a single main task node to be executed along side of a full tree. When the main task finishes, the setting in **Finish Mode** dictates if the node should finish immediately, aborting the secondary tree, or if it should delay for the secondary tree to finish.
 
 | Property | Description |
 | -------- | ----------- |
