#UE4doc #docs:
Unbolded a few terms

diff --git a/Engine/Documentation/Source/Engine/AI/EnvironmentQuerySystem/EnvironmentQuerySystem.INT.udn b/Engine/Documentation/Source/Engine/AI/EnvironmentQuerySystem/EnvironmentQuerySystem.INT.udn
--- a/Engine/Documentation/Source/Engine/AI/EnvironmentQuerySystem/EnvironmentQuerySystem.INT.udn
+++ b/Engine/Documentation/Source/Engine/AI/EnvironmentQuerySystem/EnvironmentQuerySystem.INT.udn
@@ -42,8 +42,8 @@ Description:Landing Page for all the documentation relating to the Environment Q
 [/OBJECT]
 [/VAR]
 
-The **Environment Query System** is a feature of the **Artificial Intelligence** system in Unreal Engine 4 for collecting data on the
-environment, asking questions of the data through **Tests**, then returning then one **Item** that best fits the questions asked.
+The **Environment Query System** is a feature of the Artificial Intelligence system in Unreal Engine 4 for collecting data on the
+environment, asking questions of the data through Tests, then returning then one Item that best fits the questions asked.
 
 You can use it to find the closest power up, figure out which enemy is the highest threat, or to find cover.
 
