#UE4doc #docs:
Changed font settings for a few terms
Added engine version field

diff --git a/Engine/Documentation/Source/Engine/AI/EnvironmentQuerySystem/EQSPawn/EQSPawn.INT.udn b/Engine/Documentation/Source/Engine/AI/EnvironmentQuerySystem/EQSPawn/EQSPawn.INT.udn
--- a/Engine/Documentation/Source/Engine/AI/EnvironmentQuerySystem/EQSPawn/EQSPawn.INT.udn
+++ b/Engine/Documentation/Source/Engine/AI/EnvironmentQuerySystem/EQSPawn/EQSPawn.INT.udn
@@ -2,6 +2,7 @@ Availability:Public
 Title:The EQS Testing Pawn
 Crumbs:%ROOT%, Engine, Gameplay/AI, Engine/AI/EnvironmentQuerySystem
 Description:The Environment Query System includes a specialized Pawn that will display real time results of any Environment Query in the editor. This document covers creating and using an EQSTestingPawn.
+version: 4.9
 
 [VAR:Topic]
 [OBJECT:Topic]
@@ -44,7 +45,7 @@ Description:The Environment Query System includes a specialized Pawn that will d
 
 ![](testingPawnBanner.png)(convert:false)
 
-The **Environment Query System Testing Pawn** (EQSTestingPawn) is a specialized Pawn class that enables you to check out what your Environment Queries are actually doing. The exact makeup of your Environment Query will define the size and shape of what is created, but it will always be represented as colored spheres. Spheres that are on the color scale from green to red, indicate some level of match for the various tests your Environment Query ran. Blue spheres indicate a failure or a boolean type test that returned false, such as the **Trace** test.
+The **Environment Query System Testing Pawn** (**EQSTestingPawn**) is a specialized Pawn class that enables you to check out what your Environment Queries are actually doing. The exact makeup of your Environment Query will define the size and shape of what is created, but it will always be represented as colored spheres. Spheres that are on the color scale from green to red, indicate some level of match for the various tests your Environment Query ran. Blue spheres indicate a failure or a boolean type test that returned false, such as the **Trace** test.
 
 ## Creating an EQS Testing Pawn
 
@@ -63,7 +64,7 @@ The **Environment Query System Testing Pawn** (EQSTestingPawn) is a specialized
 [![](placeEQSPawn.png)(convert:false)](placeEQSPawn.png)
 [/REGION]
 
-1. Place the **EQSTestingPawn** Blueprint into your level, by clicking and dragging it from the **Content Browser** to the **Viewport**.
+1. Place the EQSTestingPawn Blueprint into your level, by clicking and dragging it from the **Content Browser** to the **Viewport**.
 
 	[REGION:note]
 	Given the experimental nature of the Environment Query System, the EQSTestingPawn can be very processor intensive, and making changes to its active Query Template can cause your system to hang for extremely long times. It's best to clear the Query Template property while making edits to your Environment Query.
