#UE4doc #docs:
Changed font settings for multiple terms
Added engine version field

diff --git a/Engine/Documentation/Source/Engine/AI/EnvironmentQuerySystem/NodeReference/NodeReference.INT.udn b/Engine/Documentation/Source/Engine/AI/EnvironmentQuerySystem/NodeReference/NodeReference.INT.udn
--- a/Engine/Documentation/Source/Engine/AI/EnvironmentQuerySystem/NodeReference/NodeReference.INT.udn
+++ b/Engine/Documentation/Source/Engine/AI/EnvironmentQuerySystem/NodeReference/NodeReference.INT.udn
@@ -2,6 +2,7 @@
 Title: The  Environment Query System Node Reference
 Crumbs:%ROOT%, Engine, Gameplay/AI, Engine/AI/EnvironmentQuerySystem
 Description:Break downs and descriptions of the various nodes available in the Environment Query System editor.
+version: 4.9
 
 [VAR:Topic]
 [OBJECT:Topic]
@@ -128,7 +129,7 @@ Generates a grid of Items around the Querier.
 
 ### Distance
 
-This will return the direct distance between the Item and the chosen **Distance To**. If the **Distance To** is more than one location, it averages the results of each distance check.
+This will return the direct distance between the Item and the chosen **Distance To**. If the Distance To is more than one location, it averages the results of each distance check.
 
 ![](test_dist.png)
 
@@ -178,13 +179,13 @@ Calculates the Dot Product of two vectors. These can be Context Rotations, or th
 
 <!--
 [EXCERPT:dotLine]
-**Mode** set to Rotation:
+Mode set to Rotation:
 
 |Property | Description|
 |---|---|
 | Rotation | Returns the rotation vector or the chosen context. |
 
-**Mode** set to Two Points:
+Mode set to Two Points:
 
 |Property | Description|
 |---|---|
@@ -225,7 +226,7 @@ Calculates the Dot Product of two vectors. These can be Context Rotations, or th
 | **Skip Unreachable** || Filter out any Item that is unreachable. |
 | **Use Hierarchical Pathfinding** || Should Hierarchical Pathfinding be used, as it's faster, though less accurate. |
 | **Test Purpose** || Should the test Filter, Score, or Filter and Score the result? Filtering will remove Items that fail the test, while scoring will weight an Item. |
-| **Filter Type** || Change the type of filter to be a minimum, maximum, or range of values. Any values outside the values set in **Float Value Min** and/or **Float Value Max**, will be culled. |
+| **Filter Type** || Change the type of filter to be a minimum, maximum, or range of values. Any values outside the values set in Float Value Min and/or Float Value Max, will be culled. |
 | **Float Value Min** || Filter any value less than or equal to this value. |
 | **Float Value Max** || Filter any value greater than or equal to this value. |
 | **Clamping** || If a value isn't filtered, you can clamp it to limit a test's influence. Change the Clamp Min/Max Type properties to enable this. |
@@ -257,7 +258,7 @@ Less of a test (though it can be used to randomly filter out Items), and more of
 | Property | %spacer% | Description |
 | --- | --- | :--- |
 | **Test Purpose** || Should the test Filter, Score, or Filter and Score the result? Filtering will remove Items that fail the test, while scoring will weight an Item. |
-| **Filter Type** || Change the type of filter to be a minimum, maximum, or range of values. Any values outside the values set in **Float Value Min** and/or **Float Value Max**, will be culled. |
+| **Filter Type** || Change the type of filter to be a minimum, maximum, or range of values. Any values outside the values set in Float Value Min and/or Float Value Max, will be culled. |
 | **Float Value Min** || Filter any value less than or equal to this value. |
 | **Float Value Max** || Filter any value greater than or equal to this value. |
 | **Clamping** || If a value isn't filtered, you can clamp it to limit a test's influence. Change the Clamp Min/Max Type properties to enable this. |
@@ -289,7 +290,7 @@ This will trace to or from an Item to a Context and return if it hit it or not.
 | **Item Height Offset** || This will add a Z offset, in cm, to the Item the test is tracing to/from. |
 | **Context Height Offset** || This will add a Z offset, in cm, to the Context the test is tracing to/from. |
 | **Test Purpose** || Should the test Filter, Score, or Filter and Score the result? Filtering will remove Items that fail the test, while scoring will weight an Item. |
-| **Filter Type** || Change the type of filter to be a minimum, maximum, or range of values. Any values outside the values set in **Float Value Min** and/or **Float Value Max**, will be culled. |
+| **Filter Type** || Change the type of filter to be a minimum, maximum, or range of values. Any values outside the values set in Float Value Min and/or Float Value Max, will be culled. |
 | **Float Value Min** || Filter any value less than or equal to this value. |
 | **Float Value Max** || Filter any value greater than or equal to this value. |
 | **Clamping** || If a value isn't filtered, you can clamp it to limit a test's influence. Change the Clamp Min/Max Type properties to enable this. |
