#UE4doc #docs:
Removed unncessary quotations from a term

diff --git a/Engine/Documentation/Source/Engine/AI/EnvironmentQuerySystem/QuickStart/QuickStart.INT.udn b/Engine/Documentation/Source/Engine/AI/EnvironmentQuerySystem/QuickStart/QuickStart.INT.udn
--- a/Engine/Documentation/Source/Engine/AI/EnvironmentQuerySystem/QuickStart/QuickStart.INT.udn
+++ b/Engine/Documentation/Source/Engine/AI/EnvironmentQuerySystem/QuickStart/QuickStart.INT.udn
@@ -43,7 +43,7 @@ Description:This Quick Start Guide covers creating a Behavior Tree that uses mul
 [/VAR]
 
 
-This Quick Start will teach you the basics of using the **Environment Query System**. In it, you will use Environment Queries to find Actors and Locations, then, using a series of tests, you'll be given the "best" option. Then, using a Behavior tree, you'll have the AI flee from the Player Pawn (or Pawns) to a hiding spot.
+This Quick Start will teach you the basics of using the **Environment Query System**. In it, you will use Environment Queries to find Actors and Locations, then, using a series of tests, you'll be given the best option. Then, using a Behavior tree, you'll have the AI flee from the Player Pawn (or Pawns) to a hiding spot.
 
 This is a more advanced Quick Start and assumes you are familiar with [Blueprints](Engine/Blueprints/BP_HowTo/BasicUsage) and [Behavior Trees](Engine/AI/BehaviorTrees/QuickStart).
 
