UEDOC-1400

Added Previous and Next Labels to this page

diff --git a/Engine/Documentation/Source/Resources/ContentExamples/Blueprints_HUD/1_1/1_1.INT.udn b/Engine/Documentation/Source/Resources/ContentExamples/Blueprints_HUD/1_1/1_1.INT.udn
--- a/Engine/Documentation/Source/Resources/ContentExamples/Blueprints_HUD/1_1/1_1.INT.udn
+++ b/Engine/Documentation/Source/Resources/ContentExamples/Blueprints_HUD/1_1/1_1.INT.udn
@@ -10,6 +10,21 @@ Related: Engine/Animation/AnimBlueprints/EventGraph
 Related: Engine/Blueprints/UserGuide/UserConstructionScript
 Related: Gameplay/HowTo/SetUpInput/Blueprints
 
+[VAR:Steps]
+[OBJECT:Navigation]
+	[PARAM:previous]
+	[/PARAM]
+	[PARAM:current]
+	[/PARAM]
+	[PARAM:home]
+		[Blueprint Content Examples Home](Resources/ContentExamples/Blueprints)
+	[/PARAM]
+	[PARAM:next]
+	[/PARAM]
+[/OBJECT]
+[/VAR]
+
+%Steps%
 
 
 [VAR:TopicCompact]
