#UE4doc #docs:
Unbolded a couiple hyperlinks
Fixed mention of Starter Content as it was being referred to as "Starter Assets"

diff --git a/Engine/Documentation/Source/Engine/AI/BehaviorTrees/QuickStart/1/1.INT.udn b/Engine/Documentation/Source/Engine/AI/BehaviorTrees/QuickStart/1/1.INT.udn
--- a/Engine/Documentation/Source/Engine/AI/BehaviorTrees/QuickStart/1/1.INT.udn
+++ b/Engine/Documentation/Source/Engine/AI/BehaviorTrees/QuickStart/1/1.INT.udn
@@ -33,10 +33,10 @@ This Quick Start makes some assumptions about what is available in the project,
 
 1. Generate a new project based on the **Blueprint Top Down** template
 
-	The **Starter Assets** are optional for this quick start guide.
+	**Starter Content** is optional for this quick start guide.
 
 
-If you need further explanation of starting a new project, see the [**Project Browser**](Engine/Basics/Projects/Browser) User Guide.
+If you need further explanation of starting a new project, see the [Project Browser](Engine/Basics/Projects/Browser) User Guide.
 
 [/EXCERPT:AIProjectStart]
 %Steps%
diff --git a/Engine/Documentation/Source/Engine/AI/EnvironmentQuerySystem/QuickStart/1/1.INT.udn b/Engine/Documentation/Source/Engine/AI/EnvironmentQuerySystem/QuickStart/1/1.INT.udn
--- a/Engine/Documentation/Source/Engine/AI/EnvironmentQuerySystem/QuickStart/1/1.INT.udn
+++ b/Engine/Documentation/Source/Engine/AI/EnvironmentQuerySystem/QuickStart/1/1.INT.udn
@@ -23,7 +23,7 @@ Description: Create a new project using the Top Down Template.
 This Quick Start makes some assumptions about what is available in the project, as such, if you aren't comfortable making substations or problem solving on the fly, then you should setup the project as follows.
 
 [REGION:note]
-If you have already completed the [**Behavior Tree Quick Start**](Engine/AI/BehaviorTrees/QuickStart), feel free to use that project instead of starting a new one.
+If you have already completed the [Behavior Tree Quick Start](Engine/AI/BehaviorTrees/QuickStart), feel free to use that project instead of starting a new one.
 [/REGION]
 
 ## Steps
