#UE4doc #docs:
Removed unnecessary quotations
Added engine version field

diff --git a/Engine/Documentation/Source/Engine/AI/EnvironmentQuerySystem/QuickStart/2/2.INT.udn b/Engine/Documentation/Source/Engine/AI/EnvironmentQuerySystem/QuickStart/2/2.INT.udn
--- a/Engine/Documentation/Source/Engine/AI/EnvironmentQuerySystem/QuickStart/2/2.INT.udn
+++ b/Engine/Documentation/Source/Engine/AI/EnvironmentQuerySystem/QuickStart/2/2.INT.udn
@@ -2,6 +2,7 @@ Availability:Public
 Title: 2. Enable the Environment Query System
 Crumbs:%ROOT%, Engine, Gameplay/AI, Engine/AI/EnvironmentQuerySystem, Engine/AI/EnvironmentQuerySystem/QuickStart
 Description: Enable the Environment Query System as it is currently experimental.
+version: 4.9
 
 [VAR:Steps]
 [OBJECT:Navigation]
@@ -21,7 +22,7 @@ Description: Enable the Environment Query System as it is currently experimental
 
 %Steps%
 
-The **Environment Query System** is considered **Experimental**, meaning it can change drastically at any time, so you must "opt in" to it.
+The **Environment Query System** is considered **Experimental**, meaning it can change drastically at any time, so you must opt in to it.
 
 ## Steps
 
