#UE4doc #docs:
Removed unnecessary quotations
Added engine version field

diff --git a/Engine/Documentation/Source/Engine/AI/BehaviorTrees/QuickStart/2/2.INT.udn b/Engine/Documentation/Source/Engine/AI/BehaviorTrees/QuickStart/2/2.INT.udn
--- a/Engine/Documentation/Source/Engine/AI/BehaviorTrees/QuickStart/2/2.INT.udn
+++ b/Engine/Documentation/Source/Engine/AI/BehaviorTrees/QuickStart/2/2.INT.udn
@@ -30,7 +30,7 @@ The NavMesh in Unreal has a number of functions that work well with Behavior Tre
 
 ## Steps
 
-1. [REGION:raw]**Modes Panel** %arrow% **Create** %arrow% **Volumes**.[/REGION]
+1. [REGION:raw]**Modes** panel %arrow% **Create** %arrow% **Volumes**.[/REGION]
 
 1. Click and drag a **Nav Mesh Bounds Volume** into the map.
 
@@ -40,13 +40,13 @@ The NavMesh in Unreal has a number of functions that work well with Behavior Tre
 
 	* You can use more than one Volume if you like, or one large one that encapsulates the entire level.
 
-1. While working in the Viewport, press the "**P**" key to see if your NavMesh is building.
+1. While working in the Viewport, press the **P** key to see if your NavMesh is building.
 
 	![](NMVis.png)
 
 	If you see a green mesh overlay on your floors, it is building correctly. When the mesh turns red, it means it is recalculating a section because you either moved the volume or the mesh inside of it.
 
-	You can press the "**P**" key at any time to show or hide the NavMesh.
+	You can press the **P** key at any time to show or hide the NavMesh.
 
 [/EXCERPT:AICreateNav]
 %Steps%
diff --git a/Engine/Documentation/Source/Engine/AI/EnvironmentQuerySystem/QuickStart/3/3.INT.udn b/Engine/Documentation/Source/Engine/AI/EnvironmentQuerySystem/QuickStart/3/3.INT.udn
--- a/Engine/Documentation/Source/Engine/AI/EnvironmentQuerySystem/QuickStart/3/3.INT.udn
+++ b/Engine/Documentation/Source/Engine/AI/EnvironmentQuerySystem/QuickStart/3/3.INT.udn
@@ -2,6 +2,7 @@ Availability:Public
 Title: 3. Create Nav Mesh Volume
 Crumbs:%ROOT%, Engine, Gameplay/AI, Engine/AI/EnvironmentQuerySystem, Engine/AI/EnvironmentQuerySystem/QuickStart
 Description: Create a Nav Mesh Volume so our AI can walk around.
+version: 4.9
 
 [VAR:Steps]
 [OBJECT:Navigation]
