UEDOC-1400

Adding Previous and Next to the Content Examples pages

diff --git a/Engine/Documentation/Source/Resources/ContentExamples/Decals/1_1/1_1.INT.udn b/Engine/Documentation/Source/Resources/ContentExamples/Decals/1_1/1_1.INT.udn
--- a/Engine/Documentation/Source/Resources/ContentExamples/Decals/1_1/1_1.INT.udn
+++ b/Engine/Documentation/Source/Resources/ContentExamples/Decals/1_1/1_1.INT.udn
@@ -11,6 +11,22 @@ Related: Engine/Rendering/Materials/MaterialInputs
 Related: Engine/Rendering/Materials/LayeredMaterials
 Related: Engine/Rendering/Materials/MaterialInstances
 
+[VAR:Steps]
+[OBJECT:Navigation]
+	[PARAM:previous]
+	[/PARAM]
+	[PARAM:current]
+	[/PARAM]
+	[PARAM:home]
+		[Decals Home](Resources/ContentExamples/Decals)
+	[/PARAM]
+	[PARAM:next]
+		[Next Example](Resources/ContentExamples/Decals/1_2 "%Resources/ContentExamples/Decals/1_2:title%")
+	[/PARAM]
+[/OBJECT]
+[/VAR]
+
+%Steps%
 
 [VAR:Topic]
 [OBJECT:Topic]
diff --git a/Engine/Documentation/Source/Resources/ContentExamples/Decals/1_2/1_2.INT.udn b/Engine/Documentation/Source/Resources/ContentExamples/Decals/1_2/1_2.INT.udn
--- a/Engine/Documentation/Source/Resources/ContentExamples/Decals/1_2/1_2.INT.udn
+++ b/Engine/Documentation/Source/Resources/ContentExamples/Decals/1_2/1_2.INT.udn
@@ -11,6 +11,24 @@ Related: Engine/Rendering/Materials/MaterialInputs
 Related: Engine/Rendering/Materials/LayeredMaterials
 Related: Engine/Rendering/Materials/MaterialInstances
 
+[VAR:Steps]
+[OBJECT:Navigation]
+	[PARAM:previous]
+		[Previous Example](Resources/ContentExamples/Decals/1_1 "%Resources/ContentExamples/Decals/1_1:title%")
+	[/PARAM]
+	[PARAM:current]
+	[/PARAM]
+	[PARAM:home]
+		[Decals Home](Resources/ContentExamples/Decals)
+	[/PARAM]
+	[PARAM:next]
+		[Next Example](Resources/ContentExamples/Decals/1_3 "%Resources/ContentExamples/Decals/1_3:title%")
+	[/PARAM]
+[/OBJECT]
+[/VAR]
+
+%Steps%
+
 [VAR:Topic]
 [OBJECT:Topic]
 	[PARAM:image]
diff --git a/Engine/Documentation/Source/Resources/ContentExamples/Decals/1_3/1_3.INT.udn b/Engine/Documentation/Source/Resources/ContentExamples/Decals/1_3/1_3.INT.udn
--- a/Engine/Documentation/Source/Resources/ContentExamples/Decals/1_3/1_3.INT.udn
+++ b/Engine/Documentation/Source/Resources/ContentExamples/Decals/1_3/1_3.INT.udn
@@ -11,6 +11,24 @@ Related: Engine/Rendering/Materials/MaterialInputs
 Related: Engine/Rendering/Materials/LayeredMaterials
 Related: Engine/Rendering/Materials/MaterialInstances
 
+[VAR:Steps]
+[OBJECT:Navigation]
+	[PARAM:previous]
+		[Previous Example](Resources/ContentExamples/Decals/1_2 "%Resources/ContentExamples/Decals/1_2:title%")
+	[/PARAM]
+	[PARAM:current]
+	[/PARAM]
+	[PARAM:home]
+		[Decals Home](Resources/ContentExamples/Decals)
+	[/PARAM]
+	[PARAM:next]
+		[Next Example](Resources/ContentExamples/Decals/1_4 "%Resources/ContentExamples/Decals/1_4:title%")
+	[/PARAM]
+[/OBJECT]
+[/VAR]
+
+%Steps%
+
 [VAR:Topic]
 [OBJECT:Topic]
 	[PARAM:image]
diff --git a/Engine/Documentation/Source/Resources/ContentExamples/Decals/1_4/1_4.INT.udn b/Engine/Documentation/Source/Resources/ContentExamples/Decals/1_4/1_4.INT.udn
--- a/Engine/Documentation/Source/Resources/ContentExamples/Decals/1_4/1_4.INT.udn
+++ b/Engine/Documentation/Source/Resources/ContentExamples/Decals/1_4/1_4.INT.udn
@@ -11,6 +11,24 @@ Related: Engine/Rendering/Materials/MaterialInputs
 Related: Engine/Rendering/Materials/LayeredMaterials
 Related: Engine/Rendering/Materials/MaterialInstances
 
+[VAR:Steps]
+[OBJECT:Navigation]
+	[PARAM:previous]
+		[Previous Example](Resources/ContentExamples/Decals/1_3 "%Resources/ContentExamples/Decals/1_3:title%")
+	[/PARAM]
+	[PARAM:current]
+	[/PARAM]
+	[PARAM:home]
+		[Decals Home](Resources/ContentExamples/Decals)
+	[/PARAM]
+	[PARAM:next]
+		[Next Example](Resources/ContentExamples/Decals/1_5 "%Resources/ContentExamples/Decals/1_5:title%")
+	[/PARAM]
+[/OBJECT]
+[/VAR]
+
+%Steps%
+
 [VAR:Topic]
 [OBJECT:Topic]
 	[PARAM:image]
diff --git a/Engine/Documentation/Source/Resources/ContentExamples/Decals/1_5/1_5.INT.udn b/Engine/Documentation/Source/Resources/ContentExamples/Decals/1_5/1_5.INT.udn
--- a/Engine/Documentation/Source/Resources/ContentExamples/Decals/1_5/1_5.INT.udn
+++ b/Engine/Documentation/Source/Resources/ContentExamples/Decals/1_5/1_5.INT.udn
@@ -11,6 +11,24 @@ Related: Engine/Rendering/Materials/MaterialInputs
 Related: Engine/Rendering/Materials/LayeredMaterials
 Related: Engine/Rendering/Materials/MaterialInstances
 
+[VAR:Steps]
+[OBJECT:Navigation]
+	[PARAM:previous]
+		[Previous Example](Resources/ContentExamples/Decals/1_4 "%Resources/ContentExamples/Decals/1_4:title%")
+	[/PARAM]
+	[PARAM:current]
+	[/PARAM]
+	[PARAM:home]
+		[Decals Home](Resources/ContentExamples/Decals)
+	[/PARAM]
+	[PARAM:next]
+		[Next Example](Resources/ContentExamples/Decals/1_6 "%Resources/ContentExamples/Decals/1_6:title%")
+	[/PARAM]
+[/OBJECT]
+[/VAR]
+
+%Steps%
+
 [VAR:Topic]
 [OBJECT:Topic]
 	[PARAM:image]
diff --git a/Engine/Documentation/Source/Resources/ContentExamples/Decals/1_6/1_6.INT.udn b/Engine/Documentation/Source/Resources/ContentExamples/Decals/1_6/1_6.INT.udn
--- a/Engine/Documentation/Source/Resources/ContentExamples/Decals/1_6/1_6.INT.udn
+++ b/Engine/Documentation/Source/Resources/ContentExamples/Decals/1_6/1_6.INT.udn
@@ -11,6 +11,24 @@ Related: Engine/Rendering/Materials/MaterialInputs
 Related: Engine/Rendering/Materials/LayeredMaterials
 Related: Engine/Rendering/Materials/MaterialInstances
 
+[VAR:Steps]
+[OBJECT:Navigation]
+	[PARAM:previous]
+		[Previous Example](Resources/ContentExamples/Decals/1_5 "%Resources/ContentExamples/Decals/1_5:title%")
+	[/PARAM]
+	[PARAM:current]
+	[/PARAM]
+	[PARAM:home]
+		[Decals Home](Resources/ContentExamples/Decals)
+	[/PARAM]
+	[PARAM:next]
+		[Next Example](Resources/ContentExamples/Decals/1_7 "%Resources/ContentExamples/Decals/1_7:title%")
+	[/PARAM]
+[/OBJECT]
+[/VAR]
+
+%Steps%
+
 [TOC(start:2 end:2)]
 
 [VAR:Topic]
diff --git a/Engine/Documentation/Source/Resources/ContentExamples/Decals/1_7/1_7.INT.udn b/Engine/Documentation/Source/Resources/ContentExamples/Decals/1_7/1_7.INT.udn
--- a/Engine/Documentation/Source/Resources/ContentExamples/Decals/1_7/1_7.INT.udn
+++ b/Engine/Documentation/Source/Resources/ContentExamples/Decals/1_7/1_7.INT.udn
@@ -11,6 +11,23 @@ Related: Engine/Rendering/Materials/MaterialInputs
 Related: Engine/Rendering/Materials/LayeredMaterials
 Related: Engine/Rendering/Materials/MaterialInstances
 
+[VAR:Steps]
+[OBJECT:Navigation]
+	[PARAM:previous]
+		[Previous Example](Resources/ContentExamples/Decals/1_6 "%Resources/ContentExamples/Decals/1_6:title%")
+	[/PARAM]
+	[PARAM:current]
+	[/PARAM]
+	[PARAM:home]
+		[Decals Home](Resources/ContentExamples/Decals)
+	[/PARAM]
+	[PARAM:next]
+	[/PARAM]
+[/OBJECT]
+[/VAR]
+
+%Steps%
+
 [TOC(start:2 end:2)]
 
 [VAR:Topic]
