#UE4doc #docs:
Changed font settings for multiple terms
Added quotations for terms that were suppose to be typed by the user
Added engine version field

diff --git a/Engine/Documentation/Source/Engine/AI/EnvironmentQuerySystem/QuickStart/4/4.INT.udn b/Engine/Documentation/Source/Engine/AI/EnvironmentQuerySystem/QuickStart/4/4.INT.udn
--- a/Engine/Documentation/Source/Engine/AI/EnvironmentQuerySystem/QuickStart/4/4.INT.udn
+++ b/Engine/Documentation/Source/Engine/AI/EnvironmentQuerySystem/QuickStart/4/4.INT.udn
@@ -2,6 +2,7 @@ Availability:Public
 Title: 4. Create the Basic Assets
 Crumbs:%ROOT%, Engine, Gameplay/AI, Engine/AI/EnvironmentQuerySystem, Engine/AI/EnvironmentQuerySystem/QuickStart
 Description: Create the 7 needed assets for the EQS Quick Start.
+version: 4.9
 
 [VAR:Steps]
 [OBJECT:Navigation]
@@ -21,7 +22,7 @@ Description: Create the 7 needed assets for the EQS Quick Start.
 
 %Steps%
 
-To start, you will need create 7 assets in the **Content Browser**: An **AI Controller** Blueprint, A **Character** Blueprint, and **Behavior Tree** Asset, a **Blackboard** Asset, an **EnvQueryContext_BlueprintBase** (EQS Context) Blueprint, and 2 **Environment Query** templates. These are all the assets you will need to complete this Quick Start.
+To start, you will need create 7 assets in the **Content Browser**: An AI Controller Blueprint, A Character Blueprint, and Behavior Tree Asset, a Blackboard Asset, an EnvQueryContext_BlueprintBase (EQS Context) Blueprint, and 2 Environment Query templates. These are all the assets you will need to complete this Quick Start.
 
 [REGION:raw]
 ![](Assets.png)(convert:false)
@@ -36,45 +37,45 @@ To start, you will need create 7 assets in the **Content Browser**: An **AI Cont
 
 	![](Engine/AI/BehaviorTrees/QuickStart/3/AIC.png)
 
-	Name this **Blueprint** `Runner_AI_CON`.
+	Name this **Blueprint** "Runner_AI_CON".
 
 1. [REGION:raw]**New** %arrow% **Blueprint** %arrow% **Character**.[/REGION]
 
 	![](Engine/AI/BehaviorTrees/QuickStart/3/charBP.png)
 
-	Name this **Blueprint** `AI_Character`.
+	Name this **Blueprint** "AI_Character".
 
-1. [REGION:raw]**New** %arrow% **Miscellaneous** %arrow% **Behavior Tree**.[/REGION]
+1. [REGION:raw]**New** %arrow% **Artificial Intelligence** %arrow% **Behavior Tree**.[/REGION]
 
 	[REGION:raw]
 	![](Engine/AI/BehaviorTrees/QuickStart/3/newBT.png)
 	[/REGION]
 
-	Name the **Behavior Tree** `RunnerBT`.
+	Name the **Behavior Tree** "RunnerBT".
 	
-1. [REGION:raw]**New** %arrow% **Miscellaneous** %arrow%  **Blackboard**.[/REGION]
+1. [REGION:raw]**New** %arrow% **Artificial Intelligence** %arrow%  **Blackboard**.[/REGION]
 	
 	[REGION:raw]
 	![](Engine/AI/BehaviorTrees/QuickStart/3/newBB.png)
 	[/REGION]	
 
-	Name the **Blackboard** `RunnerBlackboard`.
+	Name the **Blackboard** "RunnerBlackboard'.
 
-1. Create an **Environment Query**: [REGION:raw]**New** %arrow% **Miscellaneous** %arrow%  **Environment Query**.[/REGION]
+1. Create an **Environment Query**: [REGION:raw]**New** %arrow% **Artificial Intelligence** %arrow%  **Environment Query**.[/REGION]
 	
 	[REGION:raw]
 	![](newEQ.PNG)
 	[/REGION]	
 
-	Name one **Environment Query** `FindHidingSpot`.
+	Name one **Environment Query** "FindHidingSpot".
 
-1. Name the other **Environment Query** `FindPlayer`.
+1. Name the other **Environment Query** "FindPlayer".
 
 1. [REGION:raw]**New** %arrow% **Blueprint** %arrow%  **EnvQueryContext_BlueprintBase**.[/REGION]
 	
 	![](newContext.png)
 
-	Name the **Environment Query** `PlayerContext`.
+	Name the **Environment Query** "PlayerContext'.
 
 %Steps%
 
