UEDOC-1400

Added Previous and Next tabs to the top of the pages

diff --git a/Engine/Documentation/Source/Resources/ContentExamples/EffectsGallery/1_A/1_A.INT.udn b/Engine/Documentation/Source/Resources/ContentExamples/EffectsGallery/1_A/1_A.INT.udn
--- a/Engine/Documentation/Source/Resources/ContentExamples/EffectsGallery/1_A/1_A.INT.udn
+++ b/Engine/Documentation/Source/Resources/ContentExamples/EffectsGallery/1_A/1_A.INT.udn
@@ -9,6 +9,23 @@ Related: Engine/Rendering/ParticleSystems/Reference/TypeData/GPUSprites
 Related: Engine/Rendering/ParticleSystems/ParticleLights
 Related: Engine/Rendering/ParticleSystems/Reference/Modules/VectorField
 
+[VAR:Steps]
+[OBJECT:Navigation]
+	[PARAM:previous]
+	[/PARAM]
+	[PARAM:current]
+	[/PARAM]
+	[PARAM:home]
+		[Effects Gallery Home](Resources/ContentExamples/EffectsGallery)
+	[/PARAM]
+	[PARAM:next]
+		[Next Example](Resources/ContentExamples/EffectsGallery/1_B "%Resources/ContentExamples/EffectsGallery/1_B:title%")
+	[/PARAM]
+[/OBJECT]
+[/VAR]
+
+%Steps%
+
 [VAR:TopicCompact]
 [OBJECT:TopicCompact]
 	[PARAM:image]
diff --git a/Engine/Documentation/Source/Resources/ContentExamples/EffectsGallery/1_B/1_B.INT.udn b/Engine/Documentation/Source/Resources/ContentExamples/EffectsGallery/1_B/1_B.INT.udn
--- a/Engine/Documentation/Source/Resources/ContentExamples/EffectsGallery/1_B/1_B.INT.udn
+++ b/Engine/Documentation/Source/Resources/ContentExamples/EffectsGallery/1_B/1_B.INT.udn
@@ -9,6 +9,24 @@ Related: Engine/Rendering/ParticleSystems/Reference/TypeData/GPUSprites
 Related: Engine/Rendering/ParticleSystems/ParticleLights
 Related: Engine/Rendering/ParticleSystems/Reference/Modules/VectorField
 
+[VAR:Steps]
+[OBJECT:Navigation]
+	[PARAM:previous]
+		[Previous Example](Resources/ContentExamples/EffectsGallery/1_A "%Resources/ContentExamples/EffectsGallery/1_A:title%")
+	[/PARAM]
+	[PARAM:current]
+	[/PARAM]
+	[PARAM:home]
+		[Effects Gallery Home](Resources/ContentExamples/EffectsGallery)
+	[/PARAM]
+	[PARAM:next]
+		[Next Example](Resources/ContentExamples/EffectsGallery/1_C "%Resources/ContentExamples/EffectsGallery/1_C:title%")
+	[/PARAM]
+[/OBJECT]
+[/VAR]
+
+%Steps%
+
 [VAR:TopicCompact]
 [OBJECT:TopicCompact]
 	[PARAM:image]
diff --git a/Engine/Documentation/Source/Resources/ContentExamples/EffectsGallery/1_C/1_C.INT.udn b/Engine/Documentation/Source/Resources/ContentExamples/EffectsGallery/1_C/1_C.INT.udn
--- a/Engine/Documentation/Source/Resources/ContentExamples/EffectsGallery/1_C/1_C.INT.udn
+++ b/Engine/Documentation/Source/Resources/ContentExamples/EffectsGallery/1_C/1_C.INT.udn
@@ -9,6 +9,24 @@ Related: Engine/Rendering/ParticleSystems/Reference/TypeData/GPUSprites
 Related: Engine/Rendering/ParticleSystems/ParticleLights
 Related: Engine/Rendering/ParticleSystems/Reference/Modules/VectorField
 
+[VAR:Steps]
+[OBJECT:Navigation]
+	[PARAM:previous]
+		[Previous Example](Resources/ContentExamples/EffectsGallery/1_B "%Resources/ContentExamples/EffectsGallery/1_B:title%")
+	[/PARAM]
+	[PARAM:current]
+	[/PARAM]
+	[PARAM:home]
+		[Effects Gallery Home](Resources/ContentExamples/EffectsGallery)
+	[/PARAM]
+	[PARAM:next]
+		[Next Example](Resources/ContentExamples/EffectsGallery/1_D "%Resources/ContentExamples/EffectsGallery/1_D:title%")
+	[/PARAM]
+[/OBJECT]
+[/VAR]
+
+%Steps%
+
 [VAR:TopicCompact]
 [OBJECT:TopicCompact]
 	[PARAM:image]
diff --git a/Engine/Documentation/Source/Resources/ContentExamples/EffectsGallery/1_D/1_D.INT.udn b/Engine/Documentation/Source/Resources/ContentExamples/EffectsGallery/1_D/1_D.INT.udn
--- a/Engine/Documentation/Source/Resources/ContentExamples/EffectsGallery/1_D/1_D.INT.udn
+++ b/Engine/Documentation/Source/Resources/ContentExamples/EffectsGallery/1_D/1_D.INT.udn
@@ -9,6 +9,23 @@ Related: Engine/Rendering/ParticleSystems/Reference/TypeData/GPUSprites
 Related: Engine/Rendering/ParticleSystems/ParticleLights
 Related: Engine/Rendering/ParticleSystems/Reference/Modules/VectorField
 
+[VAR:Steps]
+[OBJECT:Navigation]
+	[PARAM:previous]
+		[Previous Example](Resources/ContentExamples/EffectsGallery/1_C "%Resources/ContentExamples/EffectsGallery/1_C:title%")
+	[/PARAM]
+	[PARAM:current]
+	[/PARAM]
+	[PARAM:home]
+		[Effects Gallery Home](Resources/ContentExamples/EffectsGallery)
+	[/PARAM]
+	[PARAM:next]
+		[Next Example](Resources/ContentExamples/EffectsGallery/1_E "%Resources/ContentExamples/EffectsGallery/1_E:title%")
+	[/PARAM]
+[/OBJECT]
+[/VAR]
+
+%Steps%
 
 [VAR:TopicCompact]
 [OBJECT:TopicCompact]
diff --git a/Engine/Documentation/Source/Resources/ContentExamples/EffectsGallery/1_E/1_E.INT.udn b/Engine/Documentation/Source/Resources/ContentExamples/EffectsGallery/1_E/1_E.INT.udn
--- a/Engine/Documentation/Source/Resources/ContentExamples/EffectsGallery/1_E/1_E.INT.udn
+++ b/Engine/Documentation/Source/Resources/ContentExamples/EffectsGallery/1_E/1_E.INT.udn
@@ -9,6 +9,24 @@ Related: Engine/Rendering/ParticleSystems/Reference/TypeData/GPUSprites
 Related: Engine/Rendering/ParticleSystems/ParticleLights
 Related: Engine/Rendering/ParticleSystems/Reference/Modules/VectorField
 
+[VAR:Steps]
+[OBJECT:Navigation]
+	[PARAM:previous]
+		[Previous Example](Resources/ContentExamples/EffectsGallery/1_D "%Resources/ContentExamples/EffectsGallery/1_D:title%")
+	[/PARAM]
+	[PARAM:current]
+	[/PARAM]
+	[PARAM:home]
+		[Effects Gallery Home](Resources/ContentExamples/EffectsGallery)
+	[/PARAM]
+	[PARAM:next]
+		[Next Example](Resources/ContentExamples/EffectsGallery/1_F "%Resources/ContentExamples/EffectsGallery/1_F:title%")
+	[/PARAM]
+[/OBJECT]
+[/VAR]
+
+%Steps%
+
 [VAR:TopicCompact]
 [OBJECT:TopicCompact]
 	[PARAM:image]
diff --git a/Engine/Documentation/Source/Resources/ContentExamples/EffectsGallery/1_F/1_F.INT.udn b/Engine/Documentation/Source/Resources/ContentExamples/EffectsGallery/1_F/1_F.INT.udn
--- a/Engine/Documentation/Source/Resources/ContentExamples/EffectsGallery/1_F/1_F.INT.udn
+++ b/Engine/Documentation/Source/Resources/ContentExamples/EffectsGallery/1_F/1_F.INT.udn
@@ -9,6 +9,24 @@ Related: Engine/Rendering/ParticleSystems/Reference/TypeData/GPUSprites
 Related: Engine/Rendering/ParticleSystems/ParticleLights
 Related: Engine/Rendering/ParticleSystems/Reference/Modules/VectorField
 
+[VAR:Steps]
+[OBJECT:Navigation]
+	[PARAM:previous]
+		[Previous Example](Resources/ContentExamples/EffectsGallery/1_E "%Resources/ContentExamples/EffectsGallery/1_E:title%")
+	[/PARAM]
+	[PARAM:current]
+	[/PARAM]
+	[PARAM:home]
+		[Effects Gallery Home](Resources/ContentExamples/EffectsGallery)
+	[/PARAM]
+	[PARAM:next]
+		[Next Example](Resources/ContentExamples/EffectsGallery/1_G "%Resources/ContentExamples/EffectsGallery/1_G:title%")
+	[/PARAM]
+[/OBJECT]
+[/VAR]
+
+%Steps%
+
 [VAR:TopicCompact]
 [OBJECT:TopicCompact]
 	[PARAM:image]
diff --git a/Engine/Documentation/Source/Resources/ContentExamples/EffectsGallery/1_G/1_G.INT.udn b/Engine/Documentation/Source/Resources/ContentExamples/EffectsGallery/1_G/1_G.INT.udn
--- a/Engine/Documentation/Source/Resources/ContentExamples/EffectsGallery/1_G/1_G.INT.udn
+++ b/Engine/Documentation/Source/Resources/ContentExamples/EffectsGallery/1_G/1_G.INT.udn
@@ -8,6 +8,25 @@ Related: Engine/Rendering/ParticleSystems/Optimization
 Related: Engine/Rendering/ParticleSystems/Reference/TypeData/GPUSprites
 Related: Engine/Rendering/ParticleSystems/ParticleLights
 Related: Engine/Rendering/ParticleSystems/Reference/Modules/VectorField
+
+[VAR:Steps]
+[OBJECT:Navigation]
+	[PARAM:previous]
+		[Previous Example](Resources/ContentExamples/EffectsGallery/1_F "%Resources/ContentExamples/EffectsGallery/1_F:title%")
+	[/PARAM]
+	[PARAM:current]
+	[/PARAM]
+	[PARAM:home]
+		[Effects Gallery Home](Resources/ContentExamples/EffectsGallery)
+	[/PARAM]
+	[PARAM:next]
+		[Next Example](Resources/ContentExamples/EffectsGallery/1_H "%Resources/ContentExamples/EffectsGallery/1_H:title%")
+	[/PARAM]
+[/OBJECT]
+[/VAR]
+
+%Steps%
+
 [VAR:TopicCompact]
 [OBJECT:TopicCompact]
 	[PARAM:image]
diff --git a/Engine/Documentation/Source/Resources/ContentExamples/EffectsGallery/1_H/1_H.INT.udn b/Engine/Documentation/Source/Resources/ContentExamples/EffectsGallery/1_H/1_H.INT.udn
--- a/Engine/Documentation/Source/Resources/ContentExamples/EffectsGallery/1_H/1_H.INT.udn
+++ b/Engine/Documentation/Source/Resources/ContentExamples/EffectsGallery/1_H/1_H.INT.udn
@@ -10,6 +10,23 @@ Related: Engine/Rendering/ParticleSystems/Reference/TypeData/GPUSprites
 Related: Engine/Rendering/ParticleSystems/ParticleLights
 Related: Engine/Rendering/ParticleSystems/Reference/Modules/VectorField
 
+[VAR:Steps]
+[OBJECT:Navigation]
+	[PARAM:previous]
+		[Previous Example](Resources/ContentExamples/EffectsGallery/1_G "%Resources/ContentExamples/EffectsGallery/1_G:title%")
+	[/PARAM]
+	[PARAM:current]
+	[/PARAM]
+	[PARAM:home]
+		[Effects Gallery Home](Resources/ContentExamples/EffectsGallery)
+	[/PARAM]
+	[PARAM:next]
+	[/PARAM]
+[/OBJECT]
+[/VAR]
+
+%Steps%
+
 [VAR:TopicCompact]
 [OBJECT:TopicCompact]
 	[PARAM:image]
diff --git a/Engine/Documentation/Source/Resources/ContentExamples/EffectsGallery/2_A/2_A.INT.udn b/Engine/Documentation/Source/Resources/ContentExamples/EffectsGallery/2_A/2_A.INT.udn
--- a/Engine/Documentation/Source/Resources/ContentExamples/EffectsGallery/2_A/2_A.INT.udn
+++ b/Engine/Documentation/Source/Resources/ContentExamples/EffectsGallery/2_A/2_A.INT.udn
@@ -9,6 +9,24 @@ Related: Engine/Rendering/ParticleSystems/Optimization
 Related: Engine/Rendering/ParticleSystems/Reference/TypeData/GPUSprites
 Related: Engine/Rendering/ParticleSystems/ParticleLights
 Related: Engine/Rendering/ParticleSystems/Reference/Modules/VectorField
+
+[VAR:Steps]
+[OBJECT:Navigation]
+	[PARAM:previous]
+	[/PARAM]
+	[PARAM:current]
+	[/PARAM]
+	[PARAM:home]
+		[Effects Gallery Home](Resources/ContentExamples/EffectsGallery)
+	[/PARAM]
+	[PARAM:next]
+		[Next Example](Resources/ContentExamples/EffectsGallery/2_B "%Resources/ContentExamples/EffectsGallery/2_B:title%")
+	[/PARAM]
+[/OBJECT]
+[/VAR]
+
+%Steps%
+
 [VAR:TopicCompact]
 [OBJECT:TopicCompact]
 	[PARAM:image]
diff --git a/Engine/Documentation/Source/Resources/ContentExamples/EffectsGallery/2_B/2_B.INT.udn b/Engine/Documentation/Source/Resources/ContentExamples/EffectsGallery/2_B/2_B.INT.udn
--- a/Engine/Documentation/Source/Resources/ContentExamples/EffectsGallery/2_B/2_B.INT.udn
+++ b/Engine/Documentation/Source/Resources/ContentExamples/EffectsGallery/2_B/2_B.INT.udn
@@ -10,6 +10,24 @@ Related: Engine/Rendering/ParticleSystems/Reference/TypeData/GPUSprites
 Related: Engine/Rendering/ParticleSystems/ParticleLights
 Related: Engine/Rendering/ParticleSystems/Reference/Modules/VectorField
 
+[VAR:Steps]
+[OBJECT:Navigation]
+	[PARAM:previous]
+		[Previous Example](Resources/ContentExamples/EffectsGallery/2_A "%Resources/ContentExamples/EffectsGallery/2_A:title%")
+	[/PARAM]
+	[PARAM:current]
+	[/PARAM]
+	[PARAM:home]
+		[Effects Gallery Home](Resources/ContentExamples/EffectsGallery)
+	[/PARAM]
+	[PARAM:next]
+		[Next Example](Resources/ContentExamples/EffectsGallery/2_C "%Resources/ContentExamples/EffectsGallery/2_C:title%")
+	[/PARAM]
+[/OBJECT]
+[/VAR]
+
+%Steps%
+
 [VAR:TopicCompact]
 [OBJECT:TopicCompact]
 	[PARAM:image]
diff --git a/Engine/Documentation/Source/Resources/ContentExamples/EffectsGallery/2_C/2_C.INT.udn b/Engine/Documentation/Source/Resources/ContentExamples/EffectsGallery/2_C/2_C.INT.udn
--- a/Engine/Documentation/Source/Resources/ContentExamples/EffectsGallery/2_C/2_C.INT.udn
+++ b/Engine/Documentation/Source/Resources/ContentExamples/EffectsGallery/2_C/2_C.INT.udn
@@ -1,5 +1,5 @@
 Availability: Public
-Title:2.3 - Black Body Node (Dynamic Parameter Control)
+Title:2.3 - Black Body Node [Dynamic Parameter Control]
 Crumbs:%ROOT%, Resources, Resources/ContentExamples, Resources/ContentExamples/EffectsGallery
 Description:An overview of the Effects Gallery example level.
 Related: Engine/Rendering/Materials/MaterialInstances
@@ -10,6 +10,24 @@ Related: Engine/Rendering/ParticleSystems/Reference/TypeData/GPUSprites
 Related: Engine/Rendering/ParticleSystems/ParticleLights
 Related: Engine/Rendering/ParticleSystems/Reference/Modules/VectorField
 
+[VAR:Steps]
+[OBJECT:Navigation]
+	[PARAM:previous]
+		[Previous Example](Resources/ContentExamples/EffectsGallery/2_B "%Resources/ContentExamples/EffectsGallery/2_B:title%")
+	[/PARAM]
+	[PARAM:current]
+	[/PARAM]
+	[PARAM:home]
+		[Effects Gallery Home](Resources/ContentExamples/EffectsGallery)
+	[/PARAM]
+	[PARAM:next]
+		[Next Example](Resources/ContentExamples/EffectsGallery/2_D "%Resources/ContentExamples/EffectsGallery/2_D:title%")
+	[/PARAM]
+[/OBJECT]
+[/VAR]
+
+%Steps%
+
 [VAR:TopicCompact]
 [OBJECT:TopicCompact]
 	[PARAM:image]
diff --git a/Engine/Documentation/Source/Resources/ContentExamples/EffectsGallery/2_D/2_D.INT.udn b/Engine/Documentation/Source/Resources/ContentExamples/EffectsGallery/2_D/2_D.INT.udn
--- a/Engine/Documentation/Source/Resources/ContentExamples/EffectsGallery/2_D/2_D.INT.udn
+++ b/Engine/Documentation/Source/Resources/ContentExamples/EffectsGallery/2_D/2_D.INT.udn
@@ -9,6 +9,24 @@ Related: Engine/Rendering/ParticleSystems/Reference/TypeData/GPUSprites
 Related: Engine/Rendering/ParticleSystems/ParticleLights
 Related: Engine/Rendering/ParticleSystems/Reference/Modules/VectorField
 
+[VAR:Steps]
+[OBJECT:Navigation]
+	[PARAM:previous]
+		[Previous Example](Resources/ContentExamples/EffectsGallery/2_C "%Resources/ContentExamples/EffectsGallery/2_C:title%")
+	[/PARAM]
+	[PARAM:current]
+	[/PARAM]
+	[PARAM:home]
+		[Effects Gallery Home](Resources/ContentExamples/EffectsGallery)
+	[/PARAM]
+	[PARAM:next]
+		[Next Example](Resources/ContentExamples/EffectsGallery/2_E "%Resources/ContentExamples/EffectsGallery/2_E:title%")
+	[/PARAM]
+[/OBJECT]
+[/VAR]
+
+%Steps%
+
 [VAR:TopicCompact]
 [OBJECT:TopicCompact]
 	[PARAM:image]
diff --git a/Engine/Documentation/Source/Resources/ContentExamples/EffectsGallery/2_E/2_E.INT.udn b/Engine/Documentation/Source/Resources/ContentExamples/EffectsGallery/2_E/2_E.INT.udn
--- a/Engine/Documentation/Source/Resources/ContentExamples/EffectsGallery/2_E/2_E.INT.udn
+++ b/Engine/Documentation/Source/Resources/ContentExamples/EffectsGallery/2_E/2_E.INT.udn
@@ -10,6 +10,24 @@ Related: Engine/Rendering/ParticleSystems/Reference/TypeData/GPUSprites
 Related: Engine/Rendering/ParticleSystems/ParticleLights
 Related: Engine/Rendering/ParticleSystems/Reference/Modules/VectorField
 
+[VAR:Steps]
+[OBJECT:Navigation]
+	[PARAM:previous]
+		[Previous Example](Resources/ContentExamples/EffectsGallery/2_D "%Resources/ContentExamples/EffectsGallery/2_D:title%")
+	[/PARAM]
+	[PARAM:current]
+	[/PARAM]
+	[PARAM:home]
+		[Effects Gallery Home](Resources/ContentExamples/EffectsGallery)
+	[/PARAM]
+	[PARAM:next]
+		[Next Example](Resources/ContentExamples/EffectsGallery/2_F "%Resources/ContentExamples/EffectsGallery/2_F:title%")
+	[/PARAM]
+[/OBJECT]
+[/VAR]
+
+%Steps%
+
 [VAR:TopicCompact]
 [OBJECT:TopicCompact]
 	[PARAM:image]
diff --git a/Engine/Documentation/Source/Resources/ContentExamples/EffectsGallery/2_F/2_F.INT.udn b/Engine/Documentation/Source/Resources/ContentExamples/EffectsGallery/2_F/2_F.INT.udn
--- a/Engine/Documentation/Source/Resources/ContentExamples/EffectsGallery/2_F/2_F.INT.udn
+++ b/Engine/Documentation/Source/Resources/ContentExamples/EffectsGallery/2_F/2_F.INT.udn
@@ -9,6 +9,25 @@ Related: Engine/Rendering/ParticleSystems/Optimization
 Related: Engine/Rendering/ParticleSystems/Reference/TypeData/GPUSprites
 Related: Engine/Rendering/ParticleSystems/ParticleLights
 Related: Engine/Rendering/ParticleSystems/Reference/Modules/VectorField
+
+[VAR:Steps]
+[OBJECT:Navigation]
+	[PARAM:previous]
+		[Previous Example](Resources/ContentExamples/EffectsGallery/2_E "%Resources/ContentExamples/EffectsGallery/2_E:title%")
+	[/PARAM]
+	[PARAM:current]
+	[/PARAM]
+	[PARAM:home]
+		[Effects Gallery Home](Resources/ContentExamples/EffectsGallery)
+	[/PARAM]
+	[PARAM:next]
+		[Next Example](Resources/ContentExamples/EffectsGallery/2_G "%Resources/ContentExamples/EffectsGallery/2_G:title%")
+	[/PARAM]
+[/OBJECT]
+[/VAR]
+
+%Steps%
+
 [VAR:TopicCompact]
 [OBJECT:TopicCompact]
 	[PARAM:image]
diff --git a/Engine/Documentation/Source/Resources/ContentExamples/EffectsGallery/2_G/2_G.INT.udn b/Engine/Documentation/Source/Resources/ContentExamples/EffectsGallery/2_G/2_G.INT.udn
--- a/Engine/Documentation/Source/Resources/ContentExamples/EffectsGallery/2_G/2_G.INT.udn
+++ b/Engine/Documentation/Source/Resources/ContentExamples/EffectsGallery/2_G/2_G.INT.udn
@@ -10,6 +10,24 @@ Related: Engine/Rendering/ParticleSystems/Reference/TypeData/GPUSprites
 Related: Engine/Rendering/ParticleSystems/ParticleLights
 Related: Engine/Rendering/ParticleSystems/Reference/Modules/VectorField
 
+[VAR:Steps]
+[OBJECT:Navigation]
+	[PARAM:previous]
+		[Previous Example](Resources/ContentExamples/EffectsGallery/2_F "%Resources/ContentExamples/EffectsGallery/2_F:title%")
+	[/PARAM]
+	[PARAM:current]
+	[/PARAM]
+	[PARAM:home]
+		[Effects Gallery Home](Resources/ContentExamples/EffectsGallery)
+	[/PARAM]
+	[PARAM:next]
+		[Next Example](Resources/ContentExamples/EffectsGallery/2_H "%Resources/ContentExamples/EffectsGallery/2_H:title%")
+	[/PARAM]
+[/OBJECT]
+[/VAR]
+
+%Steps%
+
 [VAR:TopicCompact]
 [OBJECT:TopicCompact]
 	[PARAM:image]
diff --git a/Engine/Documentation/Source/Resources/ContentExamples/EffectsGallery/2_H/2_H.INT.udn b/Engine/Documentation/Source/Resources/ContentExamples/EffectsGallery/2_H/2_H.INT.udn
--- a/Engine/Documentation/Source/Resources/ContentExamples/EffectsGallery/2_H/2_H.INT.udn
+++ b/Engine/Documentation/Source/Resources/ContentExamples/EffectsGallery/2_H/2_H.INT.udn
@@ -10,6 +10,24 @@ Related: Engine/Rendering/ParticleSystems/Reference/TypeData/GPUSprites
 Related: Engine/Rendering/ParticleSystems/ParticleLights
 Related: Engine/Rendering/ParticleSystems/Reference/Modules/VectorField
 
+
+[VAR:Steps]
+[OBJECT:Navigation]
+	[PARAM:previous]
+		[Previous Example](Resources/ContentExamples/EffectsGallery/2_G "%Resources/ContentExamples/EffectsGallery/2_G:title%")
+	[/PARAM]
+	[PARAM:current]
+	[/PARAM]
+	[PARAM:home]
+		[Effects Gallery Home](Resources/ContentExamples/EffectsGallery)
+	[/PARAM]
+	[PARAM:next]
+	[/PARAM]
+[/OBJECT]
+[/VAR]
+
+%Steps%
+
 [VAR:TopicCompact]
 [OBJECT:TopicCompact]
 	[PARAM:image]
