#UE4doc #docs:
Changed font settings for multiple terms
Added engine version field
Changed the names for multiple terms as they did not match the name of the properties in the pictures.
Removed unnecessary quotations

diff --git a/Engine/Documentation/Source/Engine/AI/EnvironmentQuerySystem/QuickStart/5/5.INT.udn b/Engine/Documentation/Source/Engine/AI/EnvironmentQuerySystem/QuickStart/5/5.INT.udn
--- a/Engine/Documentation/Source/Engine/AI/EnvironmentQuerySystem/QuickStart/5/5.INT.udn
+++ b/Engine/Documentation/Source/Engine/AI/EnvironmentQuerySystem/QuickStart/5/5.INT.udn
@@ -2,6 +2,7 @@ Availability:Public
 Title: 5. Character Blueprint
 Crumbs:%ROOT%, Engine, Gameplay/AI, Engine/AI/EnvironmentQuerySystem, Engine/AI/EnvironmentQuerySystem/QuickStart
 Description: Setup the AI's Character Blueprint.
+version: 4.9
 
 [VAR:Steps]
 [OBJECT:Navigation]
@@ -21,19 +22,19 @@ Description: Setup the AI's Character Blueprint.
 
 %Steps%
 
-As we have created a project with a **Character** already in it, we have access to a humanoid **Skeletal Mesh** and its **Animation Blueprint**. We will use these in the tutorial, but feel free to use your own if you have them available.
+As we have created a project with a Character already in it, we have access to a humanoid Skeletal Mesh and its Animation Blueprint. We will use these in the tutorial, but feel free to use your own if you have them available.
 
 This Character will become our new AI Character. The reason we use a Character, over say a Pawn, is to gain access to the Character Movement Component's functionality.
 
 ## Steps
 
-1. **Double-click** the `AI_Character` Blueprint in the **Content Browser** to open it for editing.
+1. **Double-click** the AI_Character Blueprint in the **Content Browser** to open it for editing.
 
 1. Select the **Mesh** component in the **Components** panel.
 
 	![](Engine\AI\BehaviorTrees\QuickStart\4\meshComp.png)
 
-1. In the **Details** panel, find the **Mesh** category and assign the **SK_Mannequin** (or your own) **Skeletal Mesh**.
+1. In the **Details** panel, find the **Mesh** category and assign the **SK_Mannequin** (or your own) Skeletal Mesh.
 
 	![](Engine\AI\BehaviorTrees\QuickStart\4\assignSkelMesh.png)
 
@@ -41,11 +42,11 @@ This Character will become our new AI Character. The reason we use a Character,
 
 1. Rotate the **Mesh** Component so it faces the direction the blue **Arrow Component** is pointing.
 
-1. In the **Details** panel, find the **Animation** category and assign the `TopDown_AnimBlueprint` to the **Animation Blueprint Generated Class** property.
+1. In the **Details** panel, find the **Animation** category and assign the TopDown_AnimBlueprint to the **Animation Blueprint Generated Class** property.
 
 	![](Engine\AI\BehaviorTrees\QuickStart\4\assignAnimBP.png)
 
-1. Select the **CapsuleComponent** from the **Components** panel, then in the **Details** panel, adjust the **Capsule Half Height** and **Capsule Radius** so they encompass the **Skeletal Mesh**.
+1. Select the **CapsuleComponent** from the **Components** panel, then in the **Details** panel, adjust the **Capsule Half Height** and **Capsule Radius** so they encompass the Skeletal Mesh.
 
 	![](Engine\AI\BehaviorTrees\QuickStart\4\shapeCat.png)
 
@@ -53,21 +54,21 @@ This Character will become our new AI Character. The reason we use a Character,
 
 	![](Engine\AI\BehaviorTrees\QuickStart\4\charmve.png)
 
-1. In the **Details** panel, in the **Movement Capabilities** category, adjust the Agent Radius and Agent Height values.
+1. In the **Details** panel, in the **Movement Capabilities** category, adjust the **Nav Agent Radius** and **Nav Agent Height** values.
 
 	![](Engine\AI\BehaviorTrees\QuickStart\4\agentProps.png)
 
-	_Set the Agent Radius to at least the Capsule's radius, and the Agent Height to at least twice the Capsule's "half height."_
+	Set the **Nav Agent Radius** to at least the **Capsule Radius**, and the **Nav Agent Height** to at least twice the **Capsule Half Height**.
 
-1. With the Character Movement Component still selected, change the **Max Walk Speed** property, in the **Character Movement: Walking** category to `400`.
+1. With the Character Movement Component still selected, change the **Max Walk Speed** property, in the **Character Movement: Walking** category to 400.
 	
 	![](Engine\AI\BehaviorTrees\QuickStart\4\maxWalk.png)
 
-1. Click the **Class Defaults** ![](Engine\Blueprints\UserGuide\Variables\defaultsButton.png) button and assign the new AI Controller (`Runner_AI_CON`) to the **AIController Class** property in the **Pawn** category.
+1. Click the **Class Defaults** ![](Engine\Blueprints\UserGuide\Variables\defaultsButton.png) button and assign the new AI Controller (Runner_AI_CON) to the **AIController Class** property in the **Pawn** category.
 
 	![](AI_con.png)
 
-1. Compile ![](Engine\AI\BehaviorTrees\QuickStart\4\compile_button.png) and Save ![](Engine\AI\BehaviorTrees\QuickStart\4\save_button.png)
+1. **Compile** ![](Engine\AI\BehaviorTrees\QuickStart\4\compile_button.png) and **Save**. ![](Engine\AI\BehaviorTrees\QuickStart\4\save_button.png)
 
 
 %Steps%
