UEDOC-1400

Placeing Previous and Next at the top of Content Example Pages

diff --git a/Engine/Documentation/Source/Resources/ContentExamples/ImportOptions/1_1/1_1.INT.udn b/Engine/Documentation/Source/Resources/ContentExamples/ImportOptions/1_1/1_1.INT.udn
--- a/Engine/Documentation/Source/Resources/ContentExamples/ImportOptions/1_1/1_1.INT.udn
+++ b/Engine/Documentation/Source/Resources/ContentExamples/ImportOptions/1_1/1_1.INT.udn
@@ -12,7 +12,22 @@ Related: Engine/Content/FBX/MorphTargets
 Related: Shared/Editor/FbxErrors
 Related: Engine/Content/FBX/Materials
 
+[VAR:Steps]
+[OBJECT:Navigation]
+	[PARAM:previous]
+	[/PARAM]
+	[PARAM:current]
+	[/PARAM]
+	[PARAM:home]
+		[Import Options Home](Resources/ContentExamples/ImportOptions)
+	[/PARAM]
+	[PARAM:next]
+		[Next Example](Resources/ContentExamples/ImportOptions/1_2 "%Resources/ContentExamples/ImportOptions/1_2:title%")
+	[/PARAM]
+[/OBJECT]
+[/VAR]
 
+%Steps%
 
 [VAR:TopicCompact]
 [OBJECT:TopicCompact]
diff --git a/Engine/Documentation/Source/Resources/ContentExamples/ImportOptions/1_2/1_2.INT.udn b/Engine/Documentation/Source/Resources/ContentExamples/ImportOptions/1_2/1_2.INT.udn
--- a/Engine/Documentation/Source/Resources/ContentExamples/ImportOptions/1_2/1_2.INT.udn
+++ b/Engine/Documentation/Source/Resources/ContentExamples/ImportOptions/1_2/1_2.INT.udn
@@ -13,6 +13,24 @@ Related: Shared/Editor/FbxErrors
 Related: Engine/Content/FBX/Materials
 Related: Engine/Content/Types/StaticMeshes/HowTo/LODs
 
+[VAR:Steps]
+[OBJECT:Navigation]
+	[PARAM:previous]
+		[Previous Example](Resources/ContentExamples/ImportOptions/1_1 "%Resources/ContentExamples/ImportOptions/1_1:title%")
+	[/PARAM]
+	[PARAM:current]
+	[/PARAM]
+	[PARAM:home]
+		[Import Options Home](Resources/ContentExamples/ImportOptions)
+	[/PARAM]
+	[PARAM:next]
+		[Next Example](Resources/ContentExamples/ImportOptions/1_3 "%Resources/ContentExamples/ImportOptions/1_3:title%")
+	[/PARAM]
+[/OBJECT]
+[/VAR]
+
+%Steps%
+
 
 [VAR:TopicCompact]
 [OBJECT:TopicCompact]
diff --git a/Engine/Documentation/Source/Resources/ContentExamples/ImportOptions/1_3/1_3.INT.udn b/Engine/Documentation/Source/Resources/ContentExamples/ImportOptions/1_3/1_3.INT.udn
--- a/Engine/Documentation/Source/Resources/ContentExamples/ImportOptions/1_3/1_3.INT.udn
+++ b/Engine/Documentation/Source/Resources/ContentExamples/ImportOptions/1_3/1_3.INT.udn
@@ -12,6 +12,23 @@ Related: Engine/Content/FBX/MorphTargets
 Related: Shared/Editor/FbxErrors
 Related: Engine/Content/FBX/Materials
 
+[VAR:Steps]
+[OBJECT:Navigation]
+	[PARAM:previous]
+		[Previous Example](Resources/ContentExamples/ImportOptions/1_2 "%Resources/ContentExamples/ImportOptions/1_2:title%")
+	[/PARAM]
+	[PARAM:current]
+	[/PARAM]
+	[PARAM:home]
+		[Import Options Home](Resources/ContentExamples/ImportOptions)
+	[/PARAM]
+	[PARAM:next]
+		[Next Example](Resources/ContentExamples/ImportOptions/1_4 "%Resources/ContentExamples/ImportOptions/1_4:title%")
+	[/PARAM]
+[/OBJECT]
+[/VAR]
+
+%Steps%
 
 
 [VAR:TopicCompact]
diff --git a/Engine/Documentation/Source/Resources/ContentExamples/ImportOptions/1_4/1_4.INT.udn b/Engine/Documentation/Source/Resources/ContentExamples/ImportOptions/1_4/1_4.INT.udn
--- a/Engine/Documentation/Source/Resources/ContentExamples/ImportOptions/1_4/1_4.INT.udn
+++ b/Engine/Documentation/Source/Resources/ContentExamples/ImportOptions/1_4/1_4.INT.udn
@@ -12,6 +12,24 @@ Related: Engine/Content/FBX/MorphTargets
 Related: Shared/Editor/FbxErrors
 Related: Engine/Content/FBX/Materials
 
+[VAR:Steps]
+[OBJECT:Navigation]
+	[PARAM:previous]
+		[Previous Example](Resources/ContentExamples/ImportOptions/1_3 "%Resources/ContentExamples/ImportOptions/1_3:title%")
+	[/PARAM]
+	[PARAM:current]
+	[/PARAM]
+	[PARAM:home]
+		[Import Options Home](Resources/ContentExamples/ImportOptions)
+	[/PARAM]
+	[PARAM:next]
+		[Next Example](Resources/ContentExamples/ImportOptions/1_5 "%Resources/ContentExamples/ImportOptions/1_5:title%")
+	[/PARAM]
+[/OBJECT]
+[/VAR]
+
+%Steps%
+
 
 [VAR:TopicCompact]
 [OBJECT:TopicCompact]
diff --git a/Engine/Documentation/Source/Resources/ContentExamples/ImportOptions/1_5/1_5.INT.udn b/Engine/Documentation/Source/Resources/ContentExamples/ImportOptions/1_5/1_5.INT.udn
--- a/Engine/Documentation/Source/Resources/ContentExamples/ImportOptions/1_5/1_5.INT.udn
+++ b/Engine/Documentation/Source/Resources/ContentExamples/ImportOptions/1_5/1_5.INT.udn
@@ -12,6 +12,24 @@ Related: Engine/Content/FBX/MorphTargets
 Related: Shared/Editor/FbxErrors
 Related: Engine/Content/FBX/Materials
 
+[VAR:Steps]
+[OBJECT:Navigation]
+	[PARAM:previous]
+		[Previous Example](Resources/ContentExamples/ImportOptions/1_4 "%Resources/ContentExamples/ImportOptions/1_4:title%")
+	[/PARAM]
+	[PARAM:current]
+	[/PARAM]
+	[PARAM:home]
+		[Import Options Home](Resources/ContentExamples/ImportOptions)
+	[/PARAM]
+	[PARAM:next]
+		[Next Example](Resources/ContentExamples/ImportOptions/1_6 "%Resources/ContentExamples/ImportOptions/1_6:title%")
+	[/PARAM]
+[/OBJECT]
+[/VAR]
+
+%Steps%
+
 
 [VAR:TopicCompact]
 [OBJECT:TopicCompact]
diff --git a/Engine/Documentation/Source/Resources/ContentExamples/ImportOptions/1_6/1_6.INT.udn b/Engine/Documentation/Source/Resources/ContentExamples/ImportOptions/1_6/1_6.INT.udn
--- a/Engine/Documentation/Source/Resources/ContentExamples/ImportOptions/1_6/1_6.INT.udn
+++ b/Engine/Documentation/Source/Resources/ContentExamples/ImportOptions/1_6/1_6.INT.udn
@@ -12,6 +12,23 @@ Related: Engine/Content/FBX/MorphTargets
 Related: Shared/Editor/FbxErrors
 Related: Engine/Content/FBX/Materials
 
+[VAR:Steps]
+[OBJECT:Navigation]
+	[PARAM:previous]
+		[Previous Example](Resources/ContentExamples/ImportOptions/1_5 "%Resources/ContentExamples/ImportOptions/1_5:title%")
+	[/PARAM]
+	[PARAM:current]
+	[/PARAM]
+	[PARAM:home]
+		[Import Options Home](Resources/ContentExamples/ImportOptions)
+	[/PARAM]
+	[PARAM:next]
+		[Next Example](Resources/ContentExamples/ImportOptions/1_7 "%Resources/ContentExamples/ImportOptions/1_7:title%")
+	[/PARAM]
+[/OBJECT]
+[/VAR]
+
+%Steps%
 
 
 [VAR:TopicCompact]
diff --git a/Engine/Documentation/Source/Resources/ContentExamples/ImportOptions/1_7/1_7.INT.udn b/Engine/Documentation/Source/Resources/ContentExamples/ImportOptions/1_7/1_7.INT.udn
--- a/Engine/Documentation/Source/Resources/ContentExamples/ImportOptions/1_7/1_7.INT.udn
+++ b/Engine/Documentation/Source/Resources/ContentExamples/ImportOptions/1_7/1_7.INT.udn
@@ -12,6 +12,22 @@ Related: Engine/Content/FBX/MorphTargets
 Related: Shared/Editor/FbxErrors
 Related: Engine/Content/FBX/Materials
 
+[VAR:Steps]
+[OBJECT:Navigation]
+	[PARAM:previous]
+		[Previous Example](Resources/ContentExamples/ImportOptions/1_6 "%Resources/ContentExamples/ImportOptions/1_6:title%")
+	[/PARAM]
+	[PARAM:current]
+	[/PARAM]
+	[PARAM:home]
+		[Import Options Home](Resources/ContentExamples/ImportOptions)
+	[/PARAM]
+	[PARAM:next]
+	[/PARAM]
+[/OBJECT]
+[/VAR]
+
+%Steps%
 
 [VAR:TopicCompact]
 [OBJECT:TopicCompact]
diff --git a/Engine/Documentation/Source/Resources/ContentExamples/Landscapes/1_1/1_1.INT.udn b/Engine/Documentation/Source/Resources/ContentExamples/Landscapes/1_1/1_1.INT.udn
--- a/Engine/Documentation/Source/Resources/ContentExamples/Landscapes/1_1/1_1.INT.udn
+++ b/Engine/Documentation/Source/Resources/ContentExamples/Landscapes/1_1/1_1.INT.udn
@@ -10,6 +10,23 @@ Related: Engine/Landscape/Editing/Splines
 Related: Engine/Landscape/Editing/Gizmos
 Related: Engine/Landscape/Editing/Brushes
 
+[VAR:Steps]
+[OBJECT:Navigation]
+	[PARAM:previous]
+	[/PARAM]
+	[PARAM:current]
+	[/PARAM]
+	[PARAM:home]
+		[Landscapes Home](Resources/ContentExamples/Landscapes)
+	[/PARAM]
+	[PARAM:next]
+		[Next Example](Resources/ContentExamples/Landscapes/1_2 "%Resources/ContentExamples/Landscapes/1_2:title%")
+	[/PARAM]
+[/OBJECT]
+[/VAR]
+
+%Steps%
+
 [VAR:Topic]
 [OBJECT:TopicCompact]
 	[PARAM:image]
diff --git a/Engine/Documentation/Source/Resources/ContentExamples/Landscapes/1_2/1_2.INT.udn b/Engine/Documentation/Source/Resources/ContentExamples/Landscapes/1_2/1_2.INT.udn
--- a/Engine/Documentation/Source/Resources/ContentExamples/Landscapes/1_2/1_2.INT.udn
+++ b/Engine/Documentation/Source/Resources/ContentExamples/Landscapes/1_2/1_2.INT.udn
@@ -10,6 +10,24 @@ Related: Engine/Landscape/Editing/Splines
 Related: Engine/Landscape/Editing/Gizmos
 Related: Engine/Landscape/Editing/Brushes
 
+[VAR:Steps]
+[OBJECT:Navigation]
+	[PARAM:previous]
+		[Previous Example](Resources/ContentExamples/Landscapes/1_1 "%Resources/ContentExamples/Landscapes/1_1:title%")
+	[/PARAM]
+	[PARAM:current]
+	[/PARAM]
+	[PARAM:home]
+		[Landscapes Home](Resources/ContentExamples/Landscapes)
+	[/PARAM]
+	[PARAM:next]
+		[Next Example](Resources/ContentExamples/Landscapes/1_3 "%Resources/ContentExamples/Landscapes/1_3:title%")
+	[/PARAM]
+[/OBJECT]
+[/VAR]
+
+%Steps%
+
 [VAR:Topic]
 [OBJECT:TopicCompact]
 	[PARAM:image]
diff --git a/Engine/Documentation/Source/Resources/ContentExamples/Landscapes/1_3/1_3.INT.udn b/Engine/Documentation/Source/Resources/ContentExamples/Landscapes/1_3/1_3.INT.udn
--- a/Engine/Documentation/Source/Resources/ContentExamples/Landscapes/1_3/1_3.INT.udn
+++ b/Engine/Documentation/Source/Resources/ContentExamples/Landscapes/1_3/1_3.INT.udn
@@ -10,6 +10,23 @@ Related: Engine/Landscape/Editing/Splines
 Related: Engine/Landscape/Editing/Gizmos
 Related: Engine/Landscape/Editing/Brushes
 
+[VAR:Steps]
+[OBJECT:Navigation]
+	[PARAM:previous]
+		[Previous Example](Resources/ContentExamples/Landscapes/1_2 "%Resources/ContentExamples/Landscapes/1_2:title%")
+	[/PARAM]
+	[PARAM:current]
+	[/PARAM]
+	[PARAM:home]
+		[Landscapes Home](Resources/ContentExamples/Landscapes)
+	[/PARAM]
+	[PARAM:next]
+	[/PARAM]
+[/OBJECT]
+[/VAR]
+
+%Steps%
+
 [VAR:Topic]
 [OBJECT:TopicCompact]
 	[PARAM:image]
