UEDOC-1400

Added Previous and Next steps to Content Examples pages

diff --git a/Engine/Documentation/Source/Resources/ContentExamples/LevelDesign/1/1.INT.udn b/Engine/Documentation/Source/Resources/ContentExamples/LevelDesign/1/1.INT.udn
--- a/Engine/Documentation/Source/Resources/ContentExamples/LevelDesign/1/1.INT.udn
+++ b/Engine/Documentation/Source/Resources/ContentExamples/LevelDesign/1/1.INT.udn
@@ -7,6 +7,22 @@ Related:Engine/UI/LevelEditor
 Related: Engine/UI/LevelEditor/Modes
 Related: Engine/Actors/Volumes
 
+[VAR:Steps]
+[OBJECT:Navigation]
+	[PARAM:previous]
+	[/PARAM]
+	[PARAM:current]
+	[/PARAM]
+	[PARAM:home]
+		[Level Design Home](Resources/ContentExamples/LevelDesign)
+	[/PARAM]
+	[PARAM:next]
+		[Next Example](Resources/ContentExamples/LevelDesign/2 "%Resources/ContentExamples/LevelDesign/2:title%")
+	[/PARAM]
+[/OBJECT]
+[/VAR]
+
+%Steps%
 
 [VAR:TopicCompact]
 [OBJECT:TopicCompact]
diff --git a/Engine/Documentation/Source/Resources/ContentExamples/LevelDesign/2/2.INT.udn b/Engine/Documentation/Source/Resources/ContentExamples/LevelDesign/2/2.INT.udn
--- a/Engine/Documentation/Source/Resources/ContentExamples/LevelDesign/2/2.INT.udn
+++ b/Engine/Documentation/Source/Resources/ContentExamples/LevelDesign/2/2.INT.udn
@@ -7,6 +7,23 @@ Related: Engine/Rendering/Materials
 Related:Engine/UI/LevelEditor
 Related: Engine/UI/LevelEditor/Modes
 
+[VAR:Steps]
+[OBJECT:Navigation]
+	[PARAM:previous]
+		[Previous Example](Resources/ContentExamples/LevelDesign/1 "%Resources/ContentExamples/LevelDesign/1:title%")
+	[/PARAM]
+	[PARAM:current]
+	[/PARAM]
+	[PARAM:home]
+		[Level Design Home](Resources/ContentExamples/LevelDesign)
+	[/PARAM]
+	[PARAM:next]
+		[Next Example](Resources/ContentExamples/LevelDesign/3 "%Resources/ContentExamples/LevelDesign/3:title%")
+	[/PARAM]
+[/OBJECT]
+[/VAR]
+
+%Steps%
 
 [VAR:TopicCompact]
 [OBJECT:TopicCompact]
diff --git a/Engine/Documentation/Source/Resources/ContentExamples/LevelDesign/3/3.INT.udn b/Engine/Documentation/Source/Resources/ContentExamples/LevelDesign/3/3.INT.udn
--- a/Engine/Documentation/Source/Resources/ContentExamples/LevelDesign/3/3.INT.udn
+++ b/Engine/Documentation/Source/Resources/ContentExamples/LevelDesign/3/3.INT.udn
@@ -6,6 +6,23 @@ Related: Engine/Rendering/LightingAndShadows/Basics
 Related: Engine/Rendering/PostProcessEffects
 Related: Engine/Rendering/ParticleSystems
 
+[VAR:Steps]
+[OBJECT:Navigation]
+	[PARAM:previous]
+		[Previous Example](Resources/ContentExamples/LevelDesign/2 "%Resources/ContentExamples/LevelDesign/2:title%")
+	[/PARAM]
+	[PARAM:current]
+	[/PARAM]
+	[PARAM:home]
+		[Level Design Home](Resources/ContentExamples/LevelDesign)
+	[/PARAM]
+	[PARAM:next]
+		[Next Example](Resources/ContentExamples/LevelDesign/4 "%Resources/ContentExamples/LevelDesign/4:title%")
+	[/PARAM]
+[/OBJECT]
+[/VAR]
+
+%Steps%
 
 [VAR:TopicCompact]
 [OBJECT:TopicCompact]
diff --git a/Engine/Documentation/Source/Resources/ContentExamples/LevelDesign/4/4.INT.udn b/Engine/Documentation/Source/Resources/ContentExamples/LevelDesign/4/4.INT.udn
--- a/Engine/Documentation/Source/Resources/ContentExamples/LevelDesign/4/4.INT.udn
+++ b/Engine/Documentation/Source/Resources/ContentExamples/LevelDesign/4/4.INT.udn
@@ -6,6 +6,22 @@ Related: Engine/Rendering/LightingAndShadows/Lightmass
 Related: Engine/Rendering/LightingAndShadows/ReflectionEnvironment
 Related: Engine/Rendering/LightingAndShadows/LitTranslucency
 
+[VAR:Steps]
+[OBJECT:Navigation]
+	[PARAM:previous]
+		[Previous Example](Resources/ContentExamples/LevelDesign/3 "%Resources/ContentExamples/LevelDesign/3:title%")
+	[/PARAM]
+	[PARAM:current]
+	[/PARAM]
+	[PARAM:home]
+		[Level Design Home](Resources/ContentExamples/LevelDesign)
+	[/PARAM]
+	[PARAM:next]
+	[/PARAM]
+[/OBJECT]
+[/VAR]
+
+%Steps%
 
 [VAR:TopicCompact]
 [OBJECT:TopicCompact]
