Made corrections per form user

diff --git a/Engine/Documentation/Source/Engine/Rendering/Materials/HowTo/Main_Material_Node/Main_Material_Node.INT.udn b/Engine/Documentation/Source/Engine/Rendering/Materials/HowTo/Main_Material_Node/Main_Material_Node.INT.udn
--- a/Engine/Documentation/Source/Engine/Rendering/Materials/HowTo/Main_Material_Node/Main_Material_Node.INT.udn
+++ b/Engine/Documentation/Source/Engine/Rendering/Materials/HowTo/Main_Material_Node/Main_Material_Node.INT.udn
@@ -90,7 +90,7 @@ Here is a brief breakdown of what each section in the Main Material node **Detai
 
 ## Using the Main Material Node
 
-Setting up a Material to use the Main Material Node is something can be accomplished in the following few steps.
+Setting up a Material to use the Main Material Node is something that can be accomplished in the following few steps.
 
 1. First use the mouse to **Right-Click** in the **Content Browser** and then from the **Create Basic Asset** section of the pop-up Menu, select **Material**. Once the Material has been created, name the Material **MainMaterialNode**.
 
@@ -115,10 +115,10 @@ Setting up a Material to use the Main Material Node is something can be accompli
 	![](Material_Adding_Const3Vector_Node.png)
 
 	[region:tip]
-	There are multiple ways to add nodes to the Material Graph. You can drag and drop form the **Material Pallet**, you can also **Right-Click** in the **Material Graph** and then search for nodes via the Context menu. Finally, you can use Keyboard Shortcuts to also place nodes in the Material Graph. You can read more about this in the [Material Editor UI]( Engine/Rendering/Materials/Editor/Interface) page.
+	There are multiple ways to add nodes to the Material Graph. You can drag and drop form the **Material Pallete**, you can also **Right-Click** in the **Material Graph** and then search for nodes via the Context menu. Finally, you can use Keyboard Shortcuts to also place nodes in the Material Graph. You can read more about this in the [Material Editor UI]( Engine/Rendering/Materials/Editor/Interface) page.
 	[/region]
 
-1. Next add a Constant Material Expression nodes to the Material Graph by first searching in the **Material Pallet** using the word Constant as a search term. When the Constant Material Expression is found select it by clicking on it with the **Left Mouse Button** and with the **Left Mouse Button** still held down, drag the Material Expression node into the Material Graph releasing the **Left Mouse Button** when your mouse is over the location in the Material Graph you want to place the Material Expression node. 
+1. Next add a Constant Material Expression node to the Material Graph by first searching in the **Material Pallete** using the word Constant as a search term. When the Constant Material Expression is found select it by clicking on it with the **Left Mouse Button** and with the **Left Mouse Button** still held down, drag the Material Expression node into the Material Graph and release the **Left Mouse Button** when your mouse is over the location in the Material Graph where you want to place the Material Expression node. 
 
 	[OBJECT:EmbeddedVideo]
 	[PARAMLITERAL:width]
