UEDOC-1400

Added Previous and Next steps to the top of Content Example Pages

diff --git a/Engine/Documentation/Source/Resources/ContentExamples/Lighting/1_1/1_1.INT.udn b/Engine/Documentation/Source/Resources/ContentExamples/Lighting/1_1/1_1.INT.udn
--- a/Engine/Documentation/Source/Resources/ContentExamples/Lighting/1_1/1_1.INT.udn
+++ b/Engine/Documentation/Source/Resources/ContentExamples/Lighting/1_1/1_1.INT.udn
@@ -12,6 +12,22 @@ Related: Engine/Rendering/LightingAndShadows/Shadows
 Related: Engine/Rendering/LightingAndShadows/IndirectLightingCache
 Related: Engine/Rendering/LightingAndShadows/LightMobility
 
+[VAR:Steps]
+[OBJECT:Navigation]
+	[PARAM:previous]
+	[/PARAM]
+	[PARAM:current]
+	[/PARAM]
+	[PARAM:home]
+		[Lighting Home](Resources/ContentExamples/Lighting)
+	[/PARAM]
+	[PARAM:next]
+		[Next Example](Resources/ContentExamples/Lighting/1_2 "%Resources/ContentExamples/Lighting/1_2:title%")
+	[/PARAM]
+[/OBJECT]
+[/VAR]
+
+%Steps%
 
 [VAR:Topic]
 [OBJECT:TopicCompact]
diff --git a/Engine/Documentation/Source/Resources/ContentExamples/Lighting/1_2/1_2.INT.udn b/Engine/Documentation/Source/Resources/ContentExamples/Lighting/1_2/1_2.INT.udn
--- a/Engine/Documentation/Source/Resources/ContentExamples/Lighting/1_2/1_2.INT.udn
+++ b/Engine/Documentation/Source/Resources/ContentExamples/Lighting/1_2/1_2.INT.udn
@@ -12,6 +12,24 @@ Related: Engine/Rendering/LightingAndShadows/Shadows
 Related: Engine/Rendering/LightingAndShadows/IndirectLightingCache
 Related: Engine/Rendering/LightingAndShadows/LightMobility
 
+[VAR:Steps]
+[OBJECT:Navigation]
+	[PARAM:previous]
+		[Previous Example](Resources/ContentExamples/Lighting/1_1 "%Resources/ContentExamples/Lighting/1_1:title%")
+	[/PARAM]
+	[PARAM:current]
+	[/PARAM]
+	[PARAM:home]
+		[Lighting Home](Resources/ContentExamples/Lighting)
+	[/PARAM]
+	[PARAM:next]
+		[Next Example](Resources/ContentExamples/Lighting/1_3 "%Resources/ContentExamples/Lighting/1_3:title%")
+	[/PARAM]
+[/OBJECT]
+[/VAR]
+
+%Steps%
+
 [VAR:Topic]
 [OBJECT:TopicCompact]
 	[PARAM:image]
diff --git a/Engine/Documentation/Source/Resources/ContentExamples/Lighting/1_3/1_3.INT.udn b/Engine/Documentation/Source/Resources/ContentExamples/Lighting/1_3/1_3.INT.udn
--- a/Engine/Documentation/Source/Resources/ContentExamples/Lighting/1_3/1_3.INT.udn
+++ b/Engine/Documentation/Source/Resources/ContentExamples/Lighting/1_3/1_3.INT.udn
@@ -12,6 +12,23 @@ Related: Engine/Rendering/LightingAndShadows/Shadows
 Related: Engine/Rendering/LightingAndShadows/IndirectLightingCache
 Related: Engine/Rendering/LightingAndShadows/LightMobility
 
+[VAR:Steps]
+[OBJECT:Navigation]
+	[PARAM:previous]
+		[Previous Example](Resources/ContentExamples/Lighting/1_2 "%Resources/ContentExamples/Lighting/1_2:title%")
+	[/PARAM]
+	[PARAM:current]
+	[/PARAM]
+	[PARAM:home]
+		[Lighting Home](Resources/ContentExamples/Lighting)
+	[/PARAM]
+	[PARAM:next]
+	[/PARAM]
+[/OBJECT]
+[/VAR]
+
+%Steps%
+
 [VAR:Topic]
 [OBJECT:TopicCompact]
 	[PARAM:image]
diff --git a/Engine/Documentation/Source/Resources/ContentExamples/Lighting/2_1/2_1.INT.udn b/Engine/Documentation/Source/Resources/ContentExamples/Lighting/2_1/2_1.INT.udn
--- a/Engine/Documentation/Source/Resources/ContentExamples/Lighting/2_1/2_1.INT.udn
+++ b/Engine/Documentation/Source/Resources/ContentExamples/Lighting/2_1/2_1.INT.udn
@@ -12,6 +12,23 @@ Related: Engine/Rendering/LightingAndShadows/Shadows
 Related: Engine/Rendering/LightingAndShadows/IndirectLightingCache
 Related: Engine/Rendering/LightingAndShadows/LightMobility
 
+[VAR:Steps]
+[OBJECT:Navigation]
+	[PARAM:previous]
+	[/PARAM]
+	[PARAM:current]
+	[/PARAM]
+	[PARAM:home]
+		[Lighting Home](Resources/ContentExamples/Lighting)
+	[/PARAM]
+	[PARAM:next]
+		[Next Example](Resources/ContentExamples/Lighting/2_2 "%Resources/ContentExamples/Lighting/2_2:title%")
+	[/PARAM]
+[/OBJECT]
+[/VAR]
+
+%Steps%
+
 [VAR:2xSize]456[/VAR]
 
 [VAR:Topic]
diff --git a/Engine/Documentation/Source/Resources/ContentExamples/Lighting/2_2/2_2.INT.udn b/Engine/Documentation/Source/Resources/ContentExamples/Lighting/2_2/2_2.INT.udn
--- a/Engine/Documentation/Source/Resources/ContentExamples/Lighting/2_2/2_2.INT.udn
+++ b/Engine/Documentation/Source/Resources/ContentExamples/Lighting/2_2/2_2.INT.udn
@@ -12,6 +12,24 @@ Related: Engine/Rendering/LightingAndShadows/Shadows
 Related: Engine/Rendering/LightingAndShadows/IndirectLightingCache
 Related: Engine/Rendering/LightingAndShadows/LightMobility
 
+[VAR:Steps]
+[OBJECT:Navigation]
+	[PARAM:previous]
+		[Previous Example](Resources/ContentExamples/Lighting/2_1 "%Resources/ContentExamples/Lighting/2_1:title%")
+	[/PARAM]
+	[PARAM:current]
+	[/PARAM]
+	[PARAM:home]
+		[Lighting Home](Resources/ContentExamples/Lighting)
+	[/PARAM]
+	[PARAM:next]
+		[Next Example](Resources/ContentExamples/Lighting/2_3 "%Resources/ContentExamples/Lighting/2_3:title%")
+	[/PARAM]
+[/OBJECT]
+[/VAR]
+
+%Steps%
+
 [VAR:Topic]
 [OBJECT:TopicCompact]
 	[PARAM:image]
diff --git a/Engine/Documentation/Source/Resources/ContentExamples/Lighting/2_3/2_3.INT.udn b/Engine/Documentation/Source/Resources/ContentExamples/Lighting/2_3/2_3.INT.udn
--- a/Engine/Documentation/Source/Resources/ContentExamples/Lighting/2_3/2_3.INT.udn
+++ b/Engine/Documentation/Source/Resources/ContentExamples/Lighting/2_3/2_3.INT.udn
@@ -12,6 +12,23 @@ Related: Engine/Rendering/LightingAndShadows/Shadows
 Related: Engine/Rendering/LightingAndShadows/IndirectLightingCache
 Related: Engine/Rendering/LightingAndShadows/LightMobility
 
+[VAR:Steps]
+[OBJECT:Navigation]
+	[PARAM:previous]
+		[Previous Example](Resources/ContentExamples/Lighting/2_2 "%Resources/ContentExamples/Lighting/2_2:title%")
+	[/PARAM]
+	[PARAM:current]
+	[/PARAM]
+	[PARAM:home]
+		[Lighting Home](Resources/ContentExamples/Lighting)
+	[/PARAM]
+	[PARAM:next]
+	[/PARAM]
+[/OBJECT]
+[/VAR]
+
+%Steps%
+
 [VAR:Topic]
 [OBJECT:TopicCompact]
 	[PARAM:image]
diff --git a/Engine/Documentation/Source/Resources/ContentExamples/Lighting/3_1/3_1.INT.udn b/Engine/Documentation/Source/Resources/ContentExamples/Lighting/3_1/3_1.INT.udn
--- a/Engine/Documentation/Source/Resources/ContentExamples/Lighting/3_1/3_1.INT.udn
+++ b/Engine/Documentation/Source/Resources/ContentExamples/Lighting/3_1/3_1.INT.udn
@@ -13,6 +13,22 @@ Related: Engine/Rendering/LightingAndShadows/Shadows
 Related: Engine/Rendering/LightingAndShadows/IndirectLightingCache
 Related: Engine/Rendering/LightingAndShadows/LightMobility
 
+[VAR:Steps]
+[OBJECT:Navigation]
+	[PARAM:previous]
+	[/PARAM]
+	[PARAM:current]
+	[/PARAM]
+	[PARAM:home]
+		[Lighting Home](Resources/ContentExamples/Lighting)
+	[/PARAM]
+	[PARAM:next]
+		[Next Example](Resources/ContentExamples/Lighting/3_2 "%Resources/ContentExamples/Lighting/3_2:title%")
+	[/PARAM]
+[/OBJECT]
+[/VAR]
+
+%Steps%
 
 [VAR:Topic]
 [OBJECT:TopicCompact]
diff --git a/Engine/Documentation/Source/Resources/ContentExamples/Lighting/3_2/3_2.INT.udn b/Engine/Documentation/Source/Resources/ContentExamples/Lighting/3_2/3_2.INT.udn
--- a/Engine/Documentation/Source/Resources/ContentExamples/Lighting/3_2/3_2.INT.udn
+++ b/Engine/Documentation/Source/Resources/ContentExamples/Lighting/3_2/3_2.INT.udn
@@ -13,6 +13,22 @@ Related: Engine/Rendering/LightingAndShadows/Shadows
 Related: Engine/Rendering/LightingAndShadows/IndirectLightingCache
 Related: Engine/Rendering/LightingAndShadows/LightMobility
 
+[VAR:Steps]
+[OBJECT:Navigation]
+	[PARAM:previous]
+		[Previous Example](Resources/ContentExamples/Lighting/3_1 "%Resources/ContentExamples/Lighting/3_1:title%")
+	[/PARAM]
+	[PARAM:current]
+	[/PARAM]
+	[PARAM:home]
+		[Lighting Home](Resources/ContentExamples/Lighting)
+	[/PARAM]
+	[PARAM:next]
+	[/PARAM]
+[/OBJECT]
+[/VAR]
+
+%Steps%
 
 [VAR:Topic]
 [OBJECT:TopicCompact]
diff --git a/Engine/Documentation/Source/Resources/ContentExamples/Lighting/4_1/4_1.INT.udn b/Engine/Documentation/Source/Resources/ContentExamples/Lighting/4_1/4_1.INT.udn
--- a/Engine/Documentation/Source/Resources/ContentExamples/Lighting/4_1/4_1.INT.udn
+++ b/Engine/Documentation/Source/Resources/ContentExamples/Lighting/4_1/4_1.INT.udn
@@ -12,6 +12,23 @@ Related: Engine/Rendering/LightingAndShadows/Shadows
 Related: Engine/Rendering/LightingAndShadows/IndirectLightingCache
 Related: Engine/Rendering/LightingAndShadows/LightMobility
 
+[VAR:Steps]
+[OBJECT:Navigation]
+	[PARAM:previous]
+	[/PARAM]
+	[PARAM:current]
+	[/PARAM]
+	[PARAM:home]
+		[Lighting Home](Resources/ContentExamples/Lighting)
+	[/PARAM]
+	[PARAM:next]
+		[Next Example](Resources/ContentExamples/Lighting/4_2 "%Resources/ContentExamples/Lighting/4_2:title%")
+	[/PARAM]
+[/OBJECT]
+[/VAR]
+
+%Steps%
+
 [VAR:Topic]
 [OBJECT:TopicCompact]
 	[PARAM:image]
diff --git a/Engine/Documentation/Source/Resources/ContentExamples/Lighting/4_2/4_2.INT.udn b/Engine/Documentation/Source/Resources/ContentExamples/Lighting/4_2/4_2.INT.udn
--- a/Engine/Documentation/Source/Resources/ContentExamples/Lighting/4_2/4_2.INT.udn
+++ b/Engine/Documentation/Source/Resources/ContentExamples/Lighting/4_2/4_2.INT.udn
@@ -12,6 +12,24 @@ Related: Engine/Rendering/LightingAndShadows/Shadows
 Related: Engine/Rendering/LightingAndShadows/IndirectLightingCache
 Related: Engine/Rendering/LightingAndShadows/LightMobility
 
+[VAR:Steps]
+[OBJECT:Navigation]
+	[PARAM:previous]
+		[Previous Example](Resources/ContentExamples/Lighting/4_1 "%Resources/ContentExamples/Lighting/4_1:title%")
+	[/PARAM]
+	[PARAM:current]
+	[/PARAM]
+	[PARAM:home]
+		[Lighting Home](Resources/ContentExamples/Lighting)
+	[/PARAM]
+	[PARAM:next]
+		[Next Example](Resources/ContentExamples/Lighting/4_3 "%Resources/ContentExamples/Lighting/4_3:title%")
+	[/PARAM]
+[/OBJECT]
+[/VAR]
+
+%Steps%
+
 [VAR:Topic]
 [OBJECT:TopicCompact]
 	[PARAM:image]
diff --git a/Engine/Documentation/Source/Resources/ContentExamples/Lighting/4_3/4_3.INT.udn b/Engine/Documentation/Source/Resources/ContentExamples/Lighting/4_3/4_3.INT.udn
--- a/Engine/Documentation/Source/Resources/ContentExamples/Lighting/4_3/4_3.INT.udn
+++ b/Engine/Documentation/Source/Resources/ContentExamples/Lighting/4_3/4_3.INT.udn
@@ -12,6 +12,22 @@ Related: Engine/Rendering/LightingAndShadows/Shadows
 Related: Engine/Rendering/LightingAndShadows/IndirectLightingCache
 Related: Engine/Rendering/LightingAndShadows/LightMobility
 
+[VAR:Steps]
+[OBJECT:Navigation]
+	[PARAM:previous]
+		[Previous Example](Resources/ContentExamples/Lighting/4_2 "%Resources/ContentExamples/Lighting/4_2:title%")
+	[/PARAM]
+	[PARAM:current]
+	[/PARAM]
+	[PARAM:home]
+		[Lighting Home](Resources/ContentExamples/Lighting)
+	[/PARAM]
+	[PARAM:next]
+	[/PARAM]
+[/OBJECT]
+[/VAR]
+
+%Steps%
 
 [VAR:Topic]
 [OBJECT:TopicCompact]
diff --git a/Engine/Documentation/Source/Resources/ContentExamples/Lighting/5_1/5_1.INT.udn b/Engine/Documentation/Source/Resources/ContentExamples/Lighting/5_1/5_1.INT.udn
--- a/Engine/Documentation/Source/Resources/ContentExamples/Lighting/5_1/5_1.INT.udn
+++ b/Engine/Documentation/Source/Resources/ContentExamples/Lighting/5_1/5_1.INT.udn
@@ -12,6 +12,22 @@ Related: Engine/Rendering/LightingAndShadows/Shadows
 Related: Engine/Rendering/LightingAndShadows/IndirectLightingCache
 Related: Engine/Rendering/LightingAndShadows/LightMobility
 
+[VAR:Steps]
+[OBJECT:Navigation]
+	[PARAM:previous]
+	[/PARAM]
+	[PARAM:current]
+	[/PARAM]
+	[PARAM:home]
+		[Lighting Home](Resources/ContentExamples/Lighting)
+	[/PARAM]
+	[PARAM:next]
+		[Next Example](Resources/ContentExamples/Lighting/5_2 "%Resources/ContentExamples/Lighting/5_2:title%")
+	[/PARAM]
+[/OBJECT]
+[/VAR]
+
+%Steps%
 
 [VAR:Topic]
 [OBJECT:TopicCompact]
diff --git a/Engine/Documentation/Source/Resources/ContentExamples/Lighting/5_2/5_2.INT.udn b/Engine/Documentation/Source/Resources/ContentExamples/Lighting/5_2/5_2.INT.udn
--- a/Engine/Documentation/Source/Resources/ContentExamples/Lighting/5_2/5_2.INT.udn
+++ b/Engine/Documentation/Source/Resources/ContentExamples/Lighting/5_2/5_2.INT.udn
@@ -12,6 +12,24 @@ Related: Engine/Rendering/LightingAndShadows/Shadows
 Related: Engine/Rendering/LightingAndShadows/IndirectLightingCache
 Related: Engine/Rendering/LightingAndShadows/LightMobility
 
+[VAR:Steps]
+[OBJECT:Navigation]
+	[PARAM:previous]
+		[Previous Example](Resources/ContentExamples/Lighting/5_1 "%Resources/ContentExamples/Lighting/5_1:title%")
+	[/PARAM]
+	[PARAM:current]
+	[/PARAM]
+	[PARAM:home]
+		[Lighting Home](Resources/ContentExamples/Lighting)
+	[/PARAM]
+	[PARAM:next]
+		[Next Example](Resources/ContentExamples/Lighting/5_3 "%Resources/ContentExamples/Lighting/5_3:title%")
+	[/PARAM]
+[/OBJECT]
+[/VAR]
+
+%Steps%
+
 [VAR:Topic]
 [OBJECT:TopicCompact]
 	[PARAM:image]
diff --git a/Engine/Documentation/Source/Resources/ContentExamples/Lighting/5_3/5_3.INT.udn b/Engine/Documentation/Source/Resources/ContentExamples/Lighting/5_3/5_3.INT.udn
--- a/Engine/Documentation/Source/Resources/ContentExamples/Lighting/5_3/5_3.INT.udn
+++ b/Engine/Documentation/Source/Resources/ContentExamples/Lighting/5_3/5_3.INT.udn
@@ -12,6 +12,24 @@ Related: Engine/Rendering/LightingAndShadows/Shadows
 Related: Engine/Rendering/LightingAndShadows/IndirectLightingCache
 Related: Engine/Rendering/LightingAndShadows/LightMobility
 
+[VAR:Steps]
+[OBJECT:Navigation]
+	[PARAM:previous]
+		[Previous Example](Resources/ContentExamples/Lighting/5_2 "%Resources/ContentExamples/Lighting/5_2:title%")
+	[/PARAM]
+	[PARAM:current]
+	[/PARAM]
+	[PARAM:home]
+		[Lighting Home](Resources/ContentExamples/Lighting)
+	[/PARAM]
+	[PARAM:next]
+		[Next Example](Resources/ContentExamples/Lighting/5_4 "%Resources/ContentExamples/Lighting/5_4:title%")
+	[/PARAM]
+[/OBJECT]
+[/VAR]
+
+%Steps%
+
 [VAR:Topic]
 [OBJECT:TopicCompact]
 	[PARAM:image]
diff --git a/Engine/Documentation/Source/Resources/ContentExamples/Lighting/5_4/5_4.INT.udn b/Engine/Documentation/Source/Resources/ContentExamples/Lighting/5_4/5_4.INT.udn
--- a/Engine/Documentation/Source/Resources/ContentExamples/Lighting/5_4/5_4.INT.udn
+++ b/Engine/Documentation/Source/Resources/ContentExamples/Lighting/5_4/5_4.INT.udn
@@ -12,6 +12,23 @@ Related: Engine/Rendering/LightingAndShadows/Shadows
 Related: Engine/Rendering/LightingAndShadows/IndirectLightingCache
 Related: Engine/Rendering/LightingAndShadows/LightMobility
 
+[VAR:Steps]
+[OBJECT:Navigation]
+	[PARAM:previous]
+		[Previous Example](Resources/ContentExamples/Lighting/5_3 "%Resources/ContentExamples/Lighting/5_3:title%")
+	[/PARAM]
+	[PARAM:current]
+	[/PARAM]
+	[PARAM:home]
+		[Lighting Home](Resources/ContentExamples/Lighting)
+	[/PARAM]
+	[PARAM:next]
+	[/PARAM]
+[/OBJECT]
+[/VAR]
+
+%Steps%
+
 [VAR:Topic]
 [OBJECT:TopicCompact]
 	[PARAM:image]
diff --git a/Engine/Documentation/Source/Resources/ContentExamples/Lighting/6_1/6_1.INT.udn b/Engine/Documentation/Source/Resources/ContentExamples/Lighting/6_1/6_1.INT.udn
--- a/Engine/Documentation/Source/Resources/ContentExamples/Lighting/6_1/6_1.INT.udn
+++ b/Engine/Documentation/Source/Resources/ContentExamples/Lighting/6_1/6_1.INT.udn
@@ -12,6 +12,23 @@ Related: Engine/Rendering/LightingAndShadows/Shadows
 Related: Engine/Rendering/LightingAndShadows/IndirectLightingCache
 Related: Engine/Rendering/LightingAndShadows/LightMobility
 
+[VAR:Steps]
+[OBJECT:Navigation]
+	[PARAM:previous]
+	[/PARAM]
+	[PARAM:current]
+	[/PARAM]
+	[PARAM:home]
+		[Lighting Home](Resources/ContentExamples/Lighting)
+	[/PARAM]
+	[PARAM:next]
+		[Next Example](Resources/ContentExamples/Lighting/6_2 "%Resources/ContentExamples/Lighting/6_2:title%")
+	[/PARAM]
+[/OBJECT]
+[/VAR]
+
+%Steps%
+
 [VAR:Topic]
 [OBJECT:TopicCompact]
 	[PARAM:image]
diff --git a/Engine/Documentation/Source/Resources/ContentExamples/Lighting/6_2/6_2.INT.udn b/Engine/Documentation/Source/Resources/ContentExamples/Lighting/6_2/6_2.INT.udn
--- a/Engine/Documentation/Source/Resources/ContentExamples/Lighting/6_2/6_2.INT.udn
+++ b/Engine/Documentation/Source/Resources/ContentExamples/Lighting/6_2/6_2.INT.udn
@@ -12,6 +12,23 @@ Related: Engine/Rendering/LightingAndShadows/Shadows
 Related: Engine/Rendering/LightingAndShadows/IndirectLightingCache
 Related: Engine/Rendering/LightingAndShadows/LightMobility
 
+[VAR:Steps]
+[OBJECT:Navigation]
+	[PARAM:previous]
+		[Previous Example](Resources/ContentExamples/Lighting/6_1 "%Resources/ContentExamples/Lighting/6_1:title%")
+	[/PARAM]
+	[PARAM:current]
+	[/PARAM]
+	[PARAM:home]
+		[Lighting Home](Resources/ContentExamples/Lighting)
+	[/PARAM]
+	[PARAM:next]
+	[/PARAM]
+[/OBJECT]
+[/VAR]
+
+%Steps%
+
 [VAR:Topic]
 [OBJECT:TopicCompact]
 	[PARAM:image]
