UEDOC-1400

Adding Previous and Next to Content Example pages

diff --git a/Engine/Documentation/Source/Resources/ContentExamples/MaterialInstances/1_1/1_1.INT.udn b/Engine/Documentation/Source/Resources/ContentExamples/MaterialInstances/1_1/1_1.INT.udn
--- a/Engine/Documentation/Source/Resources/ContentExamples/MaterialInstances/1_1/1_1.INT.udn
+++ b/Engine/Documentation/Source/Resources/ContentExamples/MaterialInstances/1_1/1_1.INT.udn
@@ -10,6 +10,22 @@ Related: Engine/Rendering/Materials/LayeredMaterials
 Related: Engine/Rendering/Materials/MaterialInstances
 Related: Engine/Actors/DecalActor
 
+[VAR:Steps]
+[OBJECT:Navigation]
+	[PARAM:previous]
+	[/PARAM]
+	[PARAM:current]
+	[/PARAM]
+	[PARAM:home]
+		[Content Examples Home](Resources/ContentExamples)
+	[/PARAM]
+	[PARAM:next]
+		[Next Example](Resources/ContentExamples/MaterialInstances/1_2 "%Resources/ContentExamples/MaterialInstances/1_2:title%")
+	[/PARAM]
+[/OBJECT]
+[/VAR]
+
+%Steps%
 
 [VAR:Topic]
 [OBJECT:Topic]
diff --git a/Engine/Documentation/Source/Resources/ContentExamples/MaterialInstances/1_2/1_2.INT.udn b/Engine/Documentation/Source/Resources/ContentExamples/MaterialInstances/1_2/1_2.INT.udn
--- a/Engine/Documentation/Source/Resources/ContentExamples/MaterialInstances/1_2/1_2.INT.udn
+++ b/Engine/Documentation/Source/Resources/ContentExamples/MaterialInstances/1_2/1_2.INT.udn
@@ -10,7 +10,23 @@ Related: Engine/Rendering/Materials/MaterialInputs
 Related: Engine/Rendering/Materials/LayeredMaterials
 Related: Engine/Actors/DecalActor
 
+[VAR:Steps]
+[OBJECT:Navigation]
+	[PARAM:previous]
+		[Previous Example](Resources/ContentExamples/MaterialInstances/1_1 "%Resources/ContentExamples/MaterialInstances/1_1:title%")
+	[/PARAM]
+	[PARAM:current]
+	[/PARAM]
+	[PARAM:home]
+		[Content Examples Home](Resources/ContentExamples)
+	[/PARAM]
+	[PARAM:next]
+		[Next Example](Resources/ContentExamples/MaterialInstances/1_3 "%Resources/ContentExamples/MaterialInstances/1_3:title%")
+	[/PARAM]
+[/OBJECT]
+[/VAR]
 
+%Steps%
 
 [VAR:Topic]
 [OBJECT:Topic]
diff --git a/Engine/Documentation/Source/Resources/ContentExamples/MaterialInstances/1_3/1_3.INT.udn b/Engine/Documentation/Source/Resources/ContentExamples/MaterialInstances/1_3/1_3.INT.udn
--- a/Engine/Documentation/Source/Resources/ContentExamples/MaterialInstances/1_3/1_3.INT.udn
+++ b/Engine/Documentation/Source/Resources/ContentExamples/MaterialInstances/1_3/1_3.INT.udn
@@ -10,6 +10,23 @@ Related: Engine/Rendering/Materials/LayeredMaterials
 Related: Engine/Rendering/Materials/MaterialInstances
 Related: Engine/Actors/DecalActor
 
+[VAR:Steps]
+[OBJECT:Navigation]
+	[PARAM:previous]
+		[Previous Example](Resources/ContentExamples/MaterialInstances/1_2 "%Resources/ContentExamples/MaterialInstances/1_2:title%")
+	[/PARAM]
+	[PARAM:current]
+	[/PARAM]
+	[PARAM:home]
+		[Content Examples Home](Resources/ContentExamples)
+	[/PARAM]
+	[PARAM:next]
+		[Next Example](Resources/ContentExamples/MaterialInstances/1_4 "%Resources/ContentExamples/MaterialInstances/1_4:title%")
+	[/PARAM]
+[/OBJECT]
+[/VAR]
+
+%Steps%
 
 [VAR:Topic]
 [OBJECT:Topic]
diff --git a/Engine/Documentation/Source/Resources/ContentExamples/MaterialInstances/1_4/1_4.INT.udn b/Engine/Documentation/Source/Resources/ContentExamples/MaterialInstances/1_4/1_4.INT.udn
--- a/Engine/Documentation/Source/Resources/ContentExamples/MaterialInstances/1_4/1_4.INT.udn
+++ b/Engine/Documentation/Source/Resources/ContentExamples/MaterialInstances/1_4/1_4.INT.udn
@@ -10,6 +10,22 @@ Related: Engine/Rendering/Materials/LayeredMaterials
 Related: Engine/Rendering/Materials/MaterialInstances
 Related: Engine/Actors/DecalActor
 
+[VAR:Steps]
+[OBJECT:Navigation]
+	[PARAM:previous]
+		[Previous Example](Resources/ContentExamples/MaterialInstances/1_3 "%Resources/ContentExamples/MaterialInstances/1_3:title%")
+	[/PARAM]
+	[PARAM:current]
+	[/PARAM]
+	[PARAM:home]
+		[Content Examples Home](Resources/ContentExamples)
+	[/PARAM]
+	[PARAM:next]
+	[/PARAM]
+[/OBJECT]
+[/VAR]
+
+%Steps%
 
 [VAR:Topic]
 [OBJECT:Topic]
diff --git a/Engine/Documentation/Source/Resources/ContentExamples/MaterialNodes/1_1/1_1.INT.udn b/Engine/Documentation/Source/Resources/ContentExamples/MaterialNodes/1_1/1_1.INT.udn
--- a/Engine/Documentation/Source/Resources/ContentExamples/MaterialNodes/1_1/1_1.INT.udn
+++ b/Engine/Documentation/Source/Resources/ContentExamples/MaterialNodes/1_1/1_1.INT.udn
@@ -11,6 +11,23 @@ Related: Engine/Rendering/Materials/LayeredMaterials
 Related: Engine/Rendering/Materials/MaterialInstances
 Related: Engine/Actors/DecalActor
 
+[VAR:Steps]
+[OBJECT:Navigation]
+	[PARAM:previous]
+	[/PARAM]
+	[PARAM:current]
+	[/PARAM]
+	[PARAM:home]
+		[Material Nodes Home](Resources/ContentExamples/MaterialNodes)
+	[/PARAM]
+	[PARAM:next]
+		[Next Example](Resources/ContentExamples/MaterialNodes/1_2 "%Resources/ContentExamples/MaterialNodes/1_2:title%")
+	[/PARAM]
+[/OBJECT]
+[/VAR]
+
+%Steps%
+
 [VAR:Topic]
 [OBJECT:Topic]
 	[PARAM:image]
diff --git a/Engine/Documentation/Source/Resources/ContentExamples/MaterialNodes/1_10/1_10.INT.udn b/Engine/Documentation/Source/Resources/ContentExamples/MaterialNodes/1_10/1_10.INT.udn
--- a/Engine/Documentation/Source/Resources/ContentExamples/MaterialNodes/1_10/1_10.INT.udn
+++ b/Engine/Documentation/Source/Resources/ContentExamples/MaterialNodes/1_10/1_10.INT.udn
@@ -10,6 +10,24 @@ Related: Engine/Rendering/Materials/LayeredMaterials
 Related: Engine/Rendering/Materials/MaterialInstances
 Related: Engine/Actors/DecalActor
 
+[VAR:Steps]
+[OBJECT:Navigation]
+	[PARAM:previous]
+		[Previous Example](Resources/ContentExamples/MaterialNodes/1_9 "%Resources/ContentExamples/MaterialNodes/1_9:title%")
+	[/PARAM]
+	[PARAM:current]
+	[/PARAM]
+	[PARAM:home]
+		[Material Nodes Home](Resources/ContentExamples/MaterialNodes)
+	[/PARAM]
+	[PARAM:next]
+		[Next Example](Resources/ContentExamples/MaterialNodes/1_11 "%Resources/ContentExamples/MaterialNodes/1_11:title%")
+	[/PARAM]
+[/OBJECT]
+[/VAR]
+
+%Steps%
+
 [VAR:Topic]
 [OBJECT:Topic]
 	[PARAM:image]
diff --git a/Engine/Documentation/Source/Resources/ContentExamples/MaterialNodes/1_11/1_11.INT.udn b/Engine/Documentation/Source/Resources/ContentExamples/MaterialNodes/1_11/1_11.INT.udn
--- a/Engine/Documentation/Source/Resources/ContentExamples/MaterialNodes/1_11/1_11.INT.udn
+++ b/Engine/Documentation/Source/Resources/ContentExamples/MaterialNodes/1_11/1_11.INT.udn
@@ -10,6 +10,23 @@ Related: Engine/Rendering/Materials/LayeredMaterials
 Related: Engine/Rendering/Materials/MaterialInstances
 Related: Engine/Actors/DecalActor
 
+[VAR:Steps]
+[OBJECT:Navigation]
+	[PARAM:previous]
+		[Previous Example](Resources/ContentExamples/MaterialNodes/1_10 "%Resources/ContentExamples/MaterialNodes/1_10:title%")
+	[/PARAM]
+	[PARAM:current]
+	[/PARAM]
+	[PARAM:home]
+		[Material Nodes Home](Resources/ContentExamples/MaterialNodes)
+	[/PARAM]
+	[PARAM:next]
+		[Next Example](Resources/ContentExamples/MaterialNodes/1_12 "%Resources/ContentExamples/MaterialNodes/1_12:title%")
+	[/PARAM]
+[/OBJECT]
+[/VAR]
+
+%Steps%
 
 [VAR:Topic]
 [OBJECT:Topic]
diff --git a/Engine/Documentation/Source/Resources/ContentExamples/MaterialNodes/1_12/1_12.INT.udn b/Engine/Documentation/Source/Resources/ContentExamples/MaterialNodes/1_12/1_12.INT.udn
--- a/Engine/Documentation/Source/Resources/ContentExamples/MaterialNodes/1_12/1_12.INT.udn
+++ b/Engine/Documentation/Source/Resources/ContentExamples/MaterialNodes/1_12/1_12.INT.udn
@@ -10,6 +10,24 @@ Related: Engine/Rendering/Materials/LayeredMaterials
 Related: Engine/Rendering/Materials/MaterialInstances
 Related: Engine/Actors/DecalActor
 
+[VAR:Steps]
+[OBJECT:Navigation]
+	[PARAM:previous]
+		[Previous Example](Resources/ContentExamples/MaterialNodes/1_11 "%Resources/ContentExamples/MaterialNodes/1_11:title%")
+	[/PARAM]
+	[PARAM:current]
+	[/PARAM]
+	[PARAM:home]
+		[Material Nodes Home](Resources/ContentExamples/MaterialNodes)
+	[/PARAM]
+	[PARAM:next]
+		[Next Example](Resources/ContentExamples/MaterialNodes/1_13 "%Resources/ContentExamples/MaterialNodes/1_13:title%")
+	[/PARAM]
+[/OBJECT]
+[/VAR]
+
+%Steps%
+
 [VAR:Topic]
 [OBJECT:Topic]
 	[PARAM:image]
diff --git a/Engine/Documentation/Source/Resources/ContentExamples/MaterialNodes/1_13/1_13.INT.udn b/Engine/Documentation/Source/Resources/ContentExamples/MaterialNodes/1_13/1_13.INT.udn
--- a/Engine/Documentation/Source/Resources/ContentExamples/MaterialNodes/1_13/1_13.INT.udn
+++ b/Engine/Documentation/Source/Resources/ContentExamples/MaterialNodes/1_13/1_13.INT.udn
@@ -10,6 +10,24 @@ Related: Engine/Rendering/Materials/LayeredMaterials
 Related: Engine/Rendering/Materials/MaterialInstances
 Related: Engine/Actors/DecalActor
 
+[VAR:Steps]
+[OBJECT:Navigation]
+	[PARAM:previous]
+		[Previous Example](Resources/ContentExamples/MaterialNodes/1_12 "%Resources/ContentExamples/MaterialNodes/1_12:title%")
+	[/PARAM]
+	[PARAM:current]
+	[/PARAM]
+	[PARAM:home]
+		[Material Nodes Home](Resources/ContentExamples/MaterialNodes)
+	[/PARAM]
+	[PARAM:next]
+		[Next Example](Resources/ContentExamples/MaterialNodes/1_14 "%Resources/ContentExamples/MaterialNodes/1_14:title%")
+	[/PARAM]
+[/OBJECT]
+[/VAR]
+
+%Steps%
+
 [TOC(start:2 end:2)]
 
 [VAR:Topic]
diff --git a/Engine/Documentation/Source/Resources/ContentExamples/MaterialNodes/1_14/1_14.INT.udn b/Engine/Documentation/Source/Resources/ContentExamples/MaterialNodes/1_14/1_14.INT.udn
--- a/Engine/Documentation/Source/Resources/ContentExamples/MaterialNodes/1_14/1_14.INT.udn
+++ b/Engine/Documentation/Source/Resources/ContentExamples/MaterialNodes/1_14/1_14.INT.udn
@@ -10,6 +10,23 @@ Related: Engine/Rendering/Materials/LayeredMaterials
 Related: Engine/Rendering/Materials/MaterialInstances
 Related: Engine/Actors/DecalActor
 
+[VAR:Steps]
+[OBJECT:Navigation]
+	[PARAM:previous]
+		[Previous Example](Resources/ContentExamples/MaterialNodes/1_13 "%Resources/ContentExamples/MaterialNodes/1_13:title%")
+	[/PARAM]
+	[PARAM:current]
+	[/PARAM]
+	[PARAM:home]
+		[Material Nodes Home](Resources/ContentExamples/MaterialNodes)
+	[/PARAM]
+	[PARAM:next]
+	[/PARAM]
+[/OBJECT]
+[/VAR]
+
+%Steps%
+
 [VAR:Topic]
 [OBJECT:Topic]
 	[PARAM:image]
diff --git a/Engine/Documentation/Source/Resources/ContentExamples/MaterialNodes/1_2/1_2.INT.udn b/Engine/Documentation/Source/Resources/ContentExamples/MaterialNodes/1_2/1_2.INT.udn
--- a/Engine/Documentation/Source/Resources/ContentExamples/MaterialNodes/1_2/1_2.INT.udn
+++ b/Engine/Documentation/Source/Resources/ContentExamples/MaterialNodes/1_2/1_2.INT.udn
@@ -11,6 +11,23 @@ Related: Engine/Rendering/Materials/LayeredMaterials
 Related: Engine/Rendering/Materials/MaterialInstances
 Related: Engine/Actors/DecalActor
 
+[VAR:Steps]
+[OBJECT:Navigation]
+	[PARAM:previous]
+		[Previous Example](Resources/ContentExamples/MaterialNodes/1_1 "%Resources/ContentExamples/MaterialNodes/1_1:title%")
+	[/PARAM]
+	[PARAM:current]
+	[/PARAM]
+	[PARAM:home]
+		[Material Nodes Home](Resources/ContentExamples/MaterialNodes)
+	[/PARAM]
+	[PARAM:next]
+		[Next Example](Resources/ContentExamples/MaterialNodes/1_3 "%Resources/ContentExamples/MaterialNodes/1_3:title%")
+	[/PARAM]
+[/OBJECT]
+[/VAR]
+
+%Steps%
 
 [VAR:Topic]
 [OBJECT:Topic]
diff --git a/Engine/Documentation/Source/Resources/ContentExamples/MaterialNodes/1_3/1_3.INT.udn b/Engine/Documentation/Source/Resources/ContentExamples/MaterialNodes/1_3/1_3.INT.udn
--- a/Engine/Documentation/Source/Resources/ContentExamples/MaterialNodes/1_3/1_3.INT.udn
+++ b/Engine/Documentation/Source/Resources/ContentExamples/MaterialNodes/1_3/1_3.INT.udn
@@ -11,6 +11,24 @@ Related: Engine/Rendering/Materials/LayeredMaterials
 Related: Engine/Rendering/Materials/MaterialInstances
 Related: Engine/Actors/DecalActor
 
+[VAR:Steps]
+[OBJECT:Navigation]
+	[PARAM:previous]
+		[Previous Example](Resources/ContentExamples/MaterialNodes/1_2 "%Resources/ContentExamples/MaterialNodes/1_2:title%")
+	[/PARAM]
+	[PARAM:current]
+	[/PARAM]
+	[PARAM:home]
+		[Material Nodes Home](Resources/ContentExamples/MaterialNodes)
+	[/PARAM]
+	[PARAM:next]
+		[Next Example](Resources/ContentExamples/MaterialNodes/1_4 "%Resources/ContentExamples/MaterialNodes/1_4:title%")
+	[/PARAM]
+[/OBJECT]
+[/VAR]
+
+%Steps%
+
 [VAR:Topic]
 [OBJECT:Topic]
 	[PARAM:image]
diff --git a/Engine/Documentation/Source/Resources/ContentExamples/MaterialNodes/1_4/1_4.INT.udn b/Engine/Documentation/Source/Resources/ContentExamples/MaterialNodes/1_4/1_4.INT.udn
--- a/Engine/Documentation/Source/Resources/ContentExamples/MaterialNodes/1_4/1_4.INT.udn
+++ b/Engine/Documentation/Source/Resources/ContentExamples/MaterialNodes/1_4/1_4.INT.udn
@@ -10,7 +10,23 @@ Related: Engine/Rendering/Materials/LayeredMaterials
 Related: Engine/Rendering/Materials/MaterialInstances
 Related: Engine/Actors/DecalActor
 
+[VAR:Steps]
+[OBJECT:Navigation]
+	[PARAM:previous]
+		[Previous Example](Resources/ContentExamples/MaterialNodes/1_3 "%Resources/ContentExamples/MaterialNodes/1_3:title%")
+	[/PARAM]
+	[PARAM:current]
+	[/PARAM]
+	[PARAM:home]
+		[Material Nodes Home](Resources/ContentExamples/MaterialNodes)
+	[/PARAM]
+	[PARAM:next]
+		[Next Example](Resources/ContentExamples/MaterialNodes/1_5 "%Resources/ContentExamples/MaterialNodes/1_5:title%")
+	[/PARAM]
+[/OBJECT]
+[/VAR]
 
+%Steps%
 
 [VAR:Topic]
 [OBJECT:Topic]
diff --git a/Engine/Documentation/Source/Resources/ContentExamples/MaterialNodes/1_5/1_5.INT.udn b/Engine/Documentation/Source/Resources/ContentExamples/MaterialNodes/1_5/1_5.INT.udn
--- a/Engine/Documentation/Source/Resources/ContentExamples/MaterialNodes/1_5/1_5.INT.udn
+++ b/Engine/Documentation/Source/Resources/ContentExamples/MaterialNodes/1_5/1_5.INT.udn
@@ -10,7 +10,23 @@ Related: Engine/Rendering/Materials/LayeredMaterials
 Related: Engine/Rendering/Materials/MaterialInstances
 Related: Engine/Actors/DecalActor
 
+[VAR:Steps]
+[OBJECT:Navigation]
+	[PARAM:previous]
+		[Previous Example](Resources/ContentExamples/MaterialNodes/1_4 "%Resources/ContentExamples/MaterialNodes/1_4:title%")
+	[/PARAM]
+	[PARAM:current]
+	[/PARAM]
+	[PARAM:home]
+		[Material Nodes Home](Resources/ContentExamples/MaterialNodes)
+	[/PARAM]
+	[PARAM:next]
+		[Next Example](Resources/ContentExamples/MaterialNodes/1_6 "%Resources/ContentExamples/MaterialNodes/1_6:title%")
+	[/PARAM]
+[/OBJECT]
+[/VAR]
 
+%Steps%
 
 [VAR:Topic]
 [OBJECT:Topic]
diff --git a/Engine/Documentation/Source/Resources/ContentExamples/MaterialNodes/1_6/1_6.INT.udn b/Engine/Documentation/Source/Resources/ContentExamples/MaterialNodes/1_6/1_6.INT.udn
--- a/Engine/Documentation/Source/Resources/ContentExamples/MaterialNodes/1_6/1_6.INT.udn
+++ b/Engine/Documentation/Source/Resources/ContentExamples/MaterialNodes/1_6/1_6.INT.udn
@@ -10,6 +10,23 @@ Related: Engine/Rendering/Materials/LayeredMaterials
 Related: Engine/Rendering/Materials/MaterialInstances
 Related: Engine/Actors/DecalActor
 
+[VAR:Steps]
+[OBJECT:Navigation]
+	[PARAM:previous]
+		[Previous Example](Resources/ContentExamples/MaterialNodes/1_5 "%Resources/ContentExamples/MaterialNodes/1_5:title%")
+	[/PARAM]
+	[PARAM:current]
+	[/PARAM]
+	[PARAM:home]
+		[Material Nodes Home](Resources/ContentExamples/MaterialNodes)
+	[/PARAM]
+	[PARAM:next]
+		[Next Example](Resources/ContentExamples/MaterialNodes/1_7 "%Resources/ContentExamples/MaterialNodes/1_7:title%")
+	[/PARAM]
+[/OBJECT]
+[/VAR]
+
+%Steps%
 
 [VAR:Topic]
 [OBJECT:Topic]
diff --git a/Engine/Documentation/Source/Resources/ContentExamples/MaterialNodes/1_7/1_7.INT.udn b/Engine/Documentation/Source/Resources/ContentExamples/MaterialNodes/1_7/1_7.INT.udn
--- a/Engine/Documentation/Source/Resources/ContentExamples/MaterialNodes/1_7/1_7.INT.udn
+++ b/Engine/Documentation/Source/Resources/ContentExamples/MaterialNodes/1_7/1_7.INT.udn
@@ -10,7 +10,23 @@ Related: Engine/Rendering/Materials/LayeredMaterials
 Related: Engine/Rendering/Materials/MaterialInstances
 Related: Engine/Actors/DecalActor
 
+[VAR:Steps]
+[OBJECT:Navigation]
+	[PARAM:previous]
+		[Previous Example](Resources/ContentExamples/MaterialNodes/1_6 "%Resources/ContentExamples/MaterialNodes/1_6:title%")
+	[/PARAM]
+	[PARAM:current]
+	[/PARAM]
+	[PARAM:home]
+		[Material Nodes Home](Resources/ContentExamples/MaterialNodes)
+	[/PARAM]
+	[PARAM:next]
+		[Next Example](Resources/ContentExamples/MaterialNodes/1_8 "%Resources/ContentExamples/MaterialNodes/1_8:title%")
+	[/PARAM]
+[/OBJECT]
+[/VAR]
 
+%Steps%
 
 [VAR:Topic]
 [OBJECT:Topic]
diff --git a/Engine/Documentation/Source/Resources/ContentExamples/MaterialNodes/1_8/1_8.INT.udn b/Engine/Documentation/Source/Resources/ContentExamples/MaterialNodes/1_8/1_8.INT.udn
--- a/Engine/Documentation/Source/Resources/ContentExamples/MaterialNodes/1_8/1_8.INT.udn
+++ b/Engine/Documentation/Source/Resources/ContentExamples/MaterialNodes/1_8/1_8.INT.udn
@@ -3,7 +3,23 @@ Title:1.8 - Opacity Mask
 Crumbs:%ROOT%, Resources, Resources/ContentExamples, Resources/ContentExamples/Materials
 Description:An overview of the Material Nodes example level, example 1.8: Opacity Mask.
 
+[VAR:Steps]
+[OBJECT:Navigation]
+	[PARAM:previous]
+		[Previous Example](Resources/ContentExamples/MaterialNodes/1_7 "%Resources/ContentExamples/MaterialNodes/1_7:title%")
+	[/PARAM]
+	[PARAM:current]
+	[/PARAM]
+	[PARAM:home]
+		[Material Nodes Home](Resources/ContentExamples/MaterialNodes)
+	[/PARAM]
+	[PARAM:next]
+		[Next Example](Resources/ContentExamples/MaterialNodes/1_9 "%Resources/ContentExamples/MaterialNodes/1_9:title%")
+	[/PARAM]
+[/OBJECT]
+[/VAR]
 
+%Steps%
 
 [VAR:Topic]
 [OBJECT:Topic]
diff --git a/Engine/Documentation/Source/Resources/ContentExamples/MaterialNodes/1_9/1_9.INT.udn b/Engine/Documentation/Source/Resources/ContentExamples/MaterialNodes/1_9/1_9.INT.udn
--- a/Engine/Documentation/Source/Resources/ContentExamples/MaterialNodes/1_9/1_9.INT.udn
+++ b/Engine/Documentation/Source/Resources/ContentExamples/MaterialNodes/1_9/1_9.INT.udn
@@ -10,6 +10,23 @@ Related: Engine/Rendering/Materials/LayeredMaterials
 Related: Engine/Rendering/Materials/MaterialInstances
 Related: Engine/Actors/DecalActor
 
+[VAR:Steps]
+[OBJECT:Navigation]
+	[PARAM:previous]
+		[Previous Example](Resources/ContentExamples/MaterialNodes/1_8 "%Resources/ContentExamples/MaterialNodes/1_8:title%")
+	[/PARAM]
+	[PARAM:current]
+	[/PARAM]
+	[PARAM:home]
+		[Material Nodes Home](Resources/ContentExamples/MaterialNodes)
+	[/PARAM]
+	[PARAM:next]
+		[Next Example](Resources/ContentExamples/MaterialNodes/1_10 "%Resources/ContentExamples/MaterialNodes/1_10:title%")
+	[/PARAM]
+[/OBJECT]
+[/VAR]
+
+%Steps%
 
 [VAR:Topic]
 [OBJECT:Topic]
