Multiple changees made per forum user's description of issues.

diff --git a/Engine/Documentation/Source/Engine/Rendering/Materials/HowTo/Main_Material_Node/Main_Material_Node.INT.udn b/Engine/Documentation/Source/Engine/Rendering/Materials/HowTo/Main_Material_Node/Main_Material_Node.INT.udn
--- a/Engine/Documentation/Source/Engine/Rendering/Materials/HowTo/Main_Material_Node/Main_Material_Node.INT.udn
+++ b/Engine/Documentation/Source/Engine/Rendering/Materials/HowTo/Main_Material_Node/Main_Material_Node.INT.udn
@@ -29,14 +29,14 @@ Whatever combination of Material Expression nodes is plugged into the inputs of
 
 ![](Material_Main_Node.png)
 
-The Main Material node is responsible for displaying the results of the all the Material Expression nodes that are input into it the various inputs. 
+The Main Material node is responsible for displaying the results of all the Material Expression nodes that are input into it the various inputs. 
 Each input on the Main Material node has a unique effect on how the Material will look and perform. 
 A White colored input means that input will have an affect on the Material, while grayed out input will have no affect on the Material.
 By default, the Main Material node does have some of its inputs grayed out. 
-To enable inputs or disable inputs on the Main Material node ,the following properties have to be modified in the **Details** panel of the Material.
+To enable inputs or disable inputs on the Main Material node, the following properties have to be modified in the **Details** panel of the Material.
 
 [region:note]
-To access the **Details** panel you need make sure that you have the Main Material Node selected. To select the Main Material Node click on it with the **Left Mouse Button**.
+To access the **Details** panel you need to make sure that you have the Main Material Node selected. To select the Main Material Node click on it with the **Left Mouse Button**.
 [/region]
 
 * **Blend Mode**
@@ -110,12 +110,12 @@ Setting up a Material to use the Main Material Node is something that can be acc
 	 
 	 ![](Material_Open_Material_Editor.png)
 
-1. To use the Main Material Node, you first have to add Material Expression nodes to the Material Graph. Add a Constant3Vector Material Expression node the Material Graph by holding down the Number 3 Key on the keyboard and using the **Left Mouse Button** to click anywhere inside of the Material Graph to place the node.
+1. To use the Main Material Node, you first have to add Material Expression nodes to the Material Graph. Add a Constant3Vector Material Expression node to the Material Graph by holding down the Number 3 Key on the keyboard and using the **Left Mouse Button** to click anywhere inside of the Material Graph to place the node.
 
 	![](Material_Adding_Const3Vector_Node.png)
 
 	[region:tip]
-	There are multiple ways to add nodes to the Material Graph. You can drag and drop form the **Material Pallete**, you can also **Right-Click** in the **Material Graph** and then search for nodes via the Context menu. Finally, you can use Keyboard Shortcuts to also place nodes in the Material Graph. You can read more about this in the [Material Editor UI]( Engine/Rendering/Materials/Editor/Interface) page.
+	There are multiple ways to add nodes to the Material Graph. You can drag and drop form the **Material Pallete**, you can also **Right-Click** in the **Material Graph** and then search for nodes via the Context menu. Finally, you can use Keyboard Shortcuts to place nodes in the Material Graph. You can read more about this in the [Material Editor UI]( Engine/Rendering/Materials/Editor/Interface) page.
 	[/region]
 
 1. Next add a Constant Material Expression node to the Material Graph by first searching in the **Material Pallete** using the word Constant as a search term. When the Constant Material Expression is found select it by clicking on it with the **Left Mouse Button** and with the **Left Mouse Button** still held down, drag the Material Expression node into the Material Graph and release the **Left Mouse Button** when your mouse is over the location in the Material Graph where you want to place the Material Expression node. 
@@ -150,7 +150,7 @@ Setting up a Material to use the Main Material Node is something that can be acc
 	[/PARAMLITERAL]
 	[/OBJECT]
 
-1. Now connect the **Constant3Vector** to the **Base Color** input on the Main Material node. To do this, use the **Left Mouse Button** to click on the circle output that is on the right of theConstant3Vector. Then with the **Left Mouse Button** still held down, drag to the right until you see a line coming from the circle output. Connect this line to the Base Color input on the Main Material node.
+1. Now connect the **Constant3Vector** to the **Base Color** input on the Main Material node. To do this, use the **Left Mouse Button** to click on the circle output that is on the right of theConstant3Vector Material Expression node. Then with the **Left Mouse Button** still held down, drag to the right until you see a line coming from the circle output. Connect this line to the Base Color input on the Main Material node.
 	
 	[OBJECT:EmbeddedVideo]
 	[PARAMLITERAL:width]
