#UE4doc # docs - Added table to doc with each decal blend mode and explanation of what they do. Addresses UEDOC-875.

diff --git a/Engine/Documentation/Source/Engine/Actors/DecalActor/DecalUserGuide.INT.udn b/Engine/Documentation/Source/Engine/Actors/DecalActor/DecalUserGuide.INT.udn
--- a/Engine/Documentation/Source/Engine/Actors/DecalActor/DecalUserGuide.INT.udn
+++ b/Engine/Documentation/Source/Engine/Actors/DecalActor/DecalUserGuide.INT.udn
@@ -84,7 +84,26 @@ A deferred decal has only a couple of properties:
 
 ## Material Settings
 
-The new **DecalBlendMode** setting defined how the material properties (diffuse, specular, normal, opacity, ...) are applied to the GBuffer. You can leave the material in the **_DBM_Auto_** mode. This mode chooses one of the other modes depending on which outputs are hooked up. The opacity is used to blend the decal contribution. An efficient decal is manipulating only few GBuffer properties. The cases we currently optimize for are represented by the other modes: **_DBM_Diffuse_**, **_DBM_Specular_**, **_DBM_Emissive_**, **_DBM_Normal_**.
+The **DecalBlendMode** setting defines how the material properties (diffuse, specular, normal, opacity, ...) are applied to the GBuffer. 
+
+The opacity is used to blend the decal contribution. An efficient decal is manipulating only few GBuffer properties. 
+The cases we currently optimize for are represented by the other modes: **_DBM_Diffuse_**, **_DBM_Specular_**, **_DBM_Emissive_**, **_DBM_Normal_**.
+
+| Item | Description | 
+|--- | --- |
+| **Translucent** | This will blend the full material, updating the GBuffer, and does not work for baked lighting. |
+| **Stain** | This is Modulate BaseColor, blend rest, updating the GBuffer, and does not work for baked lighting.  |
+| **Normal** | This will only blend normal, updating the GBuffer, and does not work for baked lighting. |
+| **Emissive** | This is for Additive emissive only. |
+| **DBuffer_Translucent Color, Normal, Roughness** | This is for non-metal, put into DBuffer to work for baked lighting as well. This will become DBM_TranslucentNormal if normal is not hooked up.  |
+| **DBuffer_Translucent Color** | This is for non-metal, put into DBuffer to work for baked lighting as well. |
+| **DBuffer_Translucent Color, Normal** | This is for non-metal, put into the DBuffer to work for baked lighting as well. This will become DBM_DBuffer_Color if normal is not hooked up. |
+| **DBuffer_Translucent Color, Roughness** | This is for non-metal, put into DBuffer to work for baked lighting as well. |
+| **DBuffer_Translucent Normal** | This is for non-metal, put into DBuffer to work for baked lighting as well. |
+| **DBuffer_Translucent Normal, Roughness** | This is for non-metal, put into DBuffer to work for baked lighting as well. This will become DBM_DBuffer_Roughness if normal is not hooked up. |
+| **DBuffer_Translucent Roughness** | This is for non-metal, put into DBuffer to work for baked lighting as well. |
+| **Volumetric Distance Function (experimental)** | This will output a signed distance in Opacity depending on LightVector.  |
+
 
 **_DBM_ScreenSpaceMask_** is special as it affects a special masking channel which is currently used by SSAO (Ambient Occlusion). This allows the decal to override or fade the contribution in some areas.
 
