#UE4doc #docs - addressed issues with repeated text on PBR in page. Addresses UEDOC-1122

diff --git a/Engine/Documentation/Source/Engine/Rendering/Materials/MaterialsAndTextures.INT.udn b/Engine/Documentation/Source/Engine/Rendering/Materials/MaterialsAndTextures.INT.udn
--- a/Engine/Documentation/Source/Engine/Rendering/Materials/MaterialsAndTextures.INT.udn
+++ b/Engine/Documentation/Source/Engine/Rendering/Materials/MaterialsAndTextures.INT.udn
@@ -50,7 +50,9 @@ Navigation:topic
 ![Materials](materials.png)
 [/REGION]
 
-[INCLUDE:Shared/Glossary/M/#Glossary_Material]
+A **Material** is an asset that can be applied to a mesh to control the visual look of the scene. At a high level, it is probably easiest to think of a Material as the "paint" that is applied to an object. But even that can be a little misleading, since a Material literally defines the type of surface from which your object appears to be made. You can define its color, how shiny it is, whether you can see through the object, and much more. 
+
+In more technical terms, when light from the scene hits the surface, a Material is used to calculate how that light interacts with that surface. These calculations are done using incoming data that is input to the Material from a variety of images (textures) and math expressions, as well as from various property settings inherent to the Material itself.
 
 
 [REGION:topics third]
