#UE4doc #docs - updated doc based on UEDOC-1337

diff --git a/Engine/Documentation/Source/Engine/Rendering/Materials/Functions/Overview/FunctionsOverview.INT.udn b/Engine/Documentation/Source/Engine/Rendering/Materials/Functions/Overview/FunctionsOverview.INT.udn
--- a/Engine/Documentation/Source/Engine/Rendering/Materials/Functions/Overview/FunctionsOverview.INT.udn
+++ b/Engine/Documentation/Source/Engine/Rendering/Materials/Functions/Overview/FunctionsOverview.INT.udn
@@ -178,7 +178,8 @@ Note that the input has an option called "Use Preview Value As Default".  When t
 
 ## Parameters
 
-Functions are not allowed to contain parameter node types, as this would cause name conflicts with other functions declaring the same name, but you can still use parameters by passing in the values through inputs.  To use a vector or scalar parameter value with a function, simply pass it in as a function input vector type.  
+Functions can now contain parameter node types. These parameters can be passed directly to any material for use.
+
 
 To use a texture parameter with a function, make a texture input and connect it to the texture object override in a texture sample node:
 
