TypoCorrection

diff --git a/Engine/Documentation/Source/Engine/AI/EnvironmentQuerySystem/QuickStart/4/4.INT.udn b/Engine/Documentation/Source/Engine/AI/EnvironmentQuerySystem/QuickStart/4/4.INT.udn
--- a/Engine/Documentation/Source/Engine/AI/EnvironmentQuerySystem/QuickStart/4/4.INT.udn
+++ b/Engine/Documentation/Source/Engine/AI/EnvironmentQuerySystem/QuickStart/4/4.INT.udn
@@ -1,4 +1,4 @@
-Availability:Public
+﻿Availability:Public
 Title: 4. Create the Basic Assets
 Crumbs:%ROOT%, Engine, Gameplay/AI, Engine/AI/EnvironmentQuerySystem, Engine/AI/EnvironmentQuerySystem/QuickStart
 Description: Create the 7 needed assets for the EQS Quick Start.
@@ -59,7 +59,7 @@ To start, you will need create 7 assets in the **Content Browser**: An AI Contro
 	![](Engine/AI/BehaviorTrees/QuickStart/3/newBB.png)
 	[/REGION]	
 
-	Name the **Blackboard** "RunnerBlackboard'.
+	Name the **Blackboard** "RunnerBlackboard".
 
 1. Create an **Environment Query**: [REGION:raw]**New** %arrow% **Artificial Intelligence** %arrow%  **Environment Query**.[/REGION]
 	
@@ -75,7 +75,7 @@ To start, you will need create 7 assets in the **Content Browser**: An AI Contro
 	
 	![](newContext.png)
 
-	Name the **Environment Query** "PlayerContext'.
+	Name the **Environment Query** "PlayerContext".
 
 %Steps%
 
