#UE4doc #docs - updated page to reflect new categories and changes to the Project Settings in the editor.

diff --git a/Engine/Documentation/Source/Engine/UI/ProjectSettings/ProjectSettings.INT.udn b/Engine/Documentation/Source/Engine/UI/ProjectSettings/ProjectSettings.INT.udn
--- a/Engine/Documentation/Source/Engine/UI/ProjectSettings/ProjectSettings.INT.udn
+++ b/Engine/Documentation/Source/Engine/UI/ProjectSettings/ProjectSettings.INT.udn
@@ -56,39 +56,88 @@ section in the editor.
 
 
 
+
 ### Controls
 
 When viewing the settings within a particular section, the following controls are available along the top of the dialog:
 
 | Control | Description |
 | ------- | ----------- |
-| ![](button_save_defaults.png) | Saves the values of the settings to the default configuration file used by the engine. [REGION:note]If the configuration file is under source control or otherwise not writable, the editor displays a warning and you will not be able to save the file.[/REGION] |
 | ![](button_export.png) | Exports the values of the settings to an external configuration file (`.ini`), which can later be loaded back in. |
 | ![](button_import.png) | Loads values stored in a configuration file (`.ini`), replacing the current values. |
-| ![](button_reset_defaults.png) | Resets any modified values to their default values stored in the default configuration file. |
 
-### Game Category
 
-Contains options that describe and control how the game behaves. 
+### Project Category
+
+Contains options that describe how the project behaves. 
 
 | Section | Description |
 | ------- | ----------- |
-| **General** | Allows you to set information about your project, such as the project name, version, company name, copyright, etc. These are mainly for informational purposes and will not affect how the project runs or behaves. |
-| **Maps** | Contains options for specifying which maps are loaded by default and how they are loaded. |
+| **Description** | Allows you to set information about your project, such as the project name, version, company name, copyright, etc. These are mainly for informational purposes and will not affect how the project runs or behaves. |
+| **Maps & Modes** | Contains options for specifying which maps and modes are loaded by default and how they are loaded. There is also settings for Local Multiplayer screen layout. |
+| **Movies** | Allows you to set start up movies and some options for how they should work. |
+| **Packaging** | Contains options used for packaging your game, like setting content directories, localization, your build configuration, etc.  |
+| **Supported Platforms** | Allows you to set which platforms you are targeting. |
+| **Target Hardware** | Allows you to set how your project should be optimized based on your target hardware. (ie. Desktop vs. Mobile) |
 
 
 ### Engine Category
 
+Contains options for specific systems and sections of the engine.
+
 | Section | Description |
 | ------- | ----------- |
-| **General** | Contains options used by the engine and editor for initialization and setup, such as default fonts, base classes, materials, frame rate, etc. |
-| **Navigation System** | Allows you to configure the navigation system. |
-| **Navigation Mesh** | Contains options for configuring how navigation meshes are generated and displayed. |
+| **AI System** | Contains basic options for the AI system. |
+| **Animation** |  |
+| **Audio** | Contains options for setting default and adding quality levels.  |
+| **Console** | Allows you to set your own in-game inputs for the console. |
+| **Cooker** | Contains options for various cooker settings. These will handle things like compression quality for textures. |
+| **Crowd Manager** | Contains a set of configurable options for managing AI crowds. |
+| **Garbage Collection** | Allows you to configure settings for how Garbage Collection works for your game.  |
+| **General Settings** | Contains options used by the engine and editor for initialization and setup, such as default fonts, base classes, materials, frame rate, etc. |
 | **Input** | Allows you to set up action and axis bindings (keybinds and other input mapping) for your game. [REGION:note]These are in-game keybinds, not editor keybinds.[/REGION] |
+| **Navigation Mesh** | Contains options for configuring how navigation meshes are generated and displayed. |
+| **Navigation System** | Allows you to configure the navigation system. |
+| **Network** | Contains options for Networking. |
+| **Physics** | Contains default options for physics in your game. This will also allow you to setup your own Surface types for your Physical Materails and adjuat how accurate physics simulations can be. |
+| **Rendering** | Contains default settings for a lot of rendering option. There is also settings for default post process settings, culling, mobile, textures, lighting, etc. |
+| **Streaming** | Allows you to configure various streaming options for package streaming, IO, and level streaming. |
+| **Tutorials** | Allows you to set which tutorials are available in this project. |
+| **User Interface** | Contains settings that control Slate and UMG, like DPI scaling, cursors, and the Render Focus Rule.  |
+
+
+### Editor Category
+
+Contains options for how the editor should look and feel on a per-project basis.
+
+| Section | Description |
+| ------- | ----------- |
+| **2D** | Contains configuration options for how the 2D level editor should work.  |
+| **Appearance** | Contains options for the appearance of the editor. There is options to change how units are displayed in the project, such as, choosing to display inches vs centimeters or  celcius vs farenheit. |
+| **Paper2D** | Allows you to set how 2D sprites should be handled when they are imported.  |
+
 
 ### Platforms Category
 
+Contains options for various platforms and platform SDK's.
+
 | Section | Description |
 | ------- | ----------- |
-| **Windows** | Contains options used only when running on Windows platforms. |
-| **iOS** | Allows you to configure options specific to running on the iOS platform. |
+| **Android** | Allows you to configure options specific to running on the Android platform.  |
+| **Android SDK** | Contains location setup for the Android SDK.  [REGION:note]These settings apply to all projects.[/REGION] |
+| **HTML5** | Allows you to configure options specific to running on the HTML5 platform. [REGION:note]These settings apply to all projects.[/REGION] |
+| **HTML5 SDK** | Allows you to set your SDK paths and devices. [REGION:note]These settings apply to all projects.[/REGION] |
+| **iOS** | Allows you to configure options specific to running on the iOS platform.  |
+| **Windows**| Contains options used only when running on Windows platforms. |
+
+
+### Plugins Category
+
+Contains options for select plugins.
+
+| Section | Description |
+| ------- | ----------- |
+| **Paper2D** | Contains settings specific to the Paper2D plugin. |
+| **Slate Remote** | Allows you to configure settings for Slate Remote Server. |
+| **UDP Messaging** | Configure settings for [UDP Messaging](https://docs.unrealengine.com/latest/INT/API/Runtime/Messaging/IMessageBus/index.html) plugin.   |
+
