#UE4doc #docs:
Added engine version field
Unbolded a term

diff --git a/Engine/Documentation/Source/Engine/AI/EnvironmentQuerySystem/QuickStart/6/6.INT.udn b/Engine/Documentation/Source/Engine/AI/EnvironmentQuerySystem/QuickStart/6/6.INT.udn
--- a/Engine/Documentation/Source/Engine/AI/EnvironmentQuerySystem/QuickStart/6/6.INT.udn
+++ b/Engine/Documentation/Source/Engine/AI/EnvironmentQuerySystem/QuickStart/6/6.INT.udn
@@ -2,6 +2,7 @@ Availability:Public
 Title: 6. Blackboard Setup
 Crumbs:%ROOT%, Engine, Gameplay/AI, Engine/AI/EnvironmentQuerySystem, Engine/AI/EnvironmentQuerySystem/QuickStart
 Description: Adding the needed Key Values to the Runner Blackboard.
+version: 4.9
 
 [VAR:Steps]
 [OBJECT:Navigation]
@@ -31,7 +32,7 @@ The **Blackboard** is the memory of the **Behavior Tree**. Here we will store a
 
 	![](finalBB.png)
 
-1. Save your **Blackboard** asset.
+1. Save your Blackboard asset.
 
 
 %Steps%
