#UE4doc #docs:
Added engine version field
Changed font settings for a few terms
Renamed reference to 'Event BeginPlay' from 'On Begin Play' to the correct name.

diff --git a/Engine/Documentation/Source/Engine/AI/EnvironmentQuerySystem/QuickStart/7/7.INT.udn b/Engine/Documentation/Source/Engine/AI/EnvironmentQuerySystem/QuickStart/7/7.INT.udn
--- a/Engine/Documentation/Source/Engine/AI/EnvironmentQuerySystem/QuickStart/7/7.INT.udn
+++ b/Engine/Documentation/Source/Engine/AI/EnvironmentQuerySystem/QuickStart/7/7.INT.udn
@@ -2,6 +2,7 @@ Availability:Public
 Title: 7. AI Controller Setup
 Crumbs:%ROOT%, Engine, Gameplay/AI, Engine/AI/EnvironmentQuerySystem, Engine/AI/EnvironmentQuerySystem/QuickStart
 Description: Setup the AI Controller for the EQS Quick Start.
+version: 4.9
 
 [VAR:Steps]
 [OBJECT:Navigation]
@@ -21,13 +22,13 @@ Description: Setup the AI Controller for the EQS Quick Start.
 
 %Steps%
 
-The setup for this AI Controller is a single **Run Behavior Tree** node running our `RunnerBT` **Behavior Tree** when ever **Event BeginPlay** is executed.
+The setup for this AI Controller is a single **Run Behavior Tree** node running our RunnerBT Behavior Tree when ever **Event BeginPlay** is executed.
 
 ## Steps
 
-1. Double-click the `Runner_AI_CON` Blueprint in the Content Browser to open it for editing.
+1. Double-click the Runner_AI_CON Blueprint in the **Content Browser** to open it for editing.
 
-2. In the Event Graph, have **On Begin Play** execute a **Run Beahvior Tree** node using the `RunnerBT` as the **BTAsset**:
+2. In the **Event Graph**, have **Event BeginPlay** execute a **Run Beahvior Tree** node using the RunnerBT as the **BTAsset**:
 
 	![](setup.png)
 
