#UE4Docs: Cleaning up this page and fixing an older image

diff --git a/Engine/Documentation/Source/Engine/Rendering/Materials/HowTo/Refraction/Refraction.INT.udn b/Engine/Documentation/Source/Engine/Rendering/Materials/HowTo/Refraction/Refraction.INT.udn
--- a/Engine/Documentation/Source/Engine/Rendering/Materials/HowTo/Refraction/Refraction.INT.udn
+++ b/Engine/Documentation/Source/Engine/Rendering/Materials/HowTo/Refraction/Refraction.INT.udn
@@ -18,7 +18,7 @@ Inside Unreal Engine 4, you can simulate this lighting phenomena in your Materia
 ## Refraction
 
 Refraction is the term used to describe the change in direction of a light wave due to a change in its transmission medium. 
-In other words, when light comes in contact with certain surfaces, like water or glass, the light will get slightly bent because those surfaces effect the speed at which light travels through them.
+In other words, when light comes in contact with certain surfaces, like water or glass, the light will get slightly bent because those surfaces affect the speed at which light travels through them.
 The best example of Refraction can be seen by placing part of a pencil in water.
 The part of the pencil that is placed in the water will appear bent  where the pencil and water meet due to Refraction.
 Here is a very simple example showing how Refraction works.
@@ -57,9 +57,9 @@ This tutorial was written assuming that you included the Starter Content with yo
 
 	![](R_Create_Refraction_Material.png)
 
-1. Open up the Material created above by using the **Left Mouse Button** and **Double-Clicking** on the Material in the **Content Browser**. Once opened, change the **Blend Mode** from **Opaque** to **Translucent** and change the **Translucency Lighting Model** from **TLM Volumetric Non Directional** to **TLM_Surface**. When completed, your Material **Details** panel should look like this.
+1. Open up the Material created above by using the **Left Mouse Button** and **Double-Clicking** on the Material in the **Content Browser**. Once opened, change the **Blend Mode** from **Opaque** to **Translucent** and change the **Lighting Model** from **Volumetric Non Directional** to **Surface Translucency Volume**. When completed, your Material **Details** panel should look like this.
 
-	![](R_Details_Setup.png)
+	![](DetailsPanel.jpg)
 
 1. Now that the Material has the correct properties, we can now start adding the required Material Expression nodes. For this Material, we are going to need the following type and amount of Material Expression nodes.
 	
@@ -94,7 +94,7 @@ This tutorial was written assuming that you included the Starter Content with yo
 	If you cannot find the Cube, Cone, or Sphere, this could mean that you did not include the Starter Content with your project. If you did not include the Starter Content, you can either create a new project, this time make sure to include the Starter Content, or you can [Migrate](Engine/Content/Browser/UserGuide/Migrate) the Starter Content from another project into this one.
 	[/region]
 
-1. Now that there are some testing shapes, apply the Refraction Material Instance to the shapes by first selecting the Material Instance with the **Left Mouse Button** in the **Content Browser** and then dragging the Material Instance over the shape. Release the **Left Mouse Button** when you are over the shape you want to apply the Material Instance to.  When completed, you should have something that looks like this in your level.
+1. Now that there are some testing shapes, apply the Refraction Material Instance to the shapes by first selecting the Material Instance with the **Left Mouse Button** in the **Content Browser** and then dragging the Material Instance over a shape and releasing the **Left Mouse Button**.  When completed, you should have something that looks like this in your level.
 
 	![](R_Applied_Material.png)
 
