#UE4doc #docs:
Added engine version field
Removed Monospace from a term

diff --git a/Engine/Documentation/Source/Engine/AI/BehaviorTrees/QuickStart/9/9.INT.udn b/Engine/Documentation/Source/Engine/AI/BehaviorTrees/QuickStart/9/9.INT.udn
--- a/Engine/Documentation/Source/Engine/AI/BehaviorTrees/QuickStart/9/9.INT.udn
+++ b/Engine/Documentation/Source/Engine/AI/BehaviorTrees/QuickStart/9/9.INT.udn
@@ -27,7 +27,7 @@ We need to place our AI_Character into the level so it can interact with the pla
 
 ## Steps
 
-* From the **Content Browser**, click and drag the `AI_Character`, that you created earlier, into the level.
+* From the **Content Browser**, click and drag the AI_Character, that you created earlier, into the level.
 
 [REGION:lightbox]
 [![](placeCharacter.png)(w:940 convert:false)](placeCharacter.png)
diff --git a/Engine/Documentation/Source/Engine/AI/EnvironmentQuerySystem/QuickStart/8/8.INT.udn b/Engine/Documentation/Source/Engine/AI/EnvironmentQuerySystem/QuickStart/8/8.INT.udn
--- a/Engine/Documentation/Source/Engine/AI/EnvironmentQuerySystem/QuickStart/8/8.INT.udn
+++ b/Engine/Documentation/Source/Engine/AI/EnvironmentQuerySystem/QuickStart/8/8.INT.udn
@@ -2,6 +2,7 @@ Availability:Public
 Title: 8. Place the AI Character
 Crumbs:%ROOT%, Engine, Gameplay/AI, Engine/AI/EnvironmentQuerySystem, Engine/AI/EnvironmentQuerySystem/QuickStart
 Description: Place the Runner AI Character into the level so it can interact with the Player and environment.
+version: 4.9
 
 [VAR:Steps]
 [OBJECT:Navigation]
