#UE4doc #docs:
Changed font settings for multiple terms
Added engine version field
Capitalized all mentions of the Details panel

diff --git a/Engine/Documentation/Source/Engine/AI/EnvironmentQuerySystem/QuickStart/9/9.INT.udn b/Engine/Documentation/Source/Engine/AI/EnvironmentQuerySystem/QuickStart/9/9.INT.udn
--- a/Engine/Documentation/Source/Engine/AI/EnvironmentQuerySystem/QuickStart/9/9.INT.udn
+++ b/Engine/Documentation/Source/Engine/AI/EnvironmentQuerySystem/QuickStart/9/9.INT.udn
@@ -2,6 +2,7 @@ Availability:Public
 Title: 9. Behavior Tree Setup
 Crumbs:%ROOT%, Engine, Gameplay/AI, Engine/AI/EnvironmentQuerySystem, Engine/AI/EnvironmentQuerySystem/QuickStart
 Description: Setup the EQS Quick Start Behavior Tree.
+version: 4.9
 
 [VAR:Steps]
 [OBJECT:Navigation]
@@ -25,16 +26,16 @@ Since you already have all the assets necessary to fill out the tree, you're goi
 
 ![](finalBT_Setup.png)
 
-Before we get into the steps to create this [**Behavior Tree**](Engine/AI/BehaviorTrees/UserGuide), let's talk about what it does (or rather what it will do once we create the necessary EQS bits). 
+Before we get into the steps to create this [Behavior Tree](Engine/AI/BehaviorTrees/UserGuide), let's talk about what it does (or rather what it will do once we create the necessary EQS bits). 
 
 * We enter the graph at the top Sequence Node, and then start executing it's children from left to right, and will keep executing as long as its children are successful.
-* First is the EQS Query node that runs the _FindPlayer_ template. We only allow this node to execute once every 0.5 seconds, and if it fails, the tree is locked out of executing the other nodes by the Cooldown Node returning failed, thusly failing its parent sequence node.
+* First is the EQS Query node that runs the FindPlayer template. We only allow this node to execute once every 0.5 seconds, and if it fails, the tree is locked out of executing the other nodes by the Cooldown Node returning failed, thusly failing its parent sequence node.
 * But, if the EQS Query is successful (it finds the player), then the second EQS Query is allowed to execute, finding the AI a place to hide.
 * Finally we move to the place that the second EQS Query found.
 
 ## Steps
 
-1. Double-click the `RunnerBT` Blueprint in the Content Browser to open it for editing.
+1. Double-click the RunnerBT Blueprint in the Content Browser to open it for editing.
 
 1. [REGION:raw]Add a **Sequence Node** (right-click in the Graph Area %arrow% **Composites** %arrow% **Sequence**)[/REGION]
 
@@ -52,16 +53,16 @@ Before we get into the steps to create this [**Behavior Tree**](Engine/AI/Behavi
 
 	![](connectEqs1.png)
 
-1. Select this **Run EQS Query Node**, and in the details panel, set the **Query Template** to `FindPlayer`
-1. While still in the Details panel set the **Blackboard Key** to `ActorToFleeFrom`
+1. Select this **Run EQS Query Node**, and in the **Details** panel, set the **Query Template** to FindPlayer
+1. While still in the Details panel set the **Blackboard Key** to ActorToFleeFrom
 
 	![](setEqs1props.png)
 
-1. [REGION:raw]Add a **Cooldown Decorator** to the **Run EQS Query Node** (right-click the **Run EQS Query Node** %arrow% **Decorators** %arrow% **Cooldown**)[/REGION]
+1. [REGION:raw]Add a **Cooldown Decorator** to the **Run EQS Query Node** (right-click the **Run EQS Query Node** %arrow% **Add Decorator...** %arrow% **Cooldown**)[/REGION]
 
 	![](addCooldown.png)
 
-1. Select the **Cooldown Decorator** and set its **Cool Down Time** to `0.5`
+1. Select the **Cooldown Decorator** and set its **Cool Down Time** to 0.5
 
 	![](setCooldown.png)
 
@@ -69,7 +70,7 @@ Before we get into the steps to create this [**Behavior Tree**](Engine/AI/Behavi
 	
 	![](addEqs2.png)
 
-1. Select the second **Run EQS Query Node**, and in the details panel, set the **Query Template** to `FindHidingSpot`
+1. Select the second **Run EQS Query Node**, and in the **Details** panel, set the **Query Template** to FindHidingSpot
 
 1. Again, while still in the Details panel, set the **Blackboard Key** to **FleeToLocation**
 
@@ -84,8 +85,8 @@ Before we get into the steps to create this [**Behavior Tree**](Engine/AI/Behavi
 
 	![](connectMoveTo.png)
 
-1. Select the Move To node and set **Blackboard Key** to `FleeToLocation`
-1. While still in the Details Panel, set **Acceptable Radius** to `0.5`
+1. Select the Move To node and set **Blackboard Key** to FleeToLocation
+1. While still in the Details Panel, set **Acceptable Radius** to 0.5
 
 	![](setMoveTo.png)
 
