#UE4doc #docs:
Added engine version field
Changed font settings for multiple terms

diff --git a/Engine/Documentation/Source/Engine/AI/EnvironmentQuerySystem/QuickStart/10/10.INT.udn b/Engine/Documentation/Source/Engine/AI/EnvironmentQuerySystem/QuickStart/10/10.INT.udn
--- a/Engine/Documentation/Source/Engine/AI/EnvironmentQuerySystem/QuickStart/10/10.INT.udn
+++ b/Engine/Documentation/Source/Engine/AI/EnvironmentQuerySystem/QuickStart/10/10.INT.udn
@@ -2,6 +2,7 @@ Availability:Public
 Title: 10. Creating the Player Context
 Crumbs:%ROOT%, Engine, Gameplay/AI, Engine/AI/EnvironmentQuerySystem, Engine/AI/EnvironmentQuerySystem/QuickStart
 Description: Create a new Context to be used with in the Environment Query System.
+version: 4.9
 
 [VAR:Steps]
 [OBJECT:Navigation]
@@ -26,12 +27,12 @@ While an Environment Query can generate a number of Items based on getting an Ac
 
 ## Steps
 
-1. Double click the **PlayerContext** Blueprint in the Content Browser to open it for editing.
+1. Double click the **PlayerContext** Blueprint in the **Content Browser** to open it for editing.
 1. Override the **ProvideActorsSet** Function:
 
 	[REGION:raw]![](override1.png)(convert:false)[/REGION]
 
-	The **Override** button is hidden until you mouse over the Functions portion of the My Blueprint Panel.
+	The **Override** button is hidden until you mouse over the **Functions** portion of the **My Blueprint** panel.
 
 	This will override the default behavior of the function with the Blueprint Graph that opens:
 
