#UE4doc #docs:
Added engine version field
Changed font settings for multiple terms

diff --git a/Engine/Documentation/Source/Engine/AI/EnvironmentQuerySystem/QuickStart/11/11.INT.udn b/Engine/Documentation/Source/Engine/AI/EnvironmentQuerySystem/QuickStart/11/11.INT.udn
--- a/Engine/Documentation/Source/Engine/AI/EnvironmentQuerySystem/QuickStart/11/11.INT.udn
+++ b/Engine/Documentation/Source/Engine/AI/EnvironmentQuerySystem/QuickStart/11/11.INT.udn
@@ -2,6 +2,7 @@ Availability:Public
 Title: 11. Creating the FindPlayer EQS
 Crumbs:%ROOT%, Engine, Gameplay/AI, Engine/AI/EnvironmentQuerySystem, Engine/AI/EnvironmentQuerySystem/QuickStart
 Description: Create an Environment Query that finds all instances of the player's TopDownCharacter in the level.
+version: 4.9
 
 [VAR:Steps]
 [OBJECT:Navigation]
@@ -33,7 +34,7 @@ Some quick glossary terms you're going to run into:
 
 ## Steps
 
-1. To start, double click the `FindPlayer` **Environment Query** asset in the **Content Browser**.
+1. To start, double click the FindPlayer **Environment Query** asset in the **Content Browser**.
 
 	![](findPlayerAsset.png)
 
@@ -43,7 +44,7 @@ Some quick glossary terms you're going to run into:
 
 	The Environment Query editor works almost exactly like the Behavior Tree editor.
 
-1. Right-Click into the **Update Graph** area, and select **ActorsOfClass** from the context menu:
+1. **Right-click** into the **Update Graph** area, and select **ActorsOfClass** from the **Context** menu:
 
 	![](createGetActors.png)
 
@@ -72,7 +73,7 @@ Some quick glossary terms you're going to run into:
 
 	Trace tests are useful to check if a Context has line of sight to an Item, in this case, our AI to the Items (TopDownCharacters/Players in this case)
 
-1. Now select the Trace Test on the **ActorsOfClass** node by left-clicking it, and in the Details Panel, change the **Bool Match** Property to `false` (unchecked)
+1. Now select the Trace Test on the **ActorsOfClass** node by left-clicking it, and in the Details Panel, change the **Bool Match** Property to false (unchecked)
 
 	![](aoc_TraceFinal.png)
 
