#UE4doc #docs:
Added engine version field
Changed font settings for multiple terms

diff --git a/Engine/Documentation/Source/Engine/AI/EnvironmentQuerySystem/QuickStart/12/12.INT.udn b/Engine/Documentation/Source/Engine/AI/EnvironmentQuerySystem/QuickStart/12/12.INT.udn
--- a/Engine/Documentation/Source/Engine/AI/EnvironmentQuerySystem/QuickStart/12/12.INT.udn
+++ b/Engine/Documentation/Source/Engine/AI/EnvironmentQuerySystem/QuickStart/12/12.INT.udn
@@ -2,6 +2,7 @@ Availability:Public
 Title: 12. Creating the FindHidingSpot EQS
 Crumbs:%ROOT%, Engine, Gameplay/AI, Engine/AI/EnvironmentQuerySystem, Engine/AI/EnvironmentQuerySystem/QuickStart
 Description: Create an EQS that finds the best point on a grid to hide from the Player's TopDownCharacter.
+version: 4.9
 
 [VAR:Steps]
 [OBJECT:Navigation]
@@ -21,7 +22,7 @@ Description: Create an EQS that finds the best point on a grid to hide from the
 
 %Steps%
 
-For this step you'll be creating a much more involved Environment Query. It's going to create a grid of points, projected onto the NavMesh, then run a series of tests to figure out the "best" place to hide. Some of the tests will simply cull out completely unsatisfactory locations, such as those that are visible to the Player's Character (Trace), others will weight the potential hiding spots based on if they are reachable on the NavMesh (PathFinding) or how much a point is in front of the Player's Character (Dot Product).
+For this step you'll be creating a much more involved Environment Query. It's going to create a grid of points, projected onto the NavMesh, then run a series of tests to figure out the best place to hide. Some of the tests will simply cull out completely unsatisfactory locations, such as those that are visible to the Player's Character (Trace), others will weight the potential hiding spots based on if they are reachable on the NavMesh (PathFinding) or how much a point is in front of the Player's Character (Dot Product).
 
 Just a quick reminder:
 
@@ -31,7 +32,7 @@ Just a quick reminder:
 
 ## Steps
 
-1. To start, **Double-Click** the `FindHidingSpot` Environment Query asset in the **Content Browser**.
+1. To start, **Double-click** the FindHidingSpot Environment Query asset in the **Content Browser**.
 
 	![](findHide.png)
 
@@ -47,11 +48,11 @@ Just a quick reminder:
 
 	These settings for the Grid Generator will create a sizable, but manageable, number of Items for the Query to test against. The **Project Down**, **Project Up,** and **Post Projection Vertical Offset** properties are particularly useful to change if you have an AI that needs to find something above or below them.
 
-1. [REGION:raw]Add 2 **Distance** Tests by **Right-Clicking** the **Simple Grid** node and selecting **Tests** %arrow% **Distance**[/REGION]
+1. [REGION:raw]Add 2 **Distance** Tests by **Right-clicking** the **Simple Grid** node and selecting **Tests** %arrow% **Distance**[/REGION]
 
 	![](addDistance.png)
 
-1. Select the first one and in the **Details** Panel adjust the following Properties:
+1. Select the first one and in the **Details** panel adjust the following properties:
 
 	![](distanceProps.png)
 
@@ -65,19 +66,19 @@ Just a quick reminder:
 
 	[INCLUDE:#eqsqs123]
 
-	This Distance test is meant to give Items _closer_ to the AI a higher Score, and this test carries more weight than the previous **Distance Test** meaning it will favor closer to the AI but still select away from the Player.
+	This Distance test is meant to give Items closer to the AI a higher Score, and this test carries more weight than the previous **Distance Test** meaning it will favor closer to the AI but still select away from the Player.
 
-1. [REGION:raw]Add 2 **Trace** Tests by **Right-Clicking** the **Simple Grid** node and selecting **Tests** %arrow% **Trace**[/REGION]
+1. [REGION:raw]Add 2 **Trace** Tests by **Right-clicking** the **Simple Grid** node and selecting **Tests** %arrow% **Trace**[/REGION]
 
 	![](addTrace.png)
 
-1. Select the one of the Trace tests and in the **Details** panel change **Context** to `PlayerContext` (the Context created earlier in this Quick Start):
+1. Select the one of the Trace tests and in the **Details** panel change **Context** to PlayerContext (the Context created earlier in this Quick Start):
 
 	![](traceProps.png)	
 
 	Here we remove any Item that the Player Pawn or AI pawn has line of sight to. Also, you may be concerned that this test comes after a bunch of **Distance** tests, but it's ok, as the Environment Query runs **Trace** tests first regardless of order here. Also, the other Trace requires no additional setup.
 
-1. [REGION:raw]Add a **Dot** test by **Right-Clicking** the **Simple Grid** node and selecting **Tests** %arrow% **Dot**[/REGION]
+1. [REGION:raw]Add a **Dot** test by **Right-clicking** the **Simple Grid** node and selecting **Tests** %arrow% **Dot**[/REGION]
 
 	![](addDot.png)
 
@@ -87,13 +88,13 @@ Just a quick reminder:
 
 	[INCLUDE:#eqsqs125]
 
-	Here things get tricky, at the high level, this is saying "Favor Items that are to the sides of the Player Pawn," but how it works is a bit odd. A Dot Product function will return how "alike" two vectors are on a scale of 1.0 (perfectly the same) to -1.0 (completely opposite). So if we get the **Absolute Value** of the Dot Product between the rotation vector of the Player and the vector from the Player to the Item we'll favor points in front of and behind the Player. But if we set the **Scoring Factor** to -1.0, then we favor to the sides of the Player.
+	Here things get tricky, at the high level, this is saying "Favor Items that are to the sides of the Player Pawn," but how it works is a bit odd. A Dot Product function will return how alike two vectors are on a scale of 1.0 (perfectly the same) to -1.0 (completely opposite). So if we get the **Absolute Value** of the Dot Product between the rotation vector of the Player and the vector from the Player to the Item we'll favor points in front of and behind the Player. But if we set the **Scoring Factor** to -1.0, then we favor to the sides of the Player.
 
-1. [REGION:raw]Add a **PathFinding** test by **Right-Clicking** the **Simple Grid** node and selecting **Tests** %arrow% **PathFinding**[/REGION]
+1. [REGION:raw]Add a **PathFinding** test by **Right-clicking** the **Simple Grid** node and selecting **Tests** %arrow% **PathFinding**[/REGION]
 
 	![](addPath.png)
 
-1. Select the **PathFinding** test and in the **Details** panel adjust **Test Purpose** to `Score Only`:
+1. Select the **PathFinding** test and in the **Details** panel adjust **Test Purpose** to Score Only:
 
 	![](pathProps.png)
 
@@ -107,29 +108,29 @@ Just a quick reminder:
 
 <!--
 [EXCERPT:eqsqs125]
-* **Line A - Rotation** = `PlayerContext`
-* **Line B - Line From** = `PlayerContext`
-* **Absolute Value** = `True`
-* **Test Purpose** = `Score Only`
-* **Scoring Factor** = `-1.0`
+* **Line A - Rotation** = PlayerContext
+* **Line B - Line From** = PlayerContext
+* **Absolute Value** = True
+* **Test Purpose** = Score Only
+* **Scoring Factor** = -1.0
 [/EXCERPT]
 
 [EXCERPT:eqsqs121]
-* **Grid Size** = `2000.0`
-* **Space Between** = `150.0`
-* **Navigation Filer** = `RecastFilter_UseDefaultArea`
-* **Project Down** = `2048.0`
-* **Post Projection Vertical Offset** = `32.0`
+* **Grid Size** = 2000.0
+* **Space Between** = 150.0
+* **Navigation Filer** = RecastFilter_UseDefaultArea
+* **Project Down** = 2048.0
+* **Post Projection Vertical Offset** = 32.0
 [/EXCERPT]
 
 [EXCERPT:eqsqs122]
-* **Distance To** = `Player Context`
-* **Test Purpose** = `Score Only`
+* **Distance To** = Player Context
+* **Test Purpose** = Score Only
 [/EXCERPT]
 
 [EXCERPT:eqsqs123]
-* **Test Purpose** = `Score Only`
-* **Scoring Equation** = `Inverse Linear`
-* **Scoring Factor** = `2.0`
+* **Test Purpose** = Score Only
+* **Scoring Equation** = Inverse Linear
+* **Scoring Factor** = 2.0
 [/EXCERPT]
 -->
