#UE4doc #docs:
Added engine version
Changed font settings for multiple terms
Changed color of border for a picture from green to yellow as per documentation standards

diff --git a/Engine/Documentation/Source/Engine/AI/EnvironmentQuerySystem/QuickStart/13/13.INT.udn b/Engine/Documentation/Source/Engine/AI/EnvironmentQuerySystem/QuickStart/13/13.INT.udn
--- a/Engine/Documentation/Source/Engine/AI/EnvironmentQuerySystem/QuickStart/13/13.INT.udn
+++ b/Engine/Documentation/Source/Engine/AI/EnvironmentQuerySystem/QuickStart/13/13.INT.udn
@@ -2,6 +2,7 @@ Availability:Public
 Title: 13. Play in Editor and Testing
 Crumbs:%ROOT%, Engine, Gameplay/AI, Engine/AI/EnvironmentQuerySystem, Engine/AI/EnvironmentQuerySystem/QuickStart
 Description: Finally, test your Runner AI by using Play in Editor.
+version: 4.9
 
 [VAR:Steps]
 [OBJECT:Navigation]
@@ -21,7 +22,7 @@ Description: Finally, test your Runner AI by using Play in Editor.
 
 %Steps%
 
-Now since you setup the **Behavior Tree** [earlier](Engine/AI/EnvironmentQuerySystem/QuickStart/9), and told the [AI Controller to run](Engine/AI/EnvironmentQuerySystem/QuickStart/7) it, you are ready to test! 
+Now since you setup the Behavior Tree [earlier](Engine/AI/EnvironmentQuerySystem/QuickStart/9), and told the [AI Controller to run](Engine/AI/EnvironmentQuerySystem/QuickStart/7) it, you are ready to test! 
 
 
 ## Steps
