#UE4doc #docs:
Added engine version field

diff --git a/Engine/Documentation/Source/Engine/AI/EnvironmentQuerySystem/QuickStart/14/14.INT.udn b/Engine/Documentation/Source/Engine/AI/EnvironmentQuerySystem/QuickStart/14/14.INT.udn
--- a/Engine/Documentation/Source/Engine/AI/EnvironmentQuerySystem/QuickStart/14/14.INT.udn
+++ b/Engine/Documentation/Source/Engine/AI/EnvironmentQuerySystem/QuickStart/14/14.INT.udn
@@ -2,6 +2,7 @@ Availability:Public
 Title: 14. Next Steps
 Crumbs:%ROOT%, Engine, Gameplay/AI, Engine/AI/EnvironmentQuerySystem, Engine/AI/EnvironmentQuerySystem/QuickStart
 Description: Suggestions on where to go after completing the EQS Quick Start.
+version: 4.9
 
 [VAR:Steps]
 [OBJECT:Navigation]
