#UE4doc #docs:
Changed step to accurately represent the current UI of the engine

diff --git a/Engine/Documentation/Source/Engine/AI/BehaviorTrees/QuickStart/2/2.INT.udn b/Engine/Documentation/Source/Engine/AI/BehaviorTrees/QuickStart/2/2.INT.udn
--- a/Engine/Documentation/Source/Engine/AI/BehaviorTrees/QuickStart/2/2.INT.udn
+++ b/Engine/Documentation/Source/Engine/AI/BehaviorTrees/QuickStart/2/2.INT.udn
@@ -30,7 +30,7 @@ The NavMesh in Unreal has a number of functions that work well with Behavior Tre
 
 ## Steps
 
-1. [REGION:raw]**Modes** panel %arrow% **Create** %arrow% **Volumes**.[/REGION]
+1. [REGION:raw]**Modes** panel %arrow% **Volumes** %arrow% **Volumes**.[/REGION]
 
 1. Click and drag a **Nav Mesh Bounds Volume** into the map.
 
