#UE4doc #docs:
Changed step and picture about applying the AnimBP to the Character blueprint to reflect new changes to the mannequin

diff --git a/Engine/Documentation/Source/Engine/AI/EnvironmentQuerySystem/QuickStart/5/5.INT.udn b/Engine/Documentation/Source/Engine/AI/EnvironmentQuerySystem/QuickStart/5/5.INT.udn
--- a/Engine/Documentation/Source/Engine/AI/EnvironmentQuerySystem/QuickStart/5/5.INT.udn
+++ b/Engine/Documentation/Source/Engine/AI/EnvironmentQuerySystem/QuickStart/5/5.INT.udn
@@ -42,7 +42,7 @@ This Character will become our new AI Character. The reason we use a Character,
 
 1. Rotate the **Mesh** Component so it faces the direction the blue **Arrow Component** is pointing.
 
-1. In the **Details** panel, find the **Animation** category and assign the TopDown_AnimBlueprint to the **Animation Blueprint Generated Class** property.
+1. In the **Details** panel, find the **Animation** category and assign the ThirdPerson_AnimBP to the **Animation Blueprint Generated Class** property.
 
 	![](Engine\AI\BehaviorTrees\QuickStart\4\assignAnimBP.png)
 
