#UE4doc #docs:
Fixed an error in one of the steps for the tutorial

diff --git a/Engine/Documentation/Source/Engine/AI/EnvironmentQuerySystem/QuickStart/9/9.INT.udn b/Engine/Documentation/Source/Engine/AI/EnvironmentQuerySystem/QuickStart/9/9.INT.udn
--- a/Engine/Documentation/Source/Engine/AI/EnvironmentQuerySystem/QuickStart/9/9.INT.udn
+++ b/Engine/Documentation/Source/Engine/AI/EnvironmentQuerySystem/QuickStart/9/9.INT.udn
@@ -76,7 +76,7 @@ Before we get into the steps to create this [Behavior Tree](Engine/AI/BehaviorTr
 
 	![](setEqs2props.png)
 
-1. Position it to the right of the second **Run EQS Query Node**
+1. Position it to the right of the first **Run EQS Query Node**
 1. [REGION:raw]Create a **Move To Node** (right-click in the Graph Area %arrow% **Tasks** %arrow% **Move To**)[/REGION]
 
 	![](addMoveTo.png)
