#UE4doc #docs:
Removed Overview header
Removed italics from multiple terms in the intro excerpt

diff --git a/Engine/Documentation/Source/Engine/Animation/Overview/AnimationOverview.INT.udn b/Engine/Documentation/Source/Engine/Animation/Overview/AnimationOverview.INT.udn
--- a/Engine/Documentation/Source/Engine/Animation/Overview/AnimationOverview.INT.udn
+++ b/Engine/Documentation/Source/Engine/Animation/Overview/AnimationOverview.INT.udn
@@ -5,7 +5,6 @@ Description:Overview of the system for deforming SkeletalMeshes based on animati
 
 [TOC(start:2 end:2)]
 
-## Overview
 
 [EXCERPT:Intro]
 Unreal Engine 4's animation system allows for a deep level of control for your characters and Skeletal Meshes. In this document, we will provide a high level overview  of how animation works in UE4, and the base concepts that drive it. <!---In the documents that follow, we'll cover the---> 
@@ -14,9 +13,9 @@ Unreal Engine 4's animation system allows for a deep level of control for your c
 
 The **Animation** system mixes skeletal-based deformation of meshes with morph-based vertex deformation to 
 allow for complex animation. This system can be used to make basic player movement appear much more realistic 
-by playing and blending between canned _AnimationSequences_, create customized special moves such as scaling 
-ledges and walls using _AnimMontages_, apply damage effects or facial expressions through _MorphTargets_, or 
-directly control the transformations of bones using _SkeletalControls_.
+by playing and blending between canned AnimationSequences, create customized special moves such as scaling 
+ledges and walls using AnimMontages, apply damage effects or facial expressions through MorphTargets, or 
+directly control the transformations of bones using SkeletalControls.
 [/EXCERPT:Intro]
 
 
