#UE4doc #docs:
Added engine version
Removed Overview header
Replaced picture of the 'Add New' button as it was out of date
Changed font settings for multiple terms

diff --git a/Engine/Documentation/Source/Engine/Animation/AimOffset/AimOffset.INT.udn b/Engine/Documentation/Source/Engine/Animation/AimOffset/AimOffset.INT.udn
--- a/Engine/Documentation/Source/Engine/Animation/AimOffset/AimOffset.INT.udn
+++ b/Engine/Documentation/Source/Engine/Animation/AimOffset/AimOffset.INT.udn
@@ -1,19 +1,19 @@
 Availability:Public
 Title: Aim Offset
 Crumbs: %ROOT%, Engine, Engine/Animation
-Description:An Aim Offset is an asset that stores a blendable series of poses to help a character aim a weapon.
+Description:An **Aim Offset** is an asset that stores a blendable series of poses to help a character aim a weapon.
 Related: Engine/Content/FBX/Animations
 Related: Engine/Content/Types/SkeletalMeshes
 Related: Engine/Animation/PhysicallyDrivenAnimation
 Related: Engine/Animation/StateMachines
 Related: Gameplay/AI
 Related: Engine/Content/Tools/MayaRiggingTool
+version: 4.9
 
 [REGION:fullwidth]
 ![](AimOffset.png)
 [/REGION]
 
-## Overview
 
 %Engine/Animation/AimOffset:description% During animation, the result of the Aim Offset is blended with other motions such as running, walking, jumping, etc. to cause the character to appear to smoothly look in all directions.
 
@@ -22,13 +22,13 @@ The key concept behind Aim Offsets is that they are designed to be additively bl
 In many ways, an Aim Offset is just like a Blend Space. As such, you would do well to get a solid understanding of Blend Spaces before you use them. Please see the [](Engine/Animation/Blendspaces) documentation for details.
 
 [REGION:tip]
-You can also refer to the [Animation Content Examples](Resources\ContentExamples\Animation) page under section 1.4 for an example of an **Aim Offset** applied to a character.
+You can also refer to the [Animation Content Examples](Resources\ContentExamples\Animation) page under section 1.4 for an example of an Aim Offset applied to a character.
 [/REGION]
 
 
 ## Creating Aim Offsets
 
-Creating an Aim Offset is done in the **Content Browser** just as with any other asset. You may either **right-click** or use the ![](New_Asset_Button.png) button and choose **Animation > Aim Offset** from the context menu. You will  notice that you also have the option of an Aim Offset 1D. These are the same as Aim Offsets, but only support input from one variable to control the blending, where a standard Aim Offset supports two.
+Creating an Aim Offset is done in the **Content Browser** just as with any other asset. You may either **Right-click** or use the ![](New_Asset_Button.png) button and choose **Animation > Aim Offset** from the context menu. You will  notice that you also have the option of an **Aim Offset 1D**. These are the same as Aim Offsets, but only support input from one variable to control the blending, where a standard Aim Offset supports two.
 
 When creating a new Aim Offset, you will have to designate a Skeleton asset. Be sure to choose the same one used by the Skeletal Mesh you wish to use with the Aim Offset.
 
@@ -71,12 +71,12 @@ Once you import each of the poses you need into Unreal Engine, you will need to
 
 * **Additive Aim Type** - AAT_MeshSpace
 
-[REGION:note]
-Aim Offsets only accept animations with an Additive Aim Type of Mesh Space!
+[REGION:warning]
+Aim Offsets only accept animations with an Additive Aim Type of Mesh Space.
 [/REGION]
 
 
-Mesh space is a bit of a different approach for extracting delta of additive animation, in that it is only used in certain instances such as Aim Offset assets. Mesh space uses the mesh's bounding box as its base for rotation, allowing rotations to move in the same direction regardless of the orientation of chain of the bone in the Skeletal Mesh. Consider a character that was leaning to the side, but needing to aim a pistol _upwards_. If that upward aiming motion was in Local space, then the aim would rotate outward, in the direction of the lean.
+Mesh space is a bit of a different approach for extracting delta of additive animation, in that it is only used in certain instances such as Aim Offset assets. Mesh space uses the mesh's bounding box as its base for rotation, allowing rotations to move in the same direction regardless of the orientation of chain of the bone in the Skeletal Mesh. Consider a character that was leaning to the side, but needing to aim a pistol upwards. If that upward aiming motion was in Local space, then the aim would rotate outward, in the direction of the lean.
 
 [VAR:MSSize]w:450[/VAR]
 [VAR:MSSize2]w:295[/VAR]
@@ -95,7 +95,7 @@ Mesh space is a bit of a different approach for extracting delta of additive ani
 
 This is solved with Mesh space rotation. Since rotation is based on the Skeletal Mesh's box, the orientation of that rotation remains constant no matter what the character is doing. So even though the character might be leaning sideways, the upward aim would still rotate in an upward direction. 
 
-Mesh space is set as a property of an Animation Sequence. Under the **Additive Settings** category, you will see the **Additive Anim Type** property. Setting this to _Mesh Space_ will utilize the Mesh coordinate space as described. This is a fairly expensive blending operation, and should generally only be used by certain types of Blend Spaces where you know you need the animation to move in a constant direction, such as in the case of an Aim Offset.
+Mesh space is set as a property of an Animation Sequence. Under the **Additive Settings** category, you will see the **Additive Anim Type** property. Setting this to Mesh Space will utilize the Mesh coordinate space as described. This is a fairly expensive blending operation, and should generally only be used by certain types of Blend Spaces where you know you need the animation to move in a constant direction, such as in the case of an Aim Offset.
 
 
 
