#UE4doc #docs:
Changed font settings for multiple terms

diff --git a/Engine/Documentation/Source/Engine/Animation/Blendspaces/Blendspaces.INT.udn b/Engine/Documentation/Source/Engine/Animation/Blendspaces/Blendspaces.INT.udn
--- a/Engine/Documentation/Source/Engine/Animation/Blendspaces/Blendspaces.INT.udn
+++ b/Engine/Documentation/Source/Engine/Animation/Blendspaces/Blendspaces.INT.udn
@@ -36,16 +36,16 @@ Related: Engine/Content/Tools/MayaRiggingTool
 ## Overview
 
 [EXCERPT:Intro]
-**_Blend Spaces_** are special assets which can be sampled in **AnimGraphs** that allow for blending of 
+**Blend Spaces** are special assets which can be sampled in **AnimGraphs** that allow for blending of 
 animations based on the values of two inputs. Simple blending between two animations based on a 
 single input can be performed using one of the standard [blend nodes](Engine/Animation/NodeReference/Blend) 
-available in _Animation Blueprints_. _Blend Spaces_ provide a means of doing more complex blending 
+available in Animation Blueprints. Blend Spaces provide a means of doing more complex blending 
 between multiple animations based on multiple values (currently limited to two).
 
-The goal of _Blend Spaces_ is to reduce the need for creating individual, hard-coded nodes to perform 
+The goal of Blend Spaces is to reduce the need for creating individual, hard-coded nodes to perform 
 blending based on specific properties or conditions. By allowing the animator or programmer to specify 
 the inputs, the animations, and how the inputs are used to blend between animations, virtually any 
-type of blending can be performed using the generic _Blend Space_.
+type of blending can be performed using the generic Blend Space.
 [/EXCERPT:Intro]
 
 
