#UE4doc #docs:
Removed Overview header
Added engine version field

diff --git a/Engine/Documentation/Source/Engine/Animation/AnimationBlending/AnimationBlending.INT.udn b/Engine/Documentation/Source/Engine/Animation/AnimationBlending/AnimationBlending.INT.udn
--- a/Engine/Documentation/Source/Engine/Animation/AnimationBlending/AnimationBlending.INT.udn
+++ b/Engine/Documentation/Source/Engine/Animation/AnimationBlending/AnimationBlending.INT.udn
@@ -1,9 +1,8 @@
 Availability:Public
 Title: Blending Animations
 Crumbs: %ROOT%, Engine, Engine/Animation
-Description:A look at various ways to blend together multiple animations in UE4.
+version: 4.9
 
-## Overview
 
 Animation blending, as a concept, simply means making a smooth transition between two or more animations on a single character or Skeletal Mesh. In Unreal Engine 4, there are various ways in which such blending can be applied. In this document, we will overview each one and how they can be applied to your characters.
 
