#UE4doc #docs:
Removed Overview header
Added engine version field
Changed the color of an arrow on a picture to follow guidelines
Replaced some images with updated images of the editor's UI

diff --git a/Engine/Documentation/Source/Engine/Animation/AnimationComposite/AnimationComposite.INT.udn b/Engine/Documentation/Source/Engine/Animation/AnimationComposite/AnimationComposite.INT.udn
--- a/Engine/Documentation/Source/Engine/Animation/AnimationComposite/AnimationComposite.INT.udn
+++ b/Engine/Documentation/Source/Engine/Animation/AnimationComposite/AnimationComposite.INT.udn
@@ -2,6 +2,7 @@ Availability:Public
 Title: Animation Composite
 Crumbs: %ROOT%, Engine, Engine/Animation
 Description:Animation Composites serve as a way to combine multiple animations together and treat them as a single unit.
+version: 4.9
 
 [REGION:fullwidth]
 ![](AnimationCompositeHeader.png)
@@ -9,9 +10,9 @@ Description:Animation Composites serve as a way to combine multiple animations t
 
 [TOC(start:2 end:2)]
 
-## Overview
 
-In certain circumstances, you may have a situation in which you need to take several Animation Sequences and "stitch" them together so that they can be used as if they were just one sequence instead of many. That is exactly the purpose of the **Animation Composite**. An Animation Composite is an animation asset designed to allow you to combine multiple animations together and treat them as a single unit. However, be aware that a Composite simply appends animations; it does not provide any blending ability.
+
+In certain circumstances, you may have a situation in which you need to take several Animation Sequences and stitch them together so that they can be used as if they were just one sequence instead of many. That is exactly the purpose of the **Animation Composite**. An Animation Composite is an animation asset designed to allow you to combine multiple animations together and treat them as a single unit. However, be aware that a Composite simply appends animations; it does not provide any blending ability.
 
 Animation Composites can also have their own sets of [Notifies](Engine/Animation/Sequences\Notifies) and [Curves](Engine/Animation/Sequences\Curves). For more information on these topics, please see the associated documentation for each.
 
@@ -43,21 +44,21 @@ Your Composite is now created and ready to edit!
 
 **Double-clicking** a Composite will open Persona in Animation Mode with the Composite selected and ready to edit.
 
-![](AnimationCompositeEdit.png)(w:500)
+![](AnimationCompositeEdit.png)
 
-At the top of the **Asset Editor** panel, you will see the Composite Track.
+At the top of the **Asset Editor** panel, you will see the **Composite Track**.
 
 ![](CompositeTrack.png)(w:600)
 
-To add Animation Sequences to the Composite, simply drag and drop them into the Composite Track. As you drag more Sequences in, they will be appended automatically.
+To add Animation Sequences to the Composite, simply drag and drop them into the **Composite Track**. As you drag more Sequences in, they will be appended automatically.
 
 ![](AddingAnimationsToComposite.png)(w:700)
 
 [REGION:note]
-As you drag an asset out of the **Asset Browser**, you may notice the Asset Editor gray out and warn that you are **currently previewing another asset**. This is expected behavior. Simply continue dragging the Animation Sequence onto the Composite track.
+As you drag an asset out of the **Asset Browser**, you may notice the **Asset Editor** gray out and warn that you are "Currently previewing another asset". This is expected behavior. Simply continue dragging the Animation Sequence onto the **Composite Track**.
 [/REGION]
 
-You can move Sequences along the Composite track by dragging them, and remove them from the track by **right-clicking** them and choosing **Delete Segment** from the context menu.
+You can move Sequences along the **Composite Track** by dragging them, and remove them from the track by **Right-clicking** them and choosing **Delete Segment** from the context menu.
 
 ![](DeleteSegment.png)
 
