#UE4doc #docs - Updated images to reflect latest UI and updated doc to reflect changes in wording where needed. Addresses UEDOC-1029.

diff --git a/Engine/Documentation/Source/Engine/Physics/Collision/Collision.INT.udn b/Engine/Documentation/Source/Engine/Physics/Collision/Collision.INT.udn
--- a/Engine/Documentation/Source/Engine/Physics/Collision/Collision.INT.udn
+++ b/Engine/Documentation/Source/Engine/Physics/Collision/Collision.INT.udn
@@ -50,8 +50,9 @@ Mesh editor. Complex Collision is per-poly, and it is never used for an Actor th
 | **Collision Presets** | These are a collection of presets for the **Collision Responses** below. There is an INI file you can edit to change or add profiles. **To edit anything below this point, you need to set this to "Custom..."**|
 | **Collision Enabled** | The 3 settings in here are as follows: |
 | ^ | _No Collision_ - No collision is performed against this object, neither trace nor physics. |
-| ^ | _No Physics Collision_ - This body is only used for raycasts, sweeps, and overlaps. |
-| ^ | _Collision Enabled_ - This body is used for physics simulations and collision queries. |
+| ^ | _Query Only (No Physics Collision)_ - This body is only used for raycasts, sweeps, and overlaps. |
+| ^ | _Physics Only (No Query Collision)_ - This body is used only for physics collision. |
+| ^ | _Collision Enabled (Query and Physics)_ - This body is used for physics simulations and collision queries. |
 | **Object Type** | This defines what the object should be considered when it moves. It has more to do with the interactions you see below than it does with defining an object as static or dynamic. |
 | [REGION:tablesection] &nbsp; [/REGION]||
 | **Collision Responses** | This sets all the collision responses below to the value checked. So if you check "Block" here, all of the rows below will also check "Block." |
@@ -72,6 +73,8 @@ Mesh editor. Complex Collision is per-poly, and it is never used for an Actor th
 | **Check Async Scene on Move** | If this is flagged as **true**, the component will look for collisions in both of the physics scenes (synchronous and asynchronous). The asynchronous scene is primarily used by the broken pieces of destructible Actors. |
 | **Trace Complex on Move** | If this is set to **true**, component sweeps with this component will trace against complex collision during movement. Complex Collision is, in short, per face collision, where as Simple Collision are your spheres, capsules, boxes, and generated convex shapes. |
 | **Return Material on Move ** | Setting this to **true** will return the Physics Material in the **Hit Info**. |
+| **Can Ever Affect Navigation ** | Setting this to **true** can potentially influence navigation. |
+
 
 Some components will have additional or fewer options in the collision category, and many times these very specific options will only be shown in the appropriate editor 
 for that component.
@@ -133,7 +136,7 @@ By setting both of their collision settings to block each other, you get a colli
 
 | Sphere Collision Setup| Wall Collision Setup |
 | --- | ---- |
-| ![](COL_collideNoEvent_Sphere.png)(w:300) %spacer% |![](COL_collideNoEvent_Box.png)(w:301) %spacer% |
+| ![](COL_collideNoEvent_Sphere.png)(w:300)  |![](COL_collideNoEvent_Box.png)(w:301)  |
 | In this case, the sphere is a **PhysicsBody** and it is set to `block` **WorldDynamic** (which is what the wall is). | The wall is a **WorldDynamic** and is set to `block` **PhysicsBody** Actors (which is what the sphere is). |
 
 In this case, the sphere and the wall will simply collide; no further notifications of the collision will take place.
@@ -147,7 +150,7 @@ code can be triggered:
 
 | Sphere Collision Setup| Wall Collision Setup |
 | --- | ---- |
-| ![](COL_collideEvent_Sphere.png)(w:300) %spacer%|![](COL_collideNoEvent_Box.png)(w:301) %spacer%|
+| ![](COL_collideEvent_Sphere.png)(w:300) |![](COL_collideNoEvent_Box.png)(w:301)|
 | As in the example above, the sphere is a **PhysicsBody** and it is set to `block` **WorldDynamic** (which is what the wall is). However, the sphere has also enabled Simulation Generates Hit Events so it will trigger an event for itself whenever it collides with something. | The wall is a **WorldDynamic** and is set to `block` **PhysicsBody** Actors (which is what the sphere is). Since the wall is not set to Simulation Generates Hit Events, it will not generate an event for itself. |
 
 With the sphere set to Simulation Generates Hit Events, the sphere will tell itself that it has had a collision. It will fire off events such as **ReceiveHit** or 
@@ -168,14 +171,14 @@ overlap or ignore the box:
 
 | Sphere Collision Setup| Wall Collision Setup |
 | --- | ---- |
-| ![](COL_OverlapNoEvent_Sphere.png)(w:300) %spacer% |![](COL_collideNoEvent_Box.png)(w:303) %spacer% |
+| ![](COL_OverlapNoEvent_Sphere.png)(w:300)  |![](COL_collideNoEvent_Box.png)(w:303) |
 | Here the sphere is set to `overlap` **WorldDynamic** Actors (like our wall), but it does not have Generate Overlap Events enabled. As far as the sphere is concerned, it has not collided or overlapped anything, effectively it has ignored the wall.  | The wall is a **WorldDynamic** and is set to `block` **PhysicsBody** Actors (which is what the sphere is). As stated above, both Actors need to be set to block each other's respective object types. If they do not, they will not collide. |
 
 Or:
 
 | Sphere Collision Setup| Wall Collision Setup |
 | --- | ---- |
-| ![](COL_ignore_sphere.png)(w:300) %spacer% |![](COL_collideNoEvent_Box.png)(w:303) %spacer% |
+| ![](COL_ignore_sphere.png)(w:300) |![](COL_collideNoEvent_Box.png)(w:303) |
 | Here the sphere is set to `ignore` **WorldDynamic** Actors (like our wall), and it will pass through the wall. | The wall is a **WorldDynamic** and is set to `block` **PhysicsBody** Actors (which is what the sphere is). As stated above, both Actors need to be set to block each other's respective object types. If they do not, they will not collide. |
 
 
@@ -188,7 +191,7 @@ Unlike collisions that can fire every frame, the overlap events are **BeginOverl
 
 | Sphere Collision Setup| Wall Collision Setup |
 | --- | ---- |
-| ![](SphereCollisiionSetup.png) %spacer% | ![](WallCollisionSetup.png) %spacer% |
+| ![](SphereCollisiionSetup.png) | ![](WallCollisionSetup.png) |
 | Here the sphere is set to `overlap` **WorldDynamic** Actors (like our wall), and it will generate an event for itself when it does overlap something. | The wall is a **WorldDynamic** and is set to `block` **PhysicsBody** Actors (which is what the sphere is). As stated above, both Actors need to be set to block each other's respective object types. If they do not, they will not collide. But, an **Overlap** does occur here, and an event is fired. |
 
  
