#UE4doc #docs - Updated content example page to make reference to shadow accuracy actually being shadow bias. User had pointed this out on AnswerHub.

diff --git a/Engine/Documentation/Source/Resources/ContentExamples/Lighting/6_1/6_1.INT.udn b/Engine/Documentation/Source/Resources/ContentExamples/Lighting/6_1/6_1.INT.udn
--- a/Engine/Documentation/Source/Resources/ContentExamples/Lighting/6_1/6_1.INT.udn
+++ b/Engine/Documentation/Source/Resources/ContentExamples/Lighting/6_1/6_1.INT.udn
@@ -55,7 +55,7 @@ Related: Engine/Rendering/LightingAndShadows/LightMobility
 [/REGION]
 
 
-This example shows off the Self Shadowing Accuracy. This property exists primarily to help objects feel more grounded along a surface. There is a tradeoff, however, in that increasing the value to make an object appear more grounded can sometimes lead to self-shadowing artifacts, as seen in the image below:
+This example shows off the Self Shadowing Accuracy, or **Shadow Bias**. This property exists primarily to help objects feel more grounded along a surface. There is a tradeoff, however, in that increasing the value to make an object appear more grounded can sometimes lead to self-shadowing artifacts, as seen in the image below:
 
 ![](SelfShadowingArtifacts.png)
 
