UEDOC-485

Gifs on this page have been converted into Youtube videos.

diff --git a/Engine/Documentation/Source/Engine/Animation/RootMotion/RootMotion.INT.udn b/Engine/Documentation/Source/Engine/Animation/RootMotion/RootMotion.INT.udn
--- a/Engine/Documentation/Source/Engine/Animation/RootMotion/RootMotion.INT.udn
+++ b/Engine/Documentation/Source/Engine/Animation/RootMotion/RootMotion.INT.udn
@@ -27,11 +27,18 @@ Early versions of the engine (Versions previous to 4.5) required the use of a An
 Since version 4.6, Root Motion is now handled per animation and can be toggled on/off in the animation's properties in Persona. To fully utilize Root Motion, your character will need an Animation Blueprint (More information can be found in the [Animation Blueprint documentation](https://docs.unrealengine.com/latest/INT/Engine/Animation/AnimBlueprints/index.html) ). Within the characters Animation Blueprint properties, the **Root Motion** section allows you to set which type of animations will have Root Motion applied to them. Please note that at the time of writing that you still need to use an AnimMontage for your Root Motion to be applied successfully. 
 [/REGION]
 
-[REGION:lightbox]
-[![](HammerSlamStatic.png)](HammerSlam.gif)
-[/REGION]
+[OBJECT:EmbeddedVideo]
+[PARAMLITERAL:width]
+640
+[/PARAMLITERAL]
+[PARAMLITERAL:height]
+360
+[/PARAMLITERAL]
+[PARAMLITERAL:videoid]
+TKAcNubFOH8
+[/PARAMLITERAL]
+[/OBJECT]
 
-**Click to watch playback**
 
 [REGION:tip]
 Another example of **Root Motion** applied to a character can be found on the [Animation Content Examples](Resources\ContentExamples\Animation) page under section 1.9.
@@ -71,22 +78,36 @@ Below is a breakdown of the problems with not handling root motion, and the bene
 
 Without root motion, such animations would pull the character away from his capsule, as seen below. Note the unsightly snap back to the capsule location. This happens because once the animation is done playing, the character controller pulls the character back to their original position.
 
-[REGION:lightbox]
-[![](HammerSlamStatic_NoRootMotion.png)](HammerSlam_NoRootMotion.gif)
-[/REGION]
+[OBJECT:EmbeddedVideo]
+[PARAMLITERAL:width]
+640
+[/PARAMLITERAL]
+[PARAMLITERAL:height]
+360
+[/PARAMLITERAL]
+[PARAMLITERAL:videoid]
+Xu2bVQ4pg8M
+[/PARAMLITERAL]
+[/OBJECT]
 
-**Click to watch playback**
 
 ### Why Capsule Separation is Bad
 
 And the snap-back problem is not the only issue to fight in this case. Since the character is leaving the collision shape, it will pass through world objects and break overall continuity. In this example, we can see the character's attack makes him go through a wall and then jump back.
 
 
-[REGION:lightbox]
-[![](HammerSlamStatic_NoRootMotion_Wall.png)](HammerSlam_NoRootMotion_Wall.gif)
-[/REGION]
+[OBJECT:EmbeddedVideo]
+[PARAMLITERAL:width]
+640
+[/PARAMLITERAL]
+[PARAMLITERAL:height]
+360
+[/PARAMLITERAL]
+[PARAMLITERAL:videoid]
+ov9pyx4MAOo
+[/PARAMLITERAL]
+[/OBJECT]
 
-**Click to watch playback**
 
 ### How Root Motion Handling Solves the Problem
 
@@ -94,11 +115,17 @@ Once your animation is set up to utilize root motion, then the movement that was
 
 
 
-[REGION:lightbox]
-[![](HammerSlamStatic_RootMotion_NoWall.png)](HammerSlam_RootMotion.gif)
-[/REGION]
-
-**Click to watch playback**
+[OBJECT:EmbeddedVideo]
+[PARAMLITERAL:width]
+640
+[/PARAMLITERAL]
+[PARAMLITERAL:height]
+360
+[/PARAMLITERAL]
+[PARAMLITERAL:videoid]
+PWB_mqjz3iA
+[/PARAMLITERAL]
+[/OBJECT]
 
 
 ### Root Motion and Physics
@@ -107,11 +134,17 @@ Since the capsule comes along, this means that physics collisions can still be u
 
 You will note that it is not quite perfect, as the character's bending animation causes it to pass somewhat through the walls. However, this could be easily managed by either working with collision volumes on either the wall or the character. The important point is that the capsule is keeping up with the motion and preventing the character from passing all the way through the wall and having to snap back.
 
-[REGION:lightbox]
-[![](HammerSlamStatic_RootMotion.png)](HammerAnimation_RootMotion_Wall.gif)
-[/REGION]
-
-**Click to watch playback**
+[OBJECT:EmbeddedVideo]
+[PARAMLITERAL:width]
+640
+[/PARAMLITERAL]
+[PARAMLITERAL:height]
+360
+[/PARAMLITERAL]
+[PARAMLITERAL:videoid]
+YuD29-Xr7Oc
+[/PARAMLITERAL]
+[/OBJECT]
 
 
 ## Setting Up Root Motion
@@ -199,11 +232,17 @@ We are almost finished! Finally, we just need to make sure that our slot is bein
 That is all there is to it! Now that your animation is playing back via an AnimMontage, with that Montage being set to utilize Root Motion, everything will just work!
 
 
-[REGION:lightbox]
-[![](HammerSlamStatic_RootMotion_NoWall.png)](HammerSlam_RootMotion.gif)
-[/REGION]
-
-**Click to watch playback**
+[OBJECT:EmbeddedVideo]
+[PARAMLITERAL:width]
+640
+[/PARAMLITERAL]
+[PARAMLITERAL:height]
+360
+[/PARAMLITERAL]
+[PARAMLITERAL:videoid]
+PWB_mqjz3iA
+[/PARAMLITERAL]
+[/OBJECT]
 
 
 ## Caveats and Considerations
