UEDOC-1351

Updated out of date code on this doc

diff --git a/Engine/Documentation/Source/Programming/UnrealArchitecture/Timers/GameplayTimers.INT.udn b/Engine/Documentation/Source/Programming/UnrealArchitecture/Timers/GameplayTimers.INT.udn
--- a/Engine/Documentation/Source/Programming/UnrealArchitecture/Timers/GameplayTimers.INT.udn
+++ b/Engine/Documentation/Source/Programming/UnrealArchitecture/Timers/GameplayTimers.INT.udn
@@ -104,7 +104,7 @@ Each of the timer-manipulation functions has 3 flavors of overloads for interact
 * Taking any type of non-dynamic delegate you can muster.  It is a bit muddier syntactically but far more flexible.  e.g. 
 
 		GetWorldTimerManager().SetTimer( FTimerDelegate::CreateUObject(this, &AMyClass::MyFunc), 5.0f, TRUE );
-		GetWorldTimerManager().SetTimer( FTimerDelegate::CreateRaw_ThreeVars(&MyStaticFunc, MyPassThruArg, MyOtherArg, YesOneMoreArg), 5.0f, TRUE );
+		GetWorldTimerManager().SetTimer( FTimerDelegate::CreateRaw(&MyStaticFunc, MyPassThruArg, MyOtherArg, YesOneMoreArg), 5.0f, TRUE );
                 
 		FTimerDelegate MyDel;
 		Del.BindRaw(this, &AMyClass::MyFunc2);
