Fixed to match blog version.

diff --git a/Engine/Documentation/Source/Gameplay/DataDriven/DataDrivenGameplay.INT.udn b/Engine/Documentation/Source/Gameplay/DataDriven/DataDrivenGameplay.INT.udn
--- a/Engine/Documentation/Source/Gameplay/DataDriven/DataDrivenGameplay.INT.udn
+++ b/Engine/Documentation/Source/Gameplay/DataDriven/DataDrivenGameplay.INT.udn
@@ -44,53 +44,39 @@ As an example, here is a definition for level up experience data and the CSV doc
 	USTRUCT(BlueprintType)
 	struct FLevelUpData : public FTableRowBase
 	{
-			GENERATED_USTRUCT_BODY()
+		GENERATED_USTRUCT_BODY()
 
 	public:
 
-			FLevelUpData()
-							: XPtoLvl(0)
-							, XP(0)
-			{}
+		FLevelUpData()
+		: XPtoLvl(0)
+		, AdditionalHP(0)
+		{}
 
-			/** The 'Name' column is the same as the XP Level */
+		/** The 'Name' column is the same as the XP Level */
 
-			/** XP to get to the given level from the previous level */
-			UPROPERTY(EditAnywhere, BlueprintReadWrite, Category=LevelUp)
-			int32 XPtoLvl;
+		/** XP to get to the given level from the previous level */
+		UPROPERTY(EditAnywhere, BlueprintReadWrite, Category=LevelUp)
+		int32 XPtoLvl;
 
-			/** This was the old property name (represented total XP) */
-			UPROPERTY(EditAnywhere, BlueprintReadWrite, Category=LevelUp)
-			int32 XP;
+		/** Extra HitPoints gained at this level */
+		UPROPERTY(EditAnywhere, BlueprintReadWrite, Category=LevelUp)
+		int32 AdditionalHP;
 
-			UPROPERTY(EditAnywhere, BlueprintReadWrite, Category=LevelUp)
-			TAssetPtr<UTexture> Asset;
+		/** Icon to use for Achivement */
+		UPROPERTY(EditAnywhere, BlueprintReadWrite, Category=LevelUp)
+		TAssetPtr<UTexture> AchievementIcon;
 	};
 
 
 **CSV Document:**
 
-	Name,XPtoLvl,XP,Asset
-	1,0,0,"Texture2d'/Game/UI/HUD/Actions/Barrel'"
-	2,1000,1000,"Texture2d'/Game/UI/HUD/Actions/Barrel'"
-	3,1000,2000,"Texture2d'/Game/UI/HUD/Actions/Barrel'"
-	4,1500,3500,"Texture2d'/Game/UI/HUD/Actions/Barrel'"
-	5,2000,5500,"Texture2d'/Game/UI/HUD/Actions/Barrel'"
-	6,2500,8000,"Texture2d'/Game/UI/HUD/Actions/Barrel'"
-	7,3000,11000,"Texture2d'/Game/UI/HUD/Actions/Barrel'"
-	8,3500,14500,"Texture2d'/Game/UI/HUD/Actions/Barrel'"
-	9,4000,18500,"Texture2d'/Game/UI/HUD/Actions/Barrel'"
-	10,4500,23000,"Texture2d'/Game/UI/HUD/Actions/Barrel'"
-	11,5000,28000,"Texture2d'/Game/UI/HUD/Actions/Barrel'"
-	12,5500,33500,"Texture2d'/Game/UI/HUD/Actions/Barrel'"
-	13,6000,39500,"Texture2d'/Game/UI/HUD/Actions/Barrel'"
-	14,6500,46000,"Texture2d'/Game/UI/HUD/Actions/Barrel'"
-	15,7000,53000,"Texture2d'/Game/UI/HUD/Actions/Barrel'"
-	16,7500,60500,"Texture2d'/Game/UI/HUD/Actions/Barrel'"
-	17,8000,68500,"Texture2d'/Game/UI/HUD/Actions/Barrel'"
-	18,8500,77000,"Texture2d'/Game/UI/HUD/Actions/Barrel'"
-	19,9000,86000,"Texture2d'/Game/UI/HUD/Actions/Barrel'"
-	20,9500,95500,"Texture2d'/Game/UI/HUD/Actions/Barrel'"
+	Name,XPtoLvl,AdditionalHP,AchievementIcon
+	1,0,0,"Texture2d'/Game/Textures/AchievementIcon1'"
+	2,1000,9,"Texture2d'/Game/Textures/AchievementIcon2'"
+	3,1000,10,"Texture2D'/Game/Textures/AchievementIcon3'"
+	4,1500,12,"Texture2D'/Game/Textures/AchievementIcon4'"
+	5,2000,14,"Texture2D'/Game/Textures/AchievementIcon5'"
 
 [REGION:tip]
 The double quotes around the asset type are important for the property importing pipeline. Without them, the text is imported as Texture2d'. 
@@ -112,9 +98,9 @@ Here is an example table for damage progression:
 |                      | 0   | 1   | 2   | 3   |
 | -------------------- | --- | --- | --- | --- |
 | Melee_Damage         | 15  | 20  | 25  | 30  |
-| Melee_KnockBack      | 0   | 0   | 0   | 0   |
-| Melee_KnockBackAngle | 0   | 0   | 0   | 0   |
-| Melee_StunTime       | 0   | 0   | 0   | 0   |
+| Melee_KnockBack      | 1   | 2   | 4   | 8   |
+| Melee_KnockBackAngle | 10  | 45  | 60  | 65  |
+| Melee_StunTime       | 0   | 1   | 5   | 7   |
 [/REGION]
 
 ## Importing Process
