#UE4doc #docs - Fixed images colors to match style guide and adjusted bold letting to match guide.

diff --git a/Engine/Documentation/Source/Engine/Animation/Tools/VertexAnimationTool/VAT_KF_Meshes/VAT_KF_Meshes.INT.udn b/Engine/Documentation/Source/Engine/Animation/Tools/VertexAnimationTool/VAT_KF_Meshes/VAT_KF_Meshes.INT.udn
--- a/Engine/Documentation/Source/Engine/Animation/Tools/VertexAnimationTool/VAT_KF_Meshes/VAT_KF_Meshes.INT.udn
+++ b/Engine/Documentation/Source/Engine/Animation/Tools/VertexAnimationTool/VAT_KF_Meshes/VAT_KF_Meshes.INT.udn
@@ -6,7 +6,7 @@ Version: 4.8
 
 [TOC(start:2 end:3)]
 
-## Overview
+
 You can also use the Vertex Animation tool with meshes that do not make use of the Animation Timeline in 3Ds Max. 
 Generally speaking mesh animations that use this method have been exported frame by frame from another 3D package like  Maya or Blender.
 In this section we will go over how to use this method as well as any differences in setup that you need to be aware of.
@@ -43,8 +43,8 @@ Make sure the **Object** is set to **Copy** and then set the **Number of Copies*
 	[/PARAMLITERAL]
 	[/OBJECT]
 
-1. Then select the 2nd Teapot mesh and set the **Melt Amount** in the Melt modifier to 25. 
-Then select the next Teapot in line and set the **Melt Amount** of the Melt modifier to 50 and the next one to 75, and finally set the last Teapot **Melt Amount** to 100.
+1. Then select the 2nd Teapot mesh and set the **Melt Amount** in the Melt modifier to **25**. 
+Then select the next Teapot in line and set the **Melt Amount** of the Melt modifier to **50** and the next one to **75**, and finally set the last Teapot **Melt Amount** to **100**.
 
 	[OBJECT:EmbeddedVideo]
 	[PARAMLITERAL:width]
@@ -58,7 +58,7 @@ Then select the next Teapot in line and set the **Melt Amount** of the Melt modi
 	[/PARAMLITERAL]
 	[/OBJECT]
 
-1. With all the Teapots now having some type of deformation on them, select all the Teapots and then with the mouse Right-Click  in the 3Ds Max viewport and select **Convert To:** > **Convert to Editable Poly**.
+1. With all the Teapots now having some type of deformation on them, select all the Teapots and then with the mouse **Right-Click** in the 3Ds Max viewport and select **Convert To:** > **Convert to Editable Poly**.
 
 	[region:note]
 	You have to convert any mesh that you want to use with this tool to Editable Poly. The tool **WILL NOT** work with any other mesh type.
@@ -109,21 +109,21 @@ Once the Teapot mesh has been exported from 3Ds Max, it is now time for the Teap
 	
 		![](VAT_SM_Import_Options.png)
 
-1. After the Teapot mesh has been imported, open the mesh up by Double - Clicking on it with the Left Mouse button in Content Browser. 
+1. After the Teapot mesh has been imported, open the mesh up by **Double-Clicking** on it with the **Left Mouse button** in Content Browser. 
 Once the mesh has been opened enable the **Use Full Precision UVs** option and set the **Distance Field Resolution Scale** to **0.0** then press the **Apply Changes** button.
 
 	![](VAT_KFM_UE4_Setup.png)
 
-1. With the Teapot mesh now set up correctly, create a new Material in the Content Browser and name it **MAT_Sequance_Painter** and then open up the Material. 
+1. With the Teapot mesh now set up correctly, create a new **Material** in the Content Browser and name it **MAT_Sequance_Painter** and then open up the Material. 
 
 	![](VAT_KFM_Create_Material.png)
 
-1. Once the Material is opened up, add Two Constant Parameter nodes to the Material Graph by holding down the one key on the Keyboard and Left - Clicking with the mouse in the Material Graph to place the Constant Parameter nodes.
+1. Once the Material is opened up, add two **Constant Parameter** nodes to the Material Graph by holding down the **1** key on the keyboard and **Left-Clicking** with the mouse in the Material Graph to place the Constant Parameter nodes.
 Set the value of one of the Constant Parameter node and plug that into the **Base Color** input on the Main Material Node. Then set the value of the other Constant Parameter node to 5.
 
 	![](VAT_KFM_Set_Values.png)
 
-1. Next use the mouse and Right - Click in the Material Graph to expose the search menu and then search for **MS_SequencePainter_SequenceFlipbook** Material Function using MS as the search term.
+1. Next use the mouse and **Right-Click** in the Material Graph to expose the search menu and then search for **MS_SequencePainter_SequenceFlipbook** Material Function using MS as the search term.
 When you find the MS_SequencePainter_SequenceFlipbook Material Function, click on it with the Left Mouse button to add it to the Material Graph.
 
 	[OBJECT:EmbeddedVideo]
@@ -157,7 +157,7 @@ After that is completed you can then close down the Material Graph Editor.
 
 	![](VAT_Compile_Save.png)
 
-1. With the Material created, compiled, and saved drag the Teapot mesh from the Content Browser into the level and then drag the Material that was just created for it from the Content Browser on to the mesh that was placed in the level. 
+1. With the Material created, compiled, and saved drag the Teapot mesh from the Content Browser into the level and then drag the **Material** that was just created for it from the Content Browser on to the mesh that was placed in the level. 
 Once this has been completed you will now see the mesh "playing" the animation.
 
 	[OBJECT:EmbeddedVideo]
