#UE4doc #docs - Fixed images colors to match style guide and adjusted bold letting to match guide.

diff --git a/Engine/Documentation/Source/Engine/Animation/Tools/VertexAnimationTool/VAT_KF_Meshes/VAT_KF_Meshes.INT.udn b/Engine/Documentation/Source/Engine/Animation/Tools/VertexAnimationTool/VAT_KF_Meshes/VAT_KF_Meshes.INT.udn
--- a/Engine/Documentation/Source/Engine/Animation/Tools/VertexAnimationTool/VAT_KF_Meshes/VAT_KF_Meshes.INT.udn
+++ b/Engine/Documentation/Source/Engine/Animation/Tools/VertexAnimationTool/VAT_KF_Meshes/VAT_KF_Meshes.INT.udn
@@ -2,7 +2,7 @@ Availability:Public
 Title:Vertex Animation Tool - Key Framed Meshes
 Crumbs: %ROOT%, Engine, Engine/Animation
 Description: User guide for the 3Ds Max Vertex Animation tool set.
-Version: 4.8
+Version: 4.9
 
 [TOC(start:2 end:3)]
 
diff --git a/Engine/Documentation/Source/Engine/Animation/Tools/VertexAnimationTool/VAT_TL_Meshes/VAT_TL_Meshes.INT.udn b/Engine/Documentation/Source/Engine/Animation/Tools/VertexAnimationTool/VAT_TL_Meshes/VAT_TL_Meshes.INT.udn
--- a/Engine/Documentation/Source/Engine/Animation/Tools/VertexAnimationTool/VAT_TL_Meshes/VAT_TL_Meshes.INT.udn
+++ b/Engine/Documentation/Source/Engine/Animation/Tools/VertexAnimationTool/VAT_TL_Meshes/VAT_TL_Meshes.INT.udn
@@ -2,11 +2,10 @@ Availability:Public
 Title:Vertex Animation Tool - Timeline Meshes
 Crumbs: %ROOT%, Engine, Engine/Animation
 Description: User guide for the 3Ds Max Vertex Animation tool set.
-Version: 4.8
+Version: 4.9
 
 [TOC(start:2 end:3)]
 
-## Overview
 In the following section we will cover how you can use the 3Ds Max Animation Timeline in conjunction with the Vertex Animation tools to create a 2D Texture that stores the position of vertices in any given 3D mesh.
 
 ## 3Ds Max Content Creation & Setup for Timeline Animated Meshes
@@ -49,7 +48,7 @@ a **.EXR** Texture which contains the vertices position data and a **.BMP** Text
 	![](VAT_Output_Textures.png)	
 
 1. With the Textures now created go back to 3Ds Max and select the Teapot mesh with "**_MorphExport**" added to the name and export it as a .FBX file to the same folder that the Textures where exported to.
-Name the mesh "**SM_Teapot_Bend_00**" and in the FBX options make sure that Smoothing Groups are enabled and leave all other setting with the default values.
+Name the mesh "**SM_Teapot_Bend_00**" and in the FBX options make sure that **Smoothing Groups** are enabled and leave all other setting with the default values.
 
 	[OBJECT:EmbeddedVideo]
 	[PARAMLITERAL:width]
@@ -71,7 +70,7 @@ Please keep in mind that each and every step listed below **MUST BE** completed
 ### Texture Import & Setup
 
 [region:tip]
-If you are unfamiliar with the UE4 Texture Import process please read the [](Engine/Content/Types/Textures/Importing) documentation for more information.
+If you are unfamiliar with the UE4 Texture Import process please read the **[](Engine/Content/Types/Textures/Importing)** documentation for more information.
 [/region]
 
 The first thing that needs to be completed is that both Textures that were exported from 3Ds Max need to be imported into UE4. 
@@ -93,7 +92,7 @@ yqqHsxrPk9o
 
 Once each Texture has been imported, they then need to be setup in a very specific manner to ensure that they work correctly.
 The first Texture to setup is the .EXR Texture that contains the vertices position data.
-Open up this Texture by Double- Clicking on it in the Content Browser with the Left Mouse Button and make sure the following options are set / enabled.
+Open up this Texture by **Double-Clicking** on it in the Content Browser with the **Left Mouse Button** and make sure the following options are set / enabled.
 
 * Make sure the Mip Gen Settings is set to **No Mip Maps**.
 * Change the Texture Group to **UI**.
@@ -121,7 +120,7 @@ Import the Teapot Static Mesh in the same way the Texture were imported by first
 When you do this the following dialogue box will be displayed. From this dialogue make sure that the following options are set / disabled.
 
 [region:tip]
-If you are unfamiliar with how to import Static Meshes into UE4 please check out the [](Engine/Content/FBX/StaticMeshes) document for more information about this.
+If you are unfamiliar with how to import Static Meshes into UE4 please check out the **[](Engine/Content/FBX/StaticMeshes)** document for more information about this.
 [/region]
 
 * Disable the **Auto Generate Collision** option.
@@ -145,15 +144,15 @@ With the Static Mesh now set up correctly, make sure to press the **Apply Change
 
 ### UE4 Material Setup & Use
 
-Now that all of the required content has been imported and properly setup, it is now time to create a Material that can use the content.
+Now that all of the required content has been imported and properly setup, it is now time to create a that can Material use the content.
 
 [region:tip]
-If you are unfamiliar with how Materials work inside UE4 please refer to the [](Engine/Rendering/Materials) documentation for more information.
+If you are unfamiliar with how Materials work inside UE4 please refer to the **[](Engine/Rendering/Materials)** documentation for more information.
 [/region]
 
-1. First create a new Material by Right-Clicking with the Mouse and then choosing the **Material** option that is under the **Create Basic Asset** section. Name the Material **MAT_Vertex_Animation**.
+1. First create a new **Materia**l by **Right-Clicking** with the Mouse and then choosing the **Material** option that is under the **Create Basic Asset** section. Name the Material **MAT_Vertex_Animation**.
 
-1. Next open up the **MAT_Vertex_Animation** Material and locate the two Textures that were imported into the Content Browser in the previous step. Select both Textures and then Drag them from the Content Browser and Drop them into the Material Graph.  
+1. Next open up the **MAT_Vertex_Animation** Material and locate the two Textures that were imported into the Content Browser in the previous step. Select both **Texture**s and then **Drag** them from the Content Browser and Drop them into the Material Graph.  
 	
 	[OBJECT:EmbeddedVideo]
 	[PARAMLITERAL:width]
@@ -175,11 +174,11 @@ If you are unfamiliar with how Materials work inside UE4 please refer to the [](
 
 	![](VAT_UnCheck_Tan_Normals.png)
 
-1. Then hold down the One key on the Keyboard and then click inside the Material Graph with the Left Mouse Button to place down a Constant Parameter node. Set the value of the Constant to **0.5** and then plug the output into the **BaseColor** input in the Main Material Node.
+1. Then hold down the **1** key on the Keyboard and then click inside the Material Graph with the **Left Mouse Button** to place down a Constant Parameter node. Set the value of the Constant to **0.5** and then plug the output into the **BaseColor** input in the Main Material Node.
 	
 	![](VAT_Add_Const.png)
 
-1. Now it is time to add the Material Function that powers the effect. To do this Right-Click with the Mouse inside the Material Graph and then from the menu that is displayed, type **MS** to show the Maxscript Material Functions. Then click on the **MS_VertexAnimationTools_MorphTargets** Material Function with the mouse to add it to the Material Graph. 
+1. Now it is time to add the Material Function that powers the effect. To do this **Right-Click** with the Mouse inside the Material Graph and then from the menu that is displayed, type **MS** to show the Maxscript Material Functions. Then click on the **MS_VertexAnimationTools_MorphTargets** Material Function with the mouse to add it to the Material Graph. 
 
 	[OBJECT:EmbeddedVideo]
 	[PARAMLITERAL:width]
@@ -193,7 +192,7 @@ If you are unfamiliar with how Materials work inside UE4 please refer to the [](
 	[/PARAMLITERAL]
 	[/OBJECT]
 
-1. With all the required Textures and Material Function now added to the Material Graph, we can now begin to connect everything together but first, select both Texture Samplers and convert them to **Texture Objects** by Right-Clicking on one of the Textures and selecting the **Convert To Texture Object** option.
+1. With all the required Textures and Material Function now added to the Material Graph, we can now begin to connect everything together but first, select both Texture Samplers and convert them to **Texture Objects** by **Right-Clicking** on one of the Textures and selecting the **Convert To Texture Object** option.
 	
 	[OBJECT:EmbeddedVideo]
 	[PARAMLITERAL:width]
@@ -211,7 +210,7 @@ If you are unfamiliar with how Materials work inside UE4 please refer to the [](
 
 	![](VAT_Wired_Up.png)
 
-1. Now add another Constant Parameter to the Material Graph by holding down the one key on the Keyboard and then clicking in the Material Graph with the Left Mouse Button to place the node in the Material Graph. Once the Constant Node has been added set its value to 51 and connect it to the **Number of Morph Targets** input on the **MS_VertexAnimationTools_MorphTargets** Material Function.
+1. Now add another Constant Parameter to the Material Graph by holding down the **1** key on the Keyboard and then clicking in the Material Graph with the **Left Mouse Button** to place the node in the Material Graph. Once the Constant Node has been added set its value to 51 and connect it to the **Number of Morph Targets** input on the **MS_VertexAnimationTools_MorphTargets** Material Function.
 
 	![](VAT_Num_Morph_Targets.png)
 
diff --git a/Engine/Documentation/Source/Engine/Animation/Tools/VertexAnimationTool/VertexAnimationTools.INT.udn b/Engine/Documentation/Source/Engine/Animation/Tools/VertexAnimationTool/VertexAnimationTools.INT.udn
--- a/Engine/Documentation/Source/Engine/Animation/Tools/VertexAnimationTool/VertexAnimationTools.INT.udn
+++ b/Engine/Documentation/Source/Engine/Animation/Tools/VertexAnimationTool/VertexAnimationTools.INT.udn
@@ -4,7 +4,7 @@ Crumbs: %ROOT%, Engine, Engine/Animation
 Description: User guide for the 3D Max Vertex Animation tool set.
 Related: Engine\Animation\Tools\VertexAnimationTool\VAT_KF_Meshes
 Related: Engine\Animation\Tools\VertexAnimationTool\VAT_TL_Meshes
-Version: 4.8
+Version: 4.9
 
 [TOC(start:2 end:3)]
 
