Corrections made per forum suggestion

diff --git a/Engine/Documentation/Source/Engine/Rendering/Materials/HowTo/Making_Functions/Making_Functions.INT.udn b/Engine/Documentation/Source/Engine/Rendering/Materials/HowTo/Making_Functions/Making_Functions.INT.udn
--- a/Engine/Documentation/Source/Engine/Rendering/Materials/HowTo/Making_Functions/Making_Functions.INT.udn
+++ b/Engine/Documentation/Source/Engine/Rendering/Materials/HowTo/Making_Functions/Making_Functions.INT.udn
@@ -92,7 +92,7 @@ Creating Material Functions can be accomplished in the few following steps.
 	The Fresnel Effect is the term used to describe how the light you see reflects at different intensities based off of the angle you are viewing it from. For example, if you are standing over a pool looking straight down at the pool, you will not see a lot of reflections in the water. As you start to move your head so that the water in the pool becomes more and more parallel to your eye level, you will begin to notice more and more reflections in the water and this is the type of behavior that we are trying to re-create in our shader, the ability to add or remove visual information based on viewing angle. 
 	[/region]
 	
-1. With the nodes now placed down, it is now time to set them up and wire them together. First, set the Blue Channel of the Constant 3 Vector to **1.0** as this will represent the Z value in our Fresnel term. Then take the output of the Constant 3 Vector node and wire that into the Input of the Vector Transform node. Also make sure that the Transform Node is setup so that we are transforming from Tangent Space into World Space. Then take the output of the Vector Transform node and wire that into the A input of the Dot Product node. Once that has been completed, take the output of the Camera Vector node and wire that into the B input of the Dot Product node. When completed, you should have something that looks like this.
+1. With the nodes placed down, it is now time to set them up and wire them together. First, set the Blue Channel of the Constant 3 Vector to **1.0** as this will represent the Z value in our Fresnel term. Then take the output of the Constant 3 Vector node and wire that into the Input of the Vector Transform node. Also make sure that the Transform Node is setup so that we are transforming from Tangent Space into World Space. Then take the output of the Vector Transform node and wire that into the A input of the Dot Product node. Once that has been completed, take the output of the Camera Vector node and wire that into the B input of the Dot Product node. When completed, you should have something that looks like this.
 
 	![](Material_Function_Nodes_Setup.png)
 
@@ -133,7 +133,7 @@ Creating Material Functions can be accomplished in the few following steps.
 	
 	![](Material_Function_Adding_Multiply.png)
 
-1. Now with all the nodes we need added to the Material Graph, it is now time to wire everything together. First, we need to wire the output of the OneMinus node to the Base input of the Power node. Then we need to connect our first Function Input node, Fresnel Exponent, into the Exp input of the Power node. Once that has been completed, we need to then connect the output of the Power node to the A input of the Multiply node and then connect our second Function Input, Fresnel Intensity, into the B input of the Multiply node. Finally we need to take the output of the Multiply node and connect that to the input of the Output Results node.  
+1. Now with all the nodes we need added to the Material Graph, it is now time to wire everything together. First, we need to wire the output of the OneMinus node to the Base input of the Power node. Then we need to connect our first Function Input node, Fresnel Exponent, into the Exp input of the Power node. Once that has been completed, we need to then connect the output of the Power node to the A input of the Multiply node and then connect our second Function Input, Fresnel Intensity, into the B input of the Multiply node. Finally we need to take the output of the Multiply node and connect that to the input of the Output Result node.  
 	
 	[OBJECT:EmbeddedVideo]
 	[PARAMLITERAL:width]
@@ -147,14 +147,14 @@ Creating Material Functions can be accomplished in the few following steps.
 	[/PARAMLITERAL]
 	[/OBJECT]                                                                                                                                                                                                     
 
-1. Once the Material nodes have all been connected, we need to now compile the Material Function by pressing the **Apply** button on the main tool bar. If the Material Function compiled correctly, the **Apply** button will be grayed out. When completed, you should have something that looks like the following image.
+1. Once the Material nodes have all been connected, we need to compile the Material Function by pressing the **Apply** button on the main tool bar. If the Material Function compiled correctly, the **Apply** button will be grayed out. When completed, you should have something that looks like the following image.
 
 	![](Material_Function_Completed_Material.png)
 
 	
 ## Using Material Functions 
 
-1. With our Material Function now finished, close the Material Function Editor and create a new Material in the **Content Browser**. To create a new Material, **Right-Click** in the **Content Browser** and then select **Material** from **Basic Assets List** and name your newly created Material, **Fresnel_Material**.
+1. With our Material Function finished, close the Material Function Editor and create a new Material in the **Content Browser**. To create a new Material, **Right-Click** in the **Content Browser** and then select **Material** from **Basic Assets List** and name your newly created Material, **Fresnel_Material**.
 
 	[OBJECT:EmbeddedVideo]
 	[PARAMLITERAL:width]
@@ -168,7 +168,7 @@ Creating Material Functions can be accomplished in the few following steps.
 	[/PARAMLITERAL]
 	[/OBJECT]
 
-1. Now with the **Fresnel_Material** created, open it up by **Double Clicking** on it with the **Left Mouse Button** in the **Content Browser**. With the Material now open, let us first add our newly created **Material Function** by dragging it from the **Content Browser** into our material.  
+1. With the **Fresnel_Material** created, open it up by **Double Clicking** on it with the **Left Mouse Button** in the **Content Browser**. With the Material open, let us first add our newly created **Material Function** by dragging it from the **Content Browser** into our material.  
 
 	[OBJECT:EmbeddedVideo]
 	[PARAMLITERAL:width]
