Made change per forum user's request.

diff --git a/Engine/Documentation/Source/Engine/Rendering/Materials/HowTo/Making_Parameters/Making_Parameters.INT.udn b/Engine/Documentation/Source/Engine/Rendering/Materials/HowTo/Making_Parameters/Making_Parameters.INT.udn
--- a/Engine/Documentation/Source/Engine/Rendering/Materials/HowTo/Making_Parameters/Making_Parameters.INT.udn
+++ b/Engine/Documentation/Source/Engine/Rendering/Materials/HowTo/Making_Parameters/Making_Parameters.INT.udn
@@ -171,7 +171,7 @@ Where the real power of Material Parameters comes into play is when you use them
 ## Conclusion 
 
 Material Parameters are a very powerful and extremely useful tool that when combined with Material Instances offer you a way to add an almost endless amount of variation to your Materials. 
-However keep in mind that in order to use Material Parameters you have to use a Material Instances. 
+However keep in mind that in order to use Material Parameters you have to use a Material Instance. 
 
 
 
