#UE4doc #docs:
Removed Overview header
Added engine version field
Changed font settings for multiple terms

diff --git a/Engine/Documentation/Source/Engine/Animation/AnimationRetargeting/AnimationRetargeting.INT.udn b/Engine/Documentation/Source/Engine/Animation/AnimationRetargeting/AnimationRetargeting.INT.udn
--- a/Engine/Documentation/Source/Engine/Animation/AnimationRetargeting/AnimationRetargeting.INT.udn
+++ b/Engine/Documentation/Source/Engine/Animation/AnimationRetargeting/AnimationRetargeting.INT.udn
@@ -2,6 +2,7 @@ Availability:Public
 Title: Animation Retargeting (Same Skeleton)
 Crumbs: %ROOT%, Engine, Engine/Animation
 Description:Retargeting allows you to reuse animations between characters with highly differed proportions.
+version:4.9
 
 [REGION:fullwidth]
 ![](RetargetingHeader.png)
@@ -10,11 +11,11 @@ Description:Retargeting allows you to reuse animations between characters with h
 [TOC(start:2 end:2)]
 
 
-## Overview
+
 **Animation Retargeting (Same Skeleton)** is a feature that allows animations to be reused between characters that share the same Skeleton asset, but have vastly different proportions. Through retargeting, you can prevent your animated skeletons from losing their proportions or becoming unnecessarily deformed when using animations from differently shaped characters.
 
 [REGION:warning]
-In order for different Skeletal Meshes to take advantage of retargeting, they **must** use the same Skeleton asset to ensure Bone Hierarchy, but also the same Joint Rotation in the Bind pose. Please see the [](Engine/Animation/Skeleton) documentation for details.
+In order for different Skeletal Meshes to take advantage of retargeting, they must use the same Skeleton asset to ensure Bone Hierarchy, but also the same Joint Rotation in the Bind pose. Please see the [](Engine/Animation/Skeleton) documentation for details.
 [/REGION]
 
 ### Animation Retargeting for Different Skeletons
@@ -100,9 +101,9 @@ Retargeting can be set up from the **Skeleton Tree** panel in **Persona**.
  
 Typically, you will want to use these settings for bipeds:
 
-* Root Bone, IK Bones, weapon bones, and any kind of markers will use the `Animation` mode.
-* The pelvis will use `AnimationScaled` so it sits at the right height and is still animated. And any other bone you want translation to be animated and retargeted should also use `AnimationScaled`.
-* All other bones will use `Skeleton`. They will use the static translation from the target skeleton.
+* Root Bone, IK Bones, weapon bones, and any kind of markers will use the Animation mode.
+* The pelvis will use AnimationScaled so it sits at the right height and is still animated. And any other bone you want translation to be animated and retargeted should also use AnimationScaled.
+* All other bones will use Skeleton. They will use the static translation from the target skeleton.
 
 This means that if you want to use retargeting, a fast workflow would be:
 
@@ -113,7 +114,7 @@ This means that if you want to use retargeting, a fast workflow would be:
 1. Find the Root bone, any IK bones, any weapon bones, or other marker-style bones and set them to use **Animation**.
 
 [REGION:tip]
-In the `Show` menu from the viewport, you can check `NonRetargeted Animation` and `Bones` to compare the original with the retargeted animation.
+In the **Show** menu from the viewport, you can check **NonRetargeted Animation** and **Bones** to compare the original with the retargeted animation.
 [/REGION]
 
 
@@ -134,7 +135,7 @@ Think of the **Retarget Source Manager** as a list of Skeletal Meshes that you c
 1. In Persona's menu bar, choose **Window > Retarget Source Manager**.
 1. **Right-click** in the list area (center) of the **Retarget Source Manager** panel and choose **Add...** from the context menu.
 1. Choose the Skeletal Mesh for which the special animation was created. You should now see that Skeletal Mesh listed in the **Retarget Source Manager**.
-1. In Persona's **Asset Browser**, **double-click** the Animation Sequence asset for the special animation you created for the Skeletal Mesh.
+1. In Persona's **Asset Browser**, **Double-click** the Animation Sequence asset for the special animation you created for the Skeletal Mesh.
 1. In the Anim Asset **Details** panel, locate the **Animation** category and locate the **Retarget Source** property. In the dropdown, you should now see the Skeletal Mesh you added to the **Retarget Source Manager**. By choosing this mesh, you are specifying that you want to use the proportions of that mesh to drive this animation.
 
 
