#UE4doc #docs:
Removed italics from mulitple terms in the intro excerpt for both pages

diff --git a/Engine/Documentation/Source/Engine/Animation/AnimBlueprints/AnimGraph/AnimGraph.INT.udn b/Engine/Documentation/Source/Engine/Animation/AnimBlueprints/AnimGraph/AnimGraph.INT.udn
--- a/Engine/Documentation/Source/Engine/Animation/AnimBlueprints/AnimGraph/AnimGraph.INT.udn
+++ b/Engine/Documentation/Source/Engine/Animation/AnimBlueprints/AnimGraph/AnimGraph.INT.udn
@@ -10,10 +10,10 @@ Description:Graph used to sample, blend, and manipulate poses to be applied to S
 [TOC(start:2 end:3)]
 
 [EXCERPT:Intro]
-The **AnimGraph** is used to evaluate a final pose for the _SkeletalMesh_ for the current frame. By default, 
-each _Animation Blueprint_ has an **_AnimGraph_** which can have [animation nodes](Engine/Animation/AnimBlueprints/AnimGraph#AnimationNodes) placed within it 
-to sample _AnimationSequences_, perform animation blends, or control bone transforms using _SkeletalControls_. 
-The resulting pose is then applied to the _SkeletalMesh_ each frame.
+The **AnimGraph** is used to evaluate a final pose for the SkeletalMesh for the current frame. By default, 
+each Animation Blueprint has an **AnimGraph** which can have [animation nodes](Engine/Animation/AnimBlueprints/AnimGraph#AnimationNodes) placed within it 
+to sample AnimationSequences, perform animation blends, or control bone transforms using SkeletalControls. 
+The resulting pose is then applied to the SkeletalMesh each frame.
 [/EXCERPT:Intro]
 
 ![Anim Graph](animgraph.png)(w:720)
diff --git a/Engine/Documentation/Source/Engine/Animation/AnimBlueprints/EventGraph/EventGraph.INT.udn b/Engine/Documentation/Source/Engine/Animation/AnimBlueprints/EventGraph/EventGraph.INT.udn
--- a/Engine/Documentation/Source/Engine/Animation/AnimBlueprints/EventGraph/EventGraph.INT.udn
+++ b/Engine/Documentation/Source/Engine/Animation/AnimBlueprints/EventGraph/EventGraph.INT.udn
@@ -10,9 +10,9 @@ Description:Event-based graph used to update the Animation Blueprint and calcula
 [TOC(start:2 end:3)]
 
 [EXCERPT:Intro]
-Every _Animation Blueprint_ has a single **EventGraph**, which is a standard [graph](Engine/Blueprints/Editor/UIComponents/GraphEditor) 
+Every Animation Blueprint has a single **EventGraph**, which is a standard [graph](Engine/Blueprints/Editor/UIComponents/GraphEditor) 
 that uses a collection of special animation-related events to initiate sequences of nodes. The most common use of the **EventGraph** 
-is to update values used by _BlendSpaces_ and other blend nodes to drive animations within the **AnimGraph**.
+is to update values used by BlendSpaces and other blend nodes to drive animations within the **AnimGraph**.
 [/EXCERPT:Intro]
 
 ![Event graph](eventgraph.png)(w:720)
