#UE4doc #docs:
Removed Overview header
Added engine version field
Changed EventGraph and AnimGraph page's intro excerpts that were being used on this page to remove italics.
Removed italics from multiple terms

diff --git a/Engine/Documentation/Source/Engine/Animation/AnimBlueprints/VimBlueprints.INT.udn b/Engine/Documentation/Source/Engine/Animation/AnimBlueprints/VimBlueprints.INT.udn
--- a/Engine/Documentation/Source/Engine/Animation/AnimBlueprints/VimBlueprints.INT.udn
+++ b/Engine/Documentation/Source/Engine/Animation/AnimBlueprints/VimBlueprints.INT.udn
@@ -9,6 +9,7 @@ Related: Engine/Animation/PhysicallyDrivenAnimation
 Related: Engine/Animation/StateMachines
 Related: Gameplay/AI
 Related: Engine/Content/Tools/MayaRiggingTool
+version: 4.9
 
 [VAR:Topic]
 [OBJECT:Topic]
@@ -32,19 +33,18 @@ Related: Engine/Content/Tools/MayaRiggingTool
 
 [TOC(start:2 end:2)]
 
-## Overview
 
 [EXCERPT:Intro]
 An **Animation Blueprint** is essentially a specialized [Blueprint](Engine/Blueprints) whose graphs 
-control the animation of a _SkeletalMesh_. It can perform blending of animations, directly control the 
-bones of the skeleton, and output a final pose for a _SkeletalMesh_ each frame.
-
-The _Controller_ directs its _Pawn_ (or _Character_) to move based on player input or decisions made based on the gameplay 
-environment. Each _Pawn_ has a _SkeletalMeshComponent_ that references the _SkeletalMesh_ to animate and 
-has an instance of a _Animation Blueprint_. The _Animation Blueprint_, through its two graphs, can access properties 
-of the owning _Pawn_, compute values used for blending or state transitions or driving _AnimMontages_, and 
-calculate the current pose of the _SkeletalMesh_ based on blending of _AnimationSequences_ and direct 
-transformations of the _Skeleton_ from _SkeletalControls_.
+control the animation of a SkeletalMesh. It can perform blending of animations, directly control the 
+bones of the skeleton, and output a final pose for a SkeletalMesh each frame.
+
+The Controller directs its Pawn (or Character) to move based on player input or decisions made based on the gameplay 
+environment. Each Pawn has a SkeletalMeshComponent that references the SkeletalMesh to animate and 
+has an instance of a Animation Blueprint. The Animation Blueprint, through its two graphs, can access properties 
+of the owning Pawn, compute values used for blending or state transitions or driving AnimMontages, and 
+calculate the current pose of the SkeletalMesh based on blending of AnimationSequences and direct 
+transformations of the Skeleton from SkeletalControls.
 [/EXCERPT:Intro]
 
 [REGION:tip]
@@ -54,9 +54,9 @@ A simple example of an **Animation Blueprint** is demonstrated on the [Animation
 ## Working with Animation Blueprints
 
 [EXCERPT:Graphs]
-There are two main components in an _Animation Blueprint_ that work in conjunction to create the final animation 
+There are two main components in an Animation Blueprint that work in conjunction to create the final animation 
 for each frame. The **EventGraph** can perform updates to values which can then be used in the **AnimGraph** 
-to drive _State Machines_, _BlendSpaces_, or other nodes allowing blending between multiple _AnimationSequences_ or poses 
+to drive State Machines, BlendSpaces, or other nodes allowing blending between multiple AnimationSequences or poses 
 that can fire off notifications to other systems to enable dynamic animation-driven effects to occur.
 [/EXCERPT:Graphs]
 
