#UE4doc #docs - Added skill level and version information at top of page. Bolded links in text per instructional guide.

diff --git a/Engine/Documentation/Source/Engine/Components/Physics/PhysicsComponents.INT.udn b/Engine/Documentation/Source/Engine/Components/Physics/PhysicsComponents.INT.udn
--- a/Engine/Documentation/Source/Engine/Components/Physics/PhysicsComponents.INT.udn
+++ b/Engine/Documentation/Source/Engine/Components/Physics/PhysicsComponents.INT.udn
@@ -7,6 +7,8 @@ Related: Engine/Physics/Constraints/ConstraintsBlueprints
 Related: Resources/ContentExamples/Physics
 Related: Resources/ContentExamples/Destructables
 Related: Engine/Blueprints/UserGuide/Components
+SkillLevel: Beginner
+Version: 4.9
 
 The Physics Components are used to affect any objects that are utilizing physics in your level in varying ways. 
 
@@ -22,13 +24,13 @@ An example of using this type of component would be in the form of a window fram
 
 A **PhysicsConstraintComponent** is a joint that allows you to connect two rigid bodies together. You can create different types of joints using the various parameters of this Component. 
 
-By using a **PhysicsConstraintComponent** and two **StaticMeshComponents**, you can create dangling type objects such as a tire swing, heavy bag or sign that reacts to physics in the world allowing players to interact with the Component (see [](Engine/Physics/Constraints/ConstraintsBlueprints) for an example in **Blueprints**). 
+By using a **PhysicsConstraintComponent** and two **StaticMeshComponents**, you can create dangling type objects such as a tire swing, heavy bag or sign that reacts to physics in the world allowing players to interact with the Component (see **[](Engine/Physics/Constraints/ConstraintsBlueprints)** for an example in **Blueprints**). 
 
 ![](constraint.png)(w:400)
 
 ## Physics Handle Component
 
-The **PhysicsHandleComponent** is an object for "grabbing" and moving physics objects around while allowing the object you are grabbing to continue to use physics. An example of this could be in the form of a "gravity gun" where you can pick-up and drop physics objects (see the [Physics Content Examples](Resources/ContentExamples/Physics) for more information).
+The **PhysicsHandleComponent** is an object for "grabbing" and moving physics objects around while allowing the object you are grabbing to continue to use physics. An example of this could be in the form of a "gravity gun" where you can pick-up and drop physics objects (see the [**Physics Content Examples**](Resources/ContentExamples/Physics) for more information).
 
 ![](handle.png)(w:400)
 
@@ -44,7 +46,7 @@ An example of using a thruster is in the form of a rocket (pictured below) which
 
 The **RadialForceComponent** is used to emit a radial force or impulse that can affect physics objects and or destructible objects. Unlike a **PhysicsThrusterComponent**, this type of Component applies "fire-and-forget" type of force and is not continuous. 
 
-You can use this type of Component to push the fractured pieces of a destroyed object, for example an explosion. Using a **RadialForceComponent** to specify the force and direction, when an object is destroyed you can "push" the pieces outwards in a particular direction as seen in the example image below (see the [Destructibles Content Examples](Resources/ContentExamples/Destructables) for more information). 
+You can use this type of Component to push the fractured pieces of a destroyed object, for example an explosion. Using a **RadialForceComponent** to specify the force and direction, when an object is destroyed you can "push" the pieces outwards in a particular direction as seen in the example image below (see the [**Destructibles Content Examples**](Resources/ContentExamples/Destructables) for more information). 
 
 ![](radial_force.png)(w:600)
 
