#UE4doc #docs - Added skill level and version information at top of page. Updated image to remove Blue Guy from 2d render target section.

diff --git a/Engine/Documentation/Source/Engine/Components/Rendering/RenderingComponents.INT.udn b/Engine/Documentation/Source/Engine/Components/Rendering/RenderingComponents.INT.udn
--- a/Engine/Documentation/Source/Engine/Components/Rendering/RenderingComponents.INT.udn
+++ b/Engine/Documentation/Source/Engine/Components/Rendering/RenderingComponents.INT.udn
@@ -6,14 +6,15 @@ Related: Engine/Rendering
 Related: Engine/Rendering/ParticleSystems
 Related: Engine/Blueprints
 Related: Engine/Blueprints/UserGuide/Components
-
+Skilllevel: Beginner
+Version: 4.9
 [TOC (start:2 end:2)]
 
 ## Atmospheric Fog Component
 
 **AtmosphericFogComponents** are used to create fogging effects such as clouds or ambient fog in a level. There are several settings that can be adjusted for this Component that can affect how the effect is generated in your level when placed. 
 
-Examples of this type of Component in use with varying values for its **Decay Height** setting (which controls fog density decay height, ie. lower values cause the fog to be denser, while higher values thin the fog, causing less scatter) are portrayed below. For more information, see the [](Engine/Actors/FogEffects/AtmosphericFog) page. 
+Examples of this type of Component in use with varying values for its **Decay Height** setting (which controls fog density decay height, ie. lower values cause the fog to be denser, while higher values thin the fog, causing less scatter) are portrayed below. For more information, see the **[](Engine/Actors/FogEffects/AtmosphericFog)** page. 
 
 [INCLUDE:Engine/Actors/FogEffects/AtmosphericFog#examples]
 
@@ -23,7 +24,7 @@ The **ExponentialHeightFogComponent** is used to create fog effects but with a d
 
 [INCLUDE:Engine/Actors/FogEffects/HeightFog#intro]
 
-See [](Engine/Actors/FogEffects/HeightFog) for more information and settings that can be adjusted. 
+See **[](Engine/Actors/FogEffects/HeightFog)** for more information and settings that can be adjusted. 
 
 ## Billboard Component
 
@@ -35,7 +36,7 @@ Inside the level, you can manipulate the fog sheet by selecting the **BillboardC
 
 ![](billboard2.png)
 
-For an example of how this Fog Sheet was created, see the [](Resources/Showcases/BlueprintExamples/FogSheet) documentation. 
+For an example of how this Fog Sheet was created, see the **[](Resources/Showcases/BlueprintExamples/FogSheet)** documentation. 
 
 ## Material Billboard Component
 
@@ -68,16 +69,16 @@ A **DecalComponent** is a material that is rendered onto the surface of a mesh (
 
 There are several pages you can refer to for more information on working with Decals. 
 
-* See [](Engine/Actors/DecalActor). 
-* See [](Resources/ContentExamples/Decals/1_1). 
-* See [Decals Content Examples](Resources/ContentExamples/Materials#decalsmap). 
+* See **[](Engine/Actors/DecalActor)**. 
+* See **[](Resources/ContentExamples/Decals/1_1)**. 
+* See **[Decals Content Examples](Resources/ContentExamples/Materials#decalsmap)**. 
 
 ## Instanced Static Mesh Component
 
 An **InstancedStaticMeshComponent** is a Component that efficiently renders multiple instances of the same Static Mesh. 
 This type of Component is particularly useful for procedural level or room creation, as instead of placing hundreds (or even thousands) of **Static Mesh Actors** in your level, you can place one **Instanced Static Mesh** and add multiple instances of the Static Mesh (floors or walls for example) at a much lower performance cost. 
 
-See the [Procedural Room Generation](https://docs.unrealengine.com/latest/INT/Videos/Player/index.html?series=PLZlv_N0_O1ga0aV9jVqJgog0VWz1cLL5f&video=mI7eYXMJ5eI&utm_source=uelauncher&utm_medium=software&utm_campaign=learntab) training stream for an example of working with **InstancedStaticMesh** Components and procedurally generating a random room. 
+See the [**Procedural Room Generation**](https://docs.unrealengine.com/latest/INT/Videos/Player/index.html?series=PLZlv_N0_O1ga0aV9jVqJgog0VWz1cLL5f&video=mI7eYXMJ5eI&utm_source=uelauncher&utm_medium=software&utm_campaign=learntab) training stream for an example of working with **InstancedStaticMesh** Components and procedurally generating a random room. 
 
 ## Particle System Component
 
@@ -91,22 +92,22 @@ Through script we could state that the spark effect is off by default, but activ
 
 ![](particle2.png)
 
-See [](Engine/Rendering/ParticleSystems) for more information. 
+See **[](Engine/Rendering/ParticleSystems)** for more information. 
 
 
 ## Post Process Component
 
 **PostProcessComponets** enable Post Process controls for **Blueprints**. It will use a parent `UShapeComponent` to provide volume data if available. This type of Component could be used to shift the tone of your level, when Post Process Settings are applied to it. For example suppose you defined a default Post Process Setting to use for your game, when the player is damaged (or potentially killed), you could through script change the settings to a black/white tone for the **Scene Color Tint**. 
 
-For more information see the [](Resources/ContentExamples/PostProcessing) or [](Engine/Rendering/PostProcessEffects) documentation. 
+For more information see the **[](Resources/ContentExamples/PostProcessing)** or **[](Engine/Rendering/PostProcessEffects)** documentation. 
 
 ## Scene Capture 2D Component
 
-The **SceneCapture2DComponent** is used to capture a "snapshot" of the scene from a single plane and feed it to a Render Target. There are settings for controlling the **Field of View** when capturing as well as specifying the **Render Target** texture (among other settings). An example of using this would be in the case of creating a mirror (shown below) or security camera view on a monitor (see [](Resources/Showcases/BlueprintExamples/SecurityCameraAndAlarm).) 
+The **SceneCapture2DComponent** is used to capture a "snapshot" of the scene from a single plane and feed it to a Render Target. There are settings for controlling the **Field of View** when capturing as well as specifying the **Render Target** texture (among other settings). An example of using this would be in the case of creating a mirror (shown below) or security camera view on a monitor (see **[](Resources/Showcases/BlueprintExamples/SecurityCameraAndAlarm)**.) 
 
 ![](2dRender.png)(w:440)
 
-Above we have created a **Blueprint** with a **SceneCapture2DComponent** and assigned a **Render Target** texture which was then used as a **Material** that we applied to the piece of geometry in the level. For more information, see [Scene Capture 2D Content Example](Resources/ContentExamples/Reflections/1_7). 
+Above we have created a **Blueprint** with a **SceneCapture2DComponent** and assigned a **Render Target** texture which was then used as a **Material** that we applied to the piece of geometry in the level. For more information, see [**Scene Capture 2D Content Example**](Resources/ContentExamples/Reflections/1_7). 
 
 ## Scene Capture Cube Component
 
@@ -114,7 +115,7 @@ A **SceneCaptureCubeComponent** is used to capture a "snapshot" of the scene fro
 
 ![](Resources/ContentExamples/Reflections/1_6/Reflections_1_6_Header.png)(w:640)
 
-In most instances the **SceneCapture2DComponent** should accomplish most of your scene capturing needs, however in the event that you need a 3D capture of an environment you could use this type of Component. Be aware of your usage however as these are very expensive performance wise and should only be used when absolutely necessary. See the [](Resources/Showcases/Reflections) for more information on the different the methods that can be used to create reflections in your levels. 
+In most instances the **SceneCapture2DComponent** should accomplish most of your scene capturing needs, however in the event that you need a 3D capture of an environment you could use this type of Component. Be aware of your usage however as these are very expensive performance wise and should only be used when absolutely necessary. See the **[](Resources/Showcases/Reflections)** for more information on the different the methods that can be used to create reflections in your levels. 
 
 ## Spline Mesh Component
 
@@ -126,7 +127,7 @@ While you specify the vectors for the position and tangent on the Component itse
 
 ![](splinemesh2.png)
 
-Above we can move the **Start Transform** and **End Transform** independently of moving the entire Actor, allowing us to stretch or rotate it as we desire. This example and setup can be found inside the [](Resources/ContentExamples/Blueprint_Splines) Content Examples map. 
+Above we can move the **Start Transform** and **End Transform** independently of moving the entire Actor, allowing us to stretch or rotate it as we desire. This example and setup can be found inside the **[](Resources/ContentExamples/Blueprint_Splines)** Content Examples map. 
 
 ## Text Render Component
 
@@ -143,8 +144,8 @@ In game when the player enters the trigger, the **TextRenderComponent** text is
 
 ## Vector Field Component
 
-A **VectorFieldComponent** is used for referencing a [Vector Field](Engine/Rendering/ParticleSystems/VectorFields) which are 3D containers with a grid of velocity vectors that can be used to determine the velocity or acceleration of GPU sprites. 
-Vector Fields can be used for small scale effects like wind gusts particle effects to large scale blizzards in your levels. See also [Local Vector Fields](Resources/ContentExamples/EffectsGallery/2_A) and [Global Vector Fields](Resources/ContentExamples/EffectsGallery/2_B) for more information. 
+A **VectorFieldComponent** is used for referencing a [**Vector Field**](Engine/Rendering/ParticleSystems/VectorFields) which are 3D containers with a grid of velocity vectors that can be used to determine the velocity or acceleration of GPU sprites. 
+Vector Fields can be used for small scale effects like wind gusts particle effects to large scale blizzards in your levels. See also [**Local Vector Fields**](Resources/ContentExamples/EffectsGallery/2_A) and [**Global Vector Fields**](Resources/ContentExamples/EffectsGallery/2_B) for more information. 
 
 ![](Resources/ContentExamples/EffectsGallery/2_B/image050.png)(w:640)
 
