#UE4 doc #docs - Added skill level and version information at top of page.

diff --git a/Engine/Documentation/Source/Engine/Components/Shapes/ShapesComponents.INT.udn b/Engine/Documentation/Source/Engine/Components/Shapes/ShapesComponents.INT.udn
--- a/Engine/Documentation/Source/Engine/Components/Shapes/ShapesComponents.INT.udn
+++ b/Engine/Documentation/Source/Engine/Components/Shapes/ShapesComponents.INT.udn
@@ -7,6 +7,8 @@ Related: Engine/Physics/PhysicsBodies
 Related: Engine/Landscape/Editing/Splines
 Related: Resources/ContentExamples/Blueprint_Splines
 Related: Engine/Blueprints/UserGuide/Components
+SkillLevel: Beginner
+Version: 4.9
 
 The Shape Components can be used to create collision or triggers around other objects, indicate orientation or create paths. 
 
@@ -42,7 +44,7 @@ Above a **CapsuleComponent** is automatically included when creating a new **Cha
 
 ![](sphere_component.png)
 
-Just as with the Box and Capsule Components, you can set the [Collision Responses](Engine/Physics/Collision) to generate the type of collision functionality you need. 
+Just as with the Box and Capsule Components, you can set the [**Collision Responses**](Engine/Physics/Collision) to generate the type of collision functionality you need. 
 
 ## Spline Component
 
