#UE4doc #docs - Added skill level and version information at top of page.

diff --git a/Engine/Documentation/Source/Engine/Components/StaticMesh/StaticMeshComponents.INT.udn b/Engine/Documentation/Source/Engine/Components/StaticMesh/StaticMeshComponents.INT.udn
--- a/Engine/Documentation/Source/Engine/Components/StaticMesh/StaticMeshComponents.INT.udn
+++ b/Engine/Documentation/Source/Engine/Components/StaticMesh/StaticMeshComponents.INT.udn
@@ -6,12 +6,14 @@ Related: Engine/Content/Types/StaticMeshes
 Related: Engine/Actors/StaticMeshActor
 Related: Engine/Physics/Collision/Reference
 Related: Engine/Blueprints/UserGuide/Components
+SkillLevel: Beginner
+Version: 4.9
 
 Static Mesh Components allow a piece of geometry to be used as a sub-object for another Actor. 
 
 ## Static Mesh Component
 
-The **StaticMeshComponent** is used to create an instance of a `UStaticMesh`. A Static Mesh is a piece of geometry that consists of a static set of polygons and are the basic unit used to create world geometry for levels in Unreal Engine 4. In addition to building levels, Static Meshes can be used for creating movers such as doors or lifts, rigid body physics objects, foliage and terrain decorations, procedurally created buildings, game objectives, and many more visual elements. 
+The **StaticMeshComponent** is used to create an instance of a `UStaticMesh`. A **Static Mesh** is a piece of geometry that consists of a static set of polygons and are the basic unit used to create world geometry for levels in Unreal Engine 4. In addition to building levels, Static Meshes can be used for creating movers such as doors or lifts, rigid body physics objects, foliage and terrain decorations, procedurally created buildings, game objectives, and many more visual elements. 
 
 Below, a StaticMeshComponent representing a light fixture along with a SpotLightComponent and PointLightComponent are used to create an alarm that could cycle between two different types of lights.
 
