#UE4doc #docs - Added skill level and version information at top of page.

diff --git a/Engine/Documentation/Source/Engine/Components/Utility/UtilityComponents.INT.udn b/Engine/Documentation/Source/Engine/Components/Utility/UtilityComponents.INT.udn
--- a/Engine/Documentation/Source/Engine/Components/Utility/UtilityComponents.INT.udn
+++ b/Engine/Documentation/Source/Engine/Components/Utility/UtilityComponents.INT.udn
@@ -3,6 +3,8 @@ Title:Utility Components
 Crumbs:%ROOT%, Engine, Engine/Components
 Description:Description of the Utility based Components and how you can use them. 
 Related: Engine/Blueprints/UserGuide/Components
+SkillLevel: Beginner
+Version: 4.9
 
 
 ## Application Lifecycle Component
