#UE4doc #docs - Added skill level and version information at top of page.

diff --git a/Engine/Documentation/Source/Engine/Components/Widget/WidgetComponents.INT.udn b/Engine/Documentation/Source/Engine/Components/Widget/WidgetComponents.INT.udn
--- a/Engine/Documentation/Source/Engine/Components/Widget/WidgetComponents.INT.udn
+++ b/Engine/Documentation/Source/Engine/Components/Widget/WidgetComponents.INT.udn
@@ -6,8 +6,10 @@ Related: Engine/UMG
 Related: Engine/UMG/UserGuide/WidgetBlueprints
 Related: Engine/MediaFramework
 Related: Engine/Blueprints/UserGuide/Components
+SkillLevel: Beginner
+Version: 4.9
 
-Widget Components allow you to manifest 3D UI elements created through [Unreal Motion Graphics](Engine/UMG/UserGuide) in your game world. 
+Widget Components allow you to manifest 3D UI elements created through [**Unreal Motion Graphics**](Engine/UMG/UserGuide) in your game world. 
 
 ## Widget Component
 
@@ -27,7 +29,7 @@ Once we place the Actor that contains our WidgetComponent in our level, our Widg
 
 ![](Engine/UMG/HowTo/Create3DWidgets/UMG_3D_30.png)(w:720)
 
-See [](Engine/UMG/HowTo/Create3DWidgets) for a step-by-step guide to creating this example using a WidgetCompoment and Widget Blueprint. 
+See **[](Engine/UMG/HowTo/Create3DWidgets)** for a step-by-step guide to creating this example using a WidgetCompoment and Widget Blueprint. 
 
 [REGION:note]
 This Component is considered as experimental and may be substantially changed or removed in future releases. 
