#DocFix

Updated Mouse Controls on page

diff --git a/Engine/Documentation/Source/Engine/Rendering/Materials/Editor/Interface/MaterialEditorInterface.INT.udn b/Engine/Documentation/Source/Engine/Rendering/Materials/Editor/Interface/MaterialEditorInterface.INT.udn
--- a/Engine/Documentation/Source/Engine/Rendering/Materials/Editor/Interface/MaterialEditorInterface.INT.udn
+++ b/Engine/Documentation/Source/Engine/Rendering/Materials/Editor/Interface/MaterialEditorInterface.INT.udn
@@ -180,7 +180,7 @@ per other mesh tools, etc.
 | **Ctrl + Alt + Shift + LMB drag**  | Box select (add to current selection)  |
 | **LMB drag on connector**  | Create connection (release on connector)  |
 | **LMB drag from connection**  | Move connection (release on same type connector)  |
-| **Shift + LMB on connector**  | Marks the connector. Performing the action again with a connector marked will create connection between the two connectors. This is a quick way of making connections over large distances.  |
+| **Shift + Double-Click LMB on connector**  | Marks the connector. Performing the action again with a connector marked will create connection between the two connectors. This is a quick way of making connections over large distances.  |
 | **RMB on background**  | Bring up **New Expression** menu  |
 | **RMB on object**  | Bring up **Object** menu  |
 | **RMB on connector**  | Bring up **Object** menu |
