#UE4doc #docs:
Removed unnecessary italics / bolding from the intro excerpts for each page.

diff --git a/Engine/Documentation/Source/Engine/Animation/Blendspaces/Blendspaces.INT.udn b/Engine/Documentation/Source/Engine/Animation/Blendspaces/Blendspaces.INT.udn
--- a/Engine/Documentation/Source/Engine/Animation/Blendspaces/Blendspaces.INT.udn
+++ b/Engine/Documentation/Source/Engine/Animation/Blendspaces/Blendspaces.INT.udn
@@ -36,7 +36,7 @@ Related: Engine/Content/Tools/MayaRiggingTool
 ## Overview
 
 [EXCERPT:Intro]
-**Blend Spaces** are special assets which can be sampled in **AnimGraphs** that allow for blending of 
+**Blend Spaces** are special assets which can be sampled in AnimGraphs that allow for blending of 
 animations based on the values of two inputs. Simple blending between two animations based on a 
 single input can be performed using one of the standard [blend nodes](Engine/Animation/NodeReference/Blend) 
 available in Animation Blueprints. Blend Spaces provide a means of doing more complex blending 
diff --git a/Engine/Documentation/Source/Engine/Animation/Blendspaces/UserGuide/BlendspaceUserGuide.INT.udn b/Engine/Documentation/Source/Engine/Animation/Blendspaces/UserGuide/BlendspaceUserGuide.INT.udn
--- a/Engine/Documentation/Source/Engine/Animation/Blendspaces/UserGuide/BlendspaceUserGuide.INT.udn
+++ b/Engine/Documentation/Source/Engine/Animation/Blendspaces/UserGuide/BlendspaceUserGuide.INT.udn
@@ -28,7 +28,7 @@ Description:Guide to using Blend Spaces in Animation Blueprints to input data an
 ## Overview
 
 [EXCERPT:Intro]
-_Blend Spaces_ are used by placing nodes in the **AnimGraphs** of _Animation Blueprints_. These nodes take in value data that is used to drive which animations the _Blend Space_ uses to create the final blended pose.
+Blend Spaces are used by placing nodes in the **AnimGraphs** of Animation Blueprints. These nodes take in value data that is used to drive which animations the Blend Space uses to create the final blended pose.
 [/EXCERPT:Intro]
 
 ## Adding Blend Space Nodes
diff --git a/Engine/Documentation/Source/Engine/Animation/NodeReference/AnimationNodes.INT.udn b/Engine/Documentation/Source/Engine/Animation/NodeReference/AnimationNodes.INT.udn
--- a/Engine/Documentation/Source/Engine/Animation/NodeReference/AnimationNodes.INT.udn
+++ b/Engine/Documentation/Source/Engine/Animation/NodeReference/AnimationNodes.INT.udn
@@ -4,7 +4,7 @@ Crumbs: %ROOT%, Engine, Engine/Animation
 Description:Descriptions of the various animation nodes available for use in Animation Blueprints.
 
 [EXCERPT:Intro]
-_AnimationBlueprints_ make use of various nodes within their graphs to perform operations on input 
+AnimationBlueprints make use of various nodes within their graphs to perform operations on input 
 poses, such as blending, direct bone manipulation, and more. There are a few distinct types 
 of animation nodes provided in the engine including events, blend nodes, skeletal controllers, 
 space nodes, and conversion nodes.
diff --git a/Engine/Documentation/Source/Engine/Animation/NodeReference/SkeletalControls/SkeletalControls.INT.udn b/Engine/Documentation/Source/Engine/Animation/NodeReference/SkeletalControls/SkeletalControls.INT.udn
--- a/Engine/Documentation/Source/Engine/Animation/NodeReference/SkeletalControls/SkeletalControls.INT.udn
+++ b/Engine/Documentation/Source/Engine/Animation/NodeReference/SkeletalControls/SkeletalControls.INT.udn
@@ -8,12 +8,12 @@ Description:Animation nodes that allow for direct manipulation and applying solv
 [/REGION]
 
 [EXCERPT:Intro]
-**_SkeletalControls_** (also called _SkelControls_) allow direct control of bones within a Skeleton. These can be 
-used within _Animation Blueprints_ to control an individual bone, create IK chains, and more. This direct control of the 
-underlying _Skeleton_ provides the ability to create procedural, dynamically-driven animation. The _Transform_ of 
+**SkeletalControls** (also called SkelControls) allow direct control of bones within a Skeleton. These can be 
+used within Animation Blueprints to control an individual bone, create IK chains, and more. This direct control of the 
+underlying Skeleton provides the ability to create procedural, dynamically-driven animation. The Transform of 
 one bone can be used to drive another or traces can be used to conform the feet of a character to the ground while 
 playing a generic walk animation. Any sort of modification can be used to tweak or completely override the bone 
-_Transforms_ applied by _AnimationSequences_.
+Transforms applied by AnimationSequences.
 [/EXCERPT:Intro]
 
 ## Component Space
diff --git a/Engine/Documentation/Source/Engine/Animation/NodeReference/SpaceConversion/SpaceConversion.INT.udn b/Engine/Documentation/Source/Engine/Animation/NodeReference/SpaceConversion/SpaceConversion.INT.udn
--- a/Engine/Documentation/Source/Engine/Animation/NodeReference/SpaceConversion/SpaceConversion.INT.udn
+++ b/Engine/Documentation/Source/Engine/Animation/NodeReference/SpaceConversion/SpaceConversion.INT.udn
@@ -12,16 +12,16 @@ Description:Animation nodes that convert poses between local and component space
 ## Overview
 
 [EXCERPT:Intro]
-Poses can either be in _local-space_ or _component-space_. The **Convert Spaces** nodes available in the **AnimGraph** 
-of _Animation Blueprints_ provide the ability to convert poses between _local_ and _component_ space. Local space assumes 
+Poses can either be in local-space or component-space. The **Convert Spaces** nodes available in the **AnimGraph** 
+of Animation Blueprints provide the ability to convert poses between local and component space. Local space assumes 
 the transform of a bone to be relative to that bone. Component space assumes the bone's transform to be relative to 
-the _SkeletalMeshComponent_.
+the SkeletalMeshComponent.
 
-Generally, when working with poses in an _Animation Blueprint_, they will be in local-space. However, certain blend 
-nodes and all _SkeletalControls_ operate in component-space. This means that the input pose needs to be transformed 
+Generally, when working with poses in an Animation Blueprint, they will be in local-space. However, certain blend 
+nodes and all SkeletalControls operate in component-space. This means that the input pose needs to be transformed 
 prior to being passed in to one of these types of nodes. If the input pose is coming from a node that outputs a local 
-space pose, the pose must be converted to the correct space before a _SkeletalControl_ can perform operations on it. 
-The resulting pose, after the _SkeletalControl_ have performed the operations, then must be converted back to 
+space pose, the pose must be converted to the correct space before a SkeletalControl can perform operations on it. 
+The resulting pose, after the SkeletalControl have performed the operations, then must be converted back to 
 local space for input to additional blends or the **Result** node.
 [/EXCERPT:Intro]
 
diff --git a/Engine/Documentation/Source/Engine/Animation/Sequences/AnimationSequences.INT.udn b/Engine/Documentation/Source/Engine/Animation/Sequences/AnimationSequences.INT.udn
--- a/Engine/Documentation/Source/Engine/Animation/Sequences/AnimationSequences.INT.udn
+++ b/Engine/Documentation/Source/Engine/Animation/Sequences/AnimationSequences.INT.udn
@@ -32,7 +32,7 @@ Related: Engine/Content/Tools/MayaRiggingTool
 [TOC(start:2 end:2)]
 
 [EXCERPT:Intro]
-An **Animation Sequence** is a single animation asset that can be played on a _Skeletal Mesh_. These contain keyframes that specify the position, rotation, and scale of a bone at a specific point in time. By playing these keyframes back in sequence, with blending between them, the bones of a _Skeletal Mesh_ can be smoothly animated. 
+An **Animation Sequence** is a single animation asset that can be played on a Skeletal Mesh. These contain keyframes that specify the position, rotation, and scale of a bone at a specific point in time. By playing these keyframes back in sequence, with blending between them, the bones of a Skeletal Mesh can be smoothly animated. 
 [/EXCERPT:Intro]
 
 [OBJECT:EmbeddedVideo]
diff --git a/Engine/Documentation/Source/Engine/Animation/Sequences/UserGuide/AnimSequenceUserGuide.INT.udn b/Engine/Documentation/Source/Engine/Animation/Sequences/UserGuide/AnimSequenceUserGuide.INT.udn
--- a/Engine/Documentation/Source/Engine/Animation/Sequences/UserGuide/AnimSequenceUserGuide.INT.udn
+++ b/Engine/Documentation/Source/Engine/Animation/Sequences/UserGuide/AnimSequenceUserGuide.INT.udn
@@ -5,9 +5,9 @@ Description:Guide to placing and sampling AnimationSequences in VimBlueprints
 
 [EXCERPT:Intro]
 Playback of **Animation Sequences** can be performed through code, but is generally handled within an instance 
-of an _Animation Blueprint_ applied to a _SkeletalMeshComponent_. All of the **Animation Sequences** for the _Skeleton_ 
-the _Animation Blueprint_ was created from are available to be sampled within the **AnimGraph**. These can be 
-placed as sequence player nodes whose output is a pose generated by sampling the **Animation Sequences**.
+of an Animation Blueprint applied to a SkeletalMeshComponent. All of the Animation Sequences for the Skeleton 
+the Animation Blueprint was created from are available to be sampled within the **AnimGraph**. These can be 
+placed as sequence player nodes whose output is a pose generated by sampling the Animation Sequences.
 [/EXCERPT:Intro]
 
 **To place a sequence player:**
diff --git a/Engine/Documentation/Source/Engine/Animation/SyncGroups/SyncGroups.INT.udn b/Engine/Documentation/Source/Engine/Animation/SyncGroups/SyncGroups.INT.udn
--- a/Engine/Documentation/Source/Engine/Animation/SyncGroups/SyncGroups.INT.udn
+++ b/Engine/Documentation/Source/Engine/Animation/SyncGroups/SyncGroups.INT.udn
@@ -6,9 +6,9 @@ Description:Sync Groups allow you to maintain the synchronization of animations
 
 ## Overview
 [EXCERPT:Intro]
-**Sync Groups** are to keep _related_ animations that may have different lengths synchronized. For example, you may have a walk cycle and a run cycle that you would like to blend together so that the character can smoothly accelerate or decelerate. But what if these animations were different lengths, such as the walk being significantly longer than the run? In such a case, directly blending from one to the other would have unnatural results, an unsightly "beat" as the foot animation switches.
+**Sync Groups** are to keep related animations that may have different lengths synchronized. For example, you may have a walk cycle and a run cycle that you would like to blend together so that the character can smoothly accelerate or decelerate. But what if these animations were different lengths, such as the walk being significantly longer than the run? In such a case, directly blending from one to the other would have unnatural results, an unsightly "beat" as the foot animation switches.
 
-**Sync Groups** solve this problem by allowing one primary animation node to serve as the _Leader_, and all other related animations will simply scale their time length to match. Typically, the leader is  the node with the greatest blend weight. As the weight blends and the follower's weight exceeds the leader, the follower _becomes_ the leader. In this way, the two animations can work smoothly together offering a seamless transition from one motion to the next. 
+**Sync Groups** solve this problem by allowing one primary animation node to serve as the Leader, and all other related animations will simply scale their time length to match. Typically, the leader is  the node with the greatest blend weight. As the weight blends and the follower's weight exceeds the leader, the follower becomes the leader. In this way, the two animations can work smoothly together offering a seamless transition from one motion to the next. 
 
 It should be noted, however, that since the time of the animations is shifting, certain animation considerations become apparent. For instance, in the example of blending between different walk and run cycles, you would want to make sure your animations all started and ended on the same foot. Establishing these standards early will help everything blend much more smoothly!
 [/EXCERPT:Intro]
