#UE4doc #docs:
Removed Overview header
Added engine version field
Removed italics from multiple terms

diff --git a/Engine/Documentation/Source/Engine/Animation/AnimBlueprints/AnimGraph/AnimGraph.INT.udn b/Engine/Documentation/Source/Engine/Animation/AnimBlueprints/AnimGraph/AnimGraph.INT.udn
--- a/Engine/Documentation/Source/Engine/Animation/AnimBlueprints/AnimGraph/AnimGraph.INT.udn
+++ b/Engine/Documentation/Source/Engine/Animation/AnimBlueprints/AnimGraph/AnimGraph.INT.udn
@@ -2,6 +2,7 @@ Availability:Public
 Title:AnimGraph
 Crumbs: %ROOT%, Engine, Engine/Animation, Engine/Animation/AnimBlueprints
 Description:Graph used to sample, blend, and manipulate poses to be applied to SkeletalMeshes by the Animation Blueprint.
+version: 4.9
 
 [REGION:banner]
 ![AnimGraph](animgraph_banner.png)(convert:false)
@@ -11,34 +12,34 @@ Description:Graph used to sample, blend, and manipulate poses to be applied to S
 
 [EXCERPT:Intro]
 The **AnimGraph** is used to evaluate a final pose for the SkeletalMesh for the current frame. By default, 
-each Animation Blueprint has an **AnimGraph** which can have [animation nodes](Engine/Animation/AnimBlueprints/AnimGraph#AnimationNodes) placed within it 
+each Animation Blueprint has an AnimGraph which can have [animation nodes](Engine/Animation/AnimBlueprints/AnimGraph#AnimationNodes) placed within it 
 to sample AnimationSequences, perform animation blends, or control bone transforms using SkeletalControls. 
 The resulting pose is then applied to the SkeletalMesh each frame.
 [/EXCERPT:Intro]
 
 ![Anim Graph](animgraph.png)(w:720)
 
-Inside the **AnimGraph** you can use values calculated in the **EventGraph** or elsewhere, likely in code, 
-and use those as inputs for _BlendSpaces_ or other blend nodes. Of course, you can also just sample 
-_AnimationSequences_ directly without any special blending. The **AnimGraph** also allows for the use of 
-_SkeletalControls_ to modify bone positions, rotations, and scales directly.
+Inside the AnimGraph you can use values calculated in the **EventGraph** or elsewhere, likely in code, 
+and use those as inputs for BlendSpaces or other blend nodes. Of course, you can also just sample 
+AnimationSequences directly without any special blending. The AnimGraph also allows for the use of 
+SkeletalControls to modify bone positions, rotations, and scales directly.
 
 ## Execution Flow
 
-All graphs have a flow of execution shown as pulsing links between _exec pins_ on the graph's nodes. For 
-regular graphs, such as the **EventGraph** of a _Blueprint_, this flow is visualized during play as it is 
+All graphs have a flow of execution shown as pulsing links between exec pins on the graph's nodes. For 
+regular graphs, such as the **EventGraph** of a Blueprint, this flow is visualized during play as it is 
 dependent on the firing of events. The **AnimGraph** is different in that it displays the flow of execution 
 at all times since it is not event-based and is being evaluated each frame.
 
 ![Blend - Alpha 0.0](exec_flow_0.png)(w:720)
 
-In the **AnimGraph**, the flow of execution represents _poses_ being passed from one node to another. Some 
+In the **AnimGraph**, the flow of execution represents poses being passed from one node to another. Some 
 nodes, such as [blends](Engine/Animation/NodeReference/Blend), have multiple inputs and make a 
 decision internally on which input is currently part of the flow of execution. This determination is usually 
-dependent on some external input, like the value passed to a _data pin_.
+dependent on some external input, like the value passed to a data pin.
 
 In the example below, the same blend node from above switched the execution flow between its inputs based on 
-the value of its _Alpha_ pin (the value is being manually set for simplicity):
+the value of its Alpha pin (the value is being manually set for simplicity):
 
 ![Blend - Alpha 1.0](exec_flow_1.png)(w:720)
 [REGION:caption]Blend with Alpha of 1.0 - Input B is used[/REGION]
@@ -46,8 +47,8 @@ the value of its _Alpha_ pin (the value is being manually set for simplicity):
 ## Working with the AnimGraph
 
 The **AnimGraph** works by evaluating a graph of nodes. Some of these nodes perform some specific operation 
-on one or more input poses, while others are used to access or sample other types of assets, such as _AnimationSequences_, 
-_BlendSpaces_, or _AnimMontages_. These assets and nodes, as well as how they fit into the **AnimGraph**, are described below.
+on one or more input poses, while others are used to access or sample other types of assets, such as AnimationSequences, 
+BlendSpaces, or AnimMontages. These assets and nodes, as well as how they fit into the **AnimGraph**, are described below.
 
 ### AnimationSequences
 
@@ -84,7 +85,7 @@ See [](Engine/Animation/NodeReference/Blend) for more information on the space c
 [INCLUDE:Engine/Animation/NodeReference/SkeletalControls#intro]
 
 See [](Engine/Animation/NodeReference/SkeletalControls) for more information on the various 
-types of _SkeletalControls_ available.
+types of SkeletalControls available.
 
 #### Space Conversion
 
