#UE4doc #docs:
Changed font settings for multiple terms

diff --git a/Engine/Documentation/Source/Engine/Animation/NodeReference/Events/AnimationEvents.INT.udn b/Engine/Documentation/Source/Engine/Animation/NodeReference/Events/AnimationEvents.INT.udn
--- a/Engine/Documentation/Source/Engine/Animation/NodeReference/Events/AnimationEvents.INT.udn
+++ b/Engine/Documentation/Source/Engine/Animation/NodeReference/Events/AnimationEvents.INT.udn
@@ -10,34 +10,34 @@ Description:Events called by the animation system to initialize and update the A
 [TOC(start:2)]
 
 [EXCERPT:Intro]
-Various events are called by the animation system to initialize and update the _Animation Blueprint_. These can be placed in the **Event Graph** to implement or override functionality.
+Various events are called by the animation system to initialize and update the Animation Blueprint. These can be placed in the **Event Graph** to implement or override functionality.
 [/EXCERPT:Intro]
 
 ## BlueprintInitializeAnimation
 
 [EXCERPT:Initialize]
-The **Blueprint Initialize Animation** event is executed once when the _Animation Blueprint_ instance is created to perform initialization operations.
+The **Blueprint Initialize Animation** event is executed once when the Animation Blueprint instance is created to perform initialization operations.
 
 ![Initialize Anim Event](event_initializeanim.png)
 
 | Pin | Description |
 | --- | ----------- |
 | [REGION:tablesection]**Output Pins**[/REGION] ||
-| Exec (Blank) | Fired once each frame when the `BlueprintInitializeAnimation()` event is executed in code. |
+| Exec (Blank) | Fired once each frame when the **BlueprintInitializeAnimation()** event is executed in code. |
 
 [/EXCERPT:Initialize]
 
 ## BlueprintUpdateAnimation
 
 [EXCERPT:Update]
-The **Blueprint Update Animation** event is executed every frame allowing the _Animation Blueprint_ to perform calculations and updates to any values it needs. This event is the entry point into the update loop of the EventGraph. The amount of time elapsed since the last update is available from the **DeltaTimeX** pin so time-dependent interpolations or incremental updates can be performed.
+The **Blueprint Update Animation** event is executed every frame allowing the Animation Blueprint to perform calculations and updates to any values it needs. This event is the entry point into the update loop of the EventGraph. The amount of time elapsed since the last update is available from the **DeltaTimeX** pin so time-dependent interpolations or incremental updates can be performed.
 
 ![Update Anim Event](event_updateanim.png)
 
 | Pin | Description |
 | --- | ----------- |
 | [REGION:tablesection]**Output Pins**[/REGION] ||
-| Exec (Blank) | Fired once each frame when the `BlueprintUpdateAnimation()` event is executed in code. |
+| Exec (Blank) | Fired once each frame when the **BlueprintUpdateAnimation()** event is executed in code. |
 | Delta Time X | Outputs the amount of time since the last execution cycle. |
 
 [/EXCERPT:Update]
