#UE4doc #docs:
Removed Overview header
Added engine version field
Changed font settings for multiple terms

diff --git a/Engine/Documentation/Source/Engine/Animation/AnimBlueprints/EventGraph/EventGraph.INT.udn b/Engine/Documentation/Source/Engine/Animation/AnimBlueprints/EventGraph/EventGraph.INT.udn
--- a/Engine/Documentation/Source/Engine/Animation/AnimBlueprints/EventGraph/EventGraph.INT.udn
+++ b/Engine/Documentation/Source/Engine/Animation/AnimBlueprints/EventGraph/EventGraph.INT.udn
@@ -2,6 +2,7 @@ Availability:Public
 Title:EventGraph
 Crumbs: %ROOT%, Engine, Engine/Animation, Engine/Animation/AnimBlueprints
 Description:Event-based graph used to update the Animation Blueprint and calculate values for use in the AnimGraph.
+version: 4.9
 
 [REGION:banner]
 ![EventGraph](eventgraph_banner.png)(convert:false)
@@ -11,13 +12,13 @@ Description:Event-based graph used to update the Animation Blueprint and calcula
 
 [EXCERPT:Intro]
 Every Animation Blueprint has a single **EventGraph**, which is a standard [graph](Engine/Blueprints/Editor/UIComponents/GraphEditor) 
-that uses a collection of special animation-related events to initiate sequences of nodes. The most common use of the **EventGraph** 
-is to update values used by BlendSpaces and other blend nodes to drive animations within the **AnimGraph**.
+that uses a collection of special animation-related events to initiate sequences of nodes. The most common use of the EventGraph 
+is to update values used by BlendSpaces and other blend nodes to drive animations within the AnimGraph.
 [/EXCERPT:Intro]
 
 ![Event graph](eventgraph.png)(w:720)
 		
-For example, the **EventGraph** of an _Animation Blueprint_ could use the velocity of the owning Pawn to determine direction and speed. These two values can then be used as inputs into a _BlendSpace_ within the **AnimGraph** to blend between an idle animation and various directional movement animations. Similarly, the pitch and yaw of the owning Pawn can be used as inputs into another _BlendSpace_ to blend between various aiming poses applied additively in order to cause the mesh to aim where it is looking while still maintaining the previously mentioned movement animations.
+For example, the EventGraph of an Animation Blueprint could use the velocity of the owning Pawn to determine direction and speed. These two values can then be used as inputs into a BlendSpace within the AnimGraph to blend between an idle animation and various directional movement animations. Similarly, the pitch and yaw of the owning Pawn can be used as inputs into another BlendSpace to blend between various aiming poses applied additively in order to cause the mesh to aim where it is looking while still maintaining the previously mentioned movement animations.
 
 ## Working with the EventGraph
 
@@ -25,7 +26,7 @@ For general information on working with node graphs, see the [Working with Nodes
 
 ## Animation Events
 
-[INCLUDE:Engine/Animation/NodeReference/Events#Intro]
+Various events are called by the animation system to initialize and update the Animation Blueprint. These can be placed in the Event Graph to implement or override functionality.
 
 
 ### BlueprintInitializeAnimation
