#UE4doc #docs:
Removed Overview header
Added engine version field
Changed font settings for multiple terms
Reworded a step to reflect changes to the engine
Replaced all images to reflect the modern engine UI

diff --git a/Engine/Documentation/Source/Engine/Animation/AnimBlueprints/Creation/CreatingVimBlueprints.INT.udn b/Engine/Documentation/Source/Engine/Animation/AnimBlueprints/Creation/CreatingVimBlueprints.INT.udn
--- a/Engine/Documentation/Source/Engine/Animation/AnimBlueprints/Creation/CreatingVimBlueprints.INT.udn
+++ b/Engine/Documentation/Source/Engine/Animation/AnimBlueprints/Creation/CreatingVimBlueprints.INT.udn
@@ -2,41 +2,42 @@ Availability:Public
 Title:Animation Blueprint Creation
 Crumbs: %ROOT%, Engine, Engine/Animation, Engine/Animation/AnimBlueprints
 Description:Guide to creating new Animation Blueprints from Skeleton assets
+version: 4.9
 
 
 [TOC (start:2 end:3)]
 
-## Overview
 
-You can create an _Animation Blueprint_ in one of two ways:
+
+You can create an Animation Blueprint in one of two ways:
 * Targeting a specific Skeletal Mesh asset.
 * Not associating with a specific Skeletal Mesh asset.
 
 ## Creating an Animation Blueprint Targeting a Skeletal Mesh
 
-You can create an _Animation Blueprint_ that targets a specific Skeletal Mesh asset by doing the following:
+You can create an Animation Blueprint that targets a specific Skeletal Mesh asset by doing the following:
 
-1. **Right-click** on the Skeletal Mesh asset you want the _Animation Blueprint_ to target and click **Create** > **Create Animation Assets** > **Create Anim Blueprint**.
+1. **Right-click** on the Skeletal Mesh asset you want the Animation Blueprint to target and click **Create** > **Anim Assets** > **Anim Blueprint**.
 
 	![Create Animation Blueprint - Menu](CE_create_menu.png)
 
-1. The _Animation Blueprint_ is created in the same folder as the target Skeletal Mesh asset, with the name ready to be edited.
+1. The Animation Blueprint is created in the same folder as the target Skeletal Mesh asset, with the name ready to be edited.
 
 	![Create Animation Blueprint - New](CE_create_new.png)
 
-1. If you want to, type a new name for the new _Animation Blueprint_ and press **Enter**.
+1. If you want to, type a new name for the new Animation Blueprint and press **Enter**.
 
 	![Create Animation Blueprint - Rename](CE_create_name.png)
 
-1. If you are not happy with the default location, drag the new _Animation Blueprint_ to another folder.
+1. If you are not happy with the default location, drag the new Animation Blueprint to another folder.
 
 	![Create Animation Blueprint - Change Location](CE_create_move.png)(convert:false)
 
 ## Creating an Animation Blueprint Not Associated with a Skeletal Mesh
 
-You can also create an _Animation Blueprint_ that is not yet associated with a specific Skeletal Mesh asset.  To do so, do the following:
+You can also create an Animation Blueprint that is not yet associated with a specific Skeletal Mesh asset.  To do so, do the following:
 
-1. Either **right-click** in an empty space in the **Content Browser** or click the **New Asset** button.
+1. Either **Right-click** in an empty space in the **Content Browser** or click the **New Asset** button.
 
 1. On the context menu that appears, click **Animation**.
 
@@ -44,4 +45,4 @@ You can also create an _Animation Blueprint_ that is not yet associated with a s
 
 	![](CE_create_NoSkel.png)
 
-You can then rename and/or move your new _Animation Blueprint_ by following the steps in the previous section on this page.
+You can then rename and/or move your new Animation Blueprint by following the steps in the previous section on this page.
