#UE4doc #docs:
Removed Overview header
Added engine version field
Changed a couple images to reflect the changes to the editor's UI
Changed font settings for multiple terms

diff --git a/Engine/Documentation/Source/Engine/Animation/AnimBlueprints/Interface/AnimationBlueprintUI.INT.udn b/Engine/Documentation/Source/Engine/Animation/AnimBlueprints/Interface/AnimationBlueprintUI.INT.udn
--- a/Engine/Documentation/Source/Engine/Animation/AnimBlueprints/Interface/AnimationBlueprintUI.INT.udn
+++ b/Engine/Documentation/Source/Engine/Animation/AnimBlueprints/Interface/AnimationBlueprintUI.INT.udn
@@ -2,37 +2,39 @@ Availability:Public
 Title:Animation Blueprint UI
 Crumbs: %ROOT%, Engine, Engine/Animation, Engine/Animation/AnimBlueprints
 Description:An overview of the Animation Blueprint UI.
+version: 4.9
 
 [TOC(start:2 end:2)]
 
-## Overview
+
 
 The following sections describe the areas of the Animation Blueprint UI.
 
 
-![](AnimGraphUI.png)(w:900)
+![](AnimGraphUI.png)
 
-1. **Menu Bar**
-1. **Toolbar**
-1. **Viewport**
-1. **Class Defaults panel**
-1. **Graph**
-1. **Asset Browser tab**
-1. **My Blueprint tab**
-1. **Details panel**
-1. **Palette**
+1. Menu Bar
+1. Toolbar
+1. Viewport
+1. Anim Preview Editor
+1. Asset Override Editor
+1. Graph
+1. Asset Browser tab
+1. My Blueprint tab
+1. Details panel
+1. Palette
 
 ## Menu Bar
 
-The **Menu Bar** gives you access to common editor actions such as saving, undoing and redoing, as well as several Persona-specific menus. For instance, the **Asset** menu allows
+The **Menu Bar** gives you access to common editor actions such as saving, undoing and redoing, as well as several Persona-specific menus. For instance, the **Asset Browser** tab allows
 you to change the default preview mesh, remove unused bones from the skeleton, and record animations. 
 
 
 ## Toolbar
 
-The Animation Blueprint toolbar contains many of the same options as other Blueprint toolbars.
+The Animation Blueprint **Toolbar** contains many of the same options as other Blueprint Toolbars.
 
-![](AnimGraphUI_Toolbar.png)(w:900)
+![](AnimGraphUI_Toolbar.png)
 
 [EXCERPT:Toolbar]
 | Item | Description |
@@ -61,7 +63,7 @@ In addition to those options, the Animation Blueprint toolbar also contains tabs
 
 
 ## Viewport
-The viewport displays the Skeletal Mesh with any animations activated by current variable values.
+The **Viewport** displays the Skeletal Mesh with any animations activated by current variable values.
 
 ![](AnimGraphUI_Viewport.png)
 
@@ -75,7 +77,7 @@ For more information, see [](Engine/Animation/Persona\BlueprintDefaults).
 
 ## Graph
 
-The tabbed graph panel displays whichever graph is currently selected.
+The tabbed **Graph** panel displays whichever graph is currently selected.
 
 ### EventGraph Tab
 
@@ -92,25 +94,25 @@ The **AnimGraph** tab contains the State Machine that hooks up all the actual an
 For more information about State Machines, see [](Engine/Animation/StateMachines).
 
 [REGION:note]
-Although you could put all your animation nodes directly into the **AnimGraph** tab, it is more efficient to organize them into a State Machine, States, and Transitions (Rules), for ease of upkeep and simplicity.
+Although you could put all your animation nodes directly into the AnimGraph tab, it is more efficient to organize them into a State Machine, States, and Transitions (Rules), for ease of upkeep and simplicity.
 [/REGION]
 
 ## Asset Browser Tab
 
-The **Asset Browser** tab contains a list of all the animation assets in your Animation Blueprint.
+The Asset Browser tab contains a list of all the animation assets in your Animation Blueprint.
 
 ![](AnimGraphUI_AssetBrowser.png)
 
-By default, you can find the **Asset Browser** tab next to the **My Blueprint** tab, underneath the graph.
+By default, you can find the Asset Browser tab next to the **My Blueprint** tab, underneath the Graph.
 
 
 ## My Blueprint Tab
 
-The **My Blueprint** tab contains a list of all the graphs and variables in your Animation Blueprint.
+The My Blueprint tab contains a list of all the graphs and variables in your Animation Blueprint.
 
 ![](AnimGraphUI_MyBP.png)
 
-By default, you can find the **My Blueprint** tab next to the **Asset Browser** tab, underneath the graph.
+By default, you can find the **My Blueprint** tab next to the **Asset Browser** tab, underneath the Graph.
 
 
 ## Details Panel
@@ -126,5 +128,5 @@ The **Palette** panel contains all the nodes available for use in the Animation
 ![](AnimGraphUI_Palette.png)
 
 [REGION:note]
-The **Palette** may not be visible by default.  If it is not, you can open it by clicking **Palette** on the **Window** menu.
+The Palette may not be visible by default.  If it is not, you can open it by clicking **Palette** on the **Window** menu.
 [/REGION]
