#UE4doc #docs:
Added engine version field
Added skill level field
Changed a few steps due to incorrect file extensions and filenames
Added step to delete the existing ThirdPersonCharacter as the correct one wouldn't be spawned otherwise
Replaced an image due to an incorrect filename and UI

diff --git a/Engine/Documentation/Source/Engine/Animation/AnimHowTo/AdditiveAnimations/AdditiveHowTo.INT.udn b/Engine/Documentation/Source/Engine/Animation/AnimHowTo/AdditiveAnimations/AdditiveHowTo.INT.udn
--- a/Engine/Documentation/Source/Engine/Animation/AnimHowTo/AdditiveAnimations/AdditiveHowTo.INT.udn
+++ b/Engine/Documentation/Source/Engine/Animation/AnimHowTo/AdditiveAnimations/AdditiveHowTo.INT.udn
@@ -6,6 +6,8 @@ Related: Engine/Animation/AnimBlueprints
 Related: Engine/Animation/AnimMontage
 Related: Engine/Animation/NodeReference/Blend
 Related: Engine/Animation/AnimationBlending
+version: 4.9
+skilllevel: Advanced
 
 Animation blending, as a concept, simply means making a smooth transition between two or more animations on a single character or Skeletal Mesh. In Unreal Engine 4, there are various ways in which such blending can be applied either through [](Engine/Animation/Blendspaces), through additive methods by literally combining two animations based on a weighted bias or alpha value, or even applying direct overrides of the existing pose. 
 
@@ -21,7 +23,7 @@ We are also using the **Animation Starter Pack** which is available for download
 
 ![](AnimationAssetPack.png)
 
-Once downloaded, you can add the Animation Starter Pack to your project via the **Add to project** button from the Launcher.
+Once downloaded, you can add the Animation Starter Pack to your project via the **Add To Project** button from the Launcher.
 [/REGION]
 [/EXCERPT:Intro]
 
@@ -29,7 +31,7 @@ Once downloaded, you can add the Animation Starter Pack to your project via the
 
 In this section, we assign our player character and make modifications to the animation we will blend in with our movement animations. 
 
-1.  With your project open, inside the **Content Browser** under the **Game/Blueprints** folder, open the **ThirdPersonGameMode** Blueprint. 
+1.  With your project open, inside the **Content Browser** under the **Game/ThirdPersonBP/Blueprints** folder, open the **ThirdPersonGameMode** Blueprint. 
 
 1.	Under **Default Pawn Class**, click the dropdown menu and select the **Ue4ASP_Character**. 
 
@@ -42,9 +44,9 @@ In this section, we assign our player character and make modifications to the an
 1.	**Right-click** on the **Fire_Shotgun_Ironsights** animation, select **Create** then **Create AnimMontage**.
   
 	![](Animations2.png)
-
-	You can filter the window by typing **Fire** inside the search window as shown above. 
-
+	[REGION:note]
+	You can filter the window by typing "Fire" inside the search window as shown above. 
+	[/REGION]
 1.	Rename the Montage if you wish, then open it up and click the magnifying glass icon. 
 
 	![](Animations3.png)
@@ -65,7 +67,7 @@ In this section, we assign our player character and make modifications to the an
 
 Next we will make some updates to our Character Blueprint so that we can determine when they are pressing the shoot button. 
 
-1.	Inside the **Content Browser** under **Game/AnimStarterPack/Character**, open the **Ue4ASP_Character** Blueprint. 
+1.	Inside the **Content Browser** under **Game/AnimStarterPack**, open the **Ue4ASP_Character** Blueprint. 
 
 1.	Locate the **Crouching** section, and delete the **InputAction Crouch** Event (with the Warning!) sign. 
 
@@ -101,7 +103,7 @@ Next we will make some updates to our Character Blueprint so that we can determi
 
 Here we will set up the Animation Graph and place the nodes needed to blend between our current pose and the fire weapon animation. 
 
-1.	Also under **Game/AnimStarterPack/Character**, open the **UE4ASP_HeroTPP_AnimBlueprint** Animation Blueprint. 
+1.	Also under **Game/AnimStarterPack**, open the **UE4ASP_HeroTPP_AnimBlueprint** Animation Blueprint. 
 
 1.	Inside the **MyBlueprint** window, open the **AnimGraph**. 
 
@@ -189,7 +191,9 @@ Finally, we update the way we blend between animations so that our upperbody pla
 
 	This will check now if the Montage is currently playing and if it is, do not play it (and if it is not, play the Montage). 
 
-1.	**Compile** and **Save**, close the **Blueprint** and play in the editor. 
+1.	**Compile** and **Save**, close the **Blueprint** and delete the existing **ThirdPersonCharacter** in the scene. 
+
+1. **Play** in the editor. 
 
 	[OBJECT:EmbeddedVideo]
 	[PARAMLITERAL:width]
