#UE4doc #docs:
Removed unnecessary note region

diff --git a/Engine/Documentation/Source/Engine/Animation/AnimHowTo/AdditiveAnimations/AdditiveHowTo.INT.udn b/Engine/Documentation/Source/Engine/Animation/AnimHowTo/AdditiveAnimations/AdditiveHowTo.INT.udn
--- a/Engine/Documentation/Source/Engine/Animation/AnimHowTo/AdditiveAnimations/AdditiveHowTo.INT.udn
+++ b/Engine/Documentation/Source/Engine/Animation/AnimHowTo/AdditiveAnimations/AdditiveHowTo.INT.udn
@@ -16,7 +16,6 @@ You can also send animation directly to a specific bone in the Skeleton and all
 [TOC(start:2 end:2)]
 
 [EXCERPT:Intro]
-[REGION:note]
 For this guide, we have created a new project using the **Blueprint Third Person** template with **Starter Content** enabled. 
 
 We are also using the **Animation Starter Pack** which is available for download for free via the **Marketplace**. 
@@ -24,7 +23,6 @@ We are also using the **Animation Starter Pack** which is available for download
 ![](AnimationAssetPack.png)
 
 Once downloaded, you can add the Animation Starter Pack to your project via the **Add To Project** button from the Launcher.
-[/REGION]
 [/EXCERPT:Intro]
 
 ## Setting up the Animation to Blend
