#UE4doc #docs
Organizing documentation.
Created outline
Added placeholder cover image.
Availability: Docs

diff --git a/Engine/Documentation/Source/Programming/Tutorials/FirstPersonShooter/FirstPersonShooter.INT.udn b/Engine/Documentation/Source/Programming/Tutorials/FirstPersonShooter/FirstPersonShooter.INT.udn
--- a/Engine/Documentation/Source/Programming/Tutorials/FirstPersonShooter/FirstPersonShooter.INT.udn
+++ b/Engine/Documentation/Source/Programming/Tutorials/FirstPersonShooter/FirstPersonShooter.INT.udn
@@ -1,7 +1,38 @@
 Availability: Docs
-Title: First Person Shooter C++ Tutorial
+Title: C++ First Person Shooter Tutorial
 Crumbs: %ROOT%, GettingStarted, Programming/Tutorials
 Description: How to develop a First Person Shooter.
-SkillLevel: Intermediate
-Version: 4.8
+SkillLevel: Beginner
+Version: 4.9
 
+![](FPS_CoverImage.png)
+
+## Goals
+
+The purpose of this tutorial is show you how to build a basic First Person Shooter (FPS) using C++ and Blueprints.
+
+## Objectives
+
+By the end of this tutorial you'll be able to:
+
+* Set-up a Project
+* Implement a Character
+* Implement a Projectile
+* Animate a Character
+
+[OBJECT:TopicButtonList]
+	[PARAM:icon]
+		![](%ROOT%/build_icon.png)(convert:false)
+	[/PARAM]
+	[PARAM:title]
+		Steps
+	[/PARAM]
+	[PARAM:description]
+	[/PARAM]
+	[PARAM:links]
+	* **Project Setup**
+	* **Character Implementation**
+	* **Projectile Implementation**
+	* **Character Animation**
+	[/PARAM]
+[/OBJECT]
