#UE4doc #docs:
Added engine version field
Added skill level field
Changed some file names to match the current content
Changed multiple images to match the modern UI
Added a step that is vital to the end of the tutorial

diff --git a/Engine/Documentation/Source/Engine/Animation/AnimHowTo/BlendSpace/BlendspaceHowTo.INT.udn b/Engine/Documentation/Source/Engine/Animation/AnimHowTo/BlendSpace/BlendspaceHowTo.INT.udn
--- a/Engine/Documentation/Source/Engine/Animation/AnimHowTo/BlendSpace/BlendspaceHowTo.INT.udn
+++ b/Engine/Documentation/Source/Engine/Animation/AnimHowTo/BlendSpace/BlendspaceHowTo.INT.udn
@@ -7,15 +7,16 @@ Related: Engine/Animation/NodeReference/Blend
 Related: Engine/Animation/CharacterSetupOverview
 Related: Engine/Animation/AnimBlueprints
 Related: Engine/Animation/StateMachines
+version: 4.9
+skilllevel: Intermediate
 
 
-[Blend Spaces](Engine/Animation/Blendspaces) are special assets which can be sampled in **AnimGraphs** that allow for blending of animations based on the values of two inputs. Simple blending between two animations based on a single input can be performed using one of the standard [blend nodes](Engine/Animation/NodeReference/Blend) available in _Animation Blueprints_. _Blend Spaces_ provide a means of doing more complex blending between multiple animations based on multiple values (currently limited to two).
+[Blend Spaces](Engine/Animation/Blendspaces) are special assets which can be sampled in **AnimGraphs** that allow for blending of animations based on the values of two inputs. Simple blending between two animations based on a single input can be performed using one of the standard [blend nodes](Engine/Animation/NodeReference/Blend) available in Animation Blueprints. Blend Spaces provide a means of doing more complex blending between multiple animations based on multiple values (currently limited to two).
 
-This page will show you how to take an existing character and apply a _Blend Space_ which will blend between walking and jogging animations based on the characters movement speed. 
+This page will show you how to take an existing character and apply a Blend Space which will blend between walking and jogging animations based on the characters movement speed. 
 
 
 [EXCERPT:Intro]
-[REGION:note]
 For this guide, we have created a new project using the **Blueprint Third Person** template with **Starter Content** enabled. 
 
 We are also using the **Animation Starter Pack** which is available for download for free via the **Marketplace**. 
@@ -23,10 +24,9 @@ We are also using the **Animation Starter Pack** which is available for download
 ![](Engine\Animation\AnimHowTo\AdditiveAnimations\AnimationAssetPack.png)
 
 Once downloaded, you can add the Animation Starter Pack to your project via the **Add to project** button from the Launcher.
-[/REGION]
 [/EXCERPT:Intro]
 
-1.  With your project open, inside the **Content Browser** under the **Game/Blueprints** folder, open the **MyGame** Blueprint. 
+1.  With your project open, inside the **Content Browser** under the **Game/ThirdPersonBP/Blueprints** folder, open the **ThirdPersonGameMode** Blueprint. 
 
 1.	Under **Default Pawn Class**, click the dropdown menu and select the **Ue4ASP_Character**. 
 
@@ -34,15 +34,15 @@ Once downloaded, you can add the Animation Starter Pack to your project via the
 
 1.	**Compile** and **Save** then close the Blueprint. 
 
-1.	In the **Content Browser**, click the **Create** Button then under _Animation_, select **Blend Space**. 
+1.	In the **Content Browser**, click the **Add New** Button then under **Animation**, select **Blend Space**. 
 
 	![](BS_BlendSpace.jpg)
 
-1.	In the **Pick Skeleton** window, select the **HeroTPP_Skeleton** that is located in the **Game/AnimStaterPack/Character** folder. 
+1.	In the **Pick Skeleton** window, select the **UE4_Mannequin_Skeleton** that is located in the **Game/AnimStaterPack/UE4_Mannequin/Mesh** folder. 
 
 	![](BlendSpace2.png)(w:400)
 
-	Give the _Blend Space_ a name such as **Walk_Jog_BS**. 
+	Give the Blend Space a name such as **Walk_Jog_BS**. 
 
 	![](BS_Walk_Jog_BS.jpg)
 
@@ -50,7 +50,7 @@ Once downloaded, you can add the Animation Starter Pack to your project via the
 
 	![](BlendSpace4.png)
 
-	Set the **X Axis Label** to **Direction** and its range to **-180 to 180**. Set the **Y Axis Label** to **Speed** and its range to **150 to 270**. With the values set, click the **Apply Parameter Changes** button which will make the grid space to the left update with your values. Here we are defining the parameters that will determine the motion the _Blend Space_ will produce based on the variables **Direction** and **Speed** (clamped by the range provided) which we will get from the player character. 
+	Set the **X Axis Label** to **Direction** and its range to **-180 to 180**. Set the **Y Axis Label** to **Speed** and its range to **150 to 270**. With the values set, click the **Apply Parameter Changes** button which will make the grid space to the left update with your values. Here we are defining the parameters that will determine the motion the Blend Space will produce based on the variables **Direction** and **Speed** (clamped by the range provided) which we will get from the player character. 
 
 1.	From the **Asset Browser** window, drag the **Walk_Fwd_Rifle_Ironsights** animation into the graph and the bottom center location. 
 
@@ -78,15 +78,15 @@ Once downloaded, you can add the Animation Starter Pack to your project via the
 
 	If you move your mouse around inside the grid, you will see the character start to blend between the different walk/jog animations. 
 
-1.	Inside the **Anim Asset Details** window, under _Sample Interpolation_ set **Target Weight Interpolation Speed Per Sec** to **4.0**. 
+1.	Inside the **Anim Asset Details** window, under Sample Interpolation set **Target Weight Interpolation Speed Per Sec** to **4.0**. 
 
 	![](BlendSpace11.png)
 
-	This will smooth out the transition between the animations in the _Blend Space_ to prevent popping from pose-to-pose. 
+	This will smooth out the transition between the animations in the Blend Space to prevent popping from pose-to-pose. 
 
-1.	Click **Save** then close the _Blend Space_, then in the **Content Browser**, open the **Ue4ASP_HeroTPP_AnimBlueprint**. 
+1.	Click **Save** then close the Blend Space, then in the **Content Browser**, open the **Ue4ASP_HeroTPP_AnimBlueprint**. 
 
-1.	Inside the _Animation Blueprint_, click the **MyBlueprint** tab, then expand **AnimGraph** and **Double-click Locomotion**. 
+1.	Inside the Animation Blueprint, click the **MyBlueprint** tab, then expand **AnimGraph** and **Double-click Locomotion**. 
 
 	![](BlendSpace12.png)(w:640)
 
@@ -100,25 +100,25 @@ Once downloaded, you can add the Animation Starter Pack to your project via the
 
 	![](BlendSpace14.png)(w:640)
 
-1.	Replace **R_Jog** with the **Walk_Jog_BS** _Blend Space_ and connect **Direction** and **Speed** and feed it into the **Result**. 
+1.	Replace **R_Jog** with the **Walk_Jog_BS** Blend Space and connect **Direction** and **Speed** and feed it into the **Result**. 
 
 	![](BlendSpace15.png)(w:640)
 
-1.	Click **Compile** and **Save** then close the _Animation Blueprint_. 
+1.	Click **Compile** and **Save** then close the Animation Blueprint. 
 
-1.	In the **Content Browser**, open the **Ue4ASP_Character** _Blueprint_. 
+1.	In the **Content Browser**, open the **Ue4ASP_Character** Blueprint. 
 
-1.	With the _Blueprint_ open, find the **Crouching** section and replace the **InputAction Crouch** node with a **C** Key Event. 
+1.	With the Blueprint open, find the **Crouching** section and replace the **InputAction Crouch** node with a **C** Key Event. 
 
 	![](BlendSpace17.png)(w:640)
 
-	This _AnimBlueprint_ was built with the expectation that there is an **Action Mapping** called **Crouch**, since it does not exist in a fresh project and to get rid of the yellow warning sign, we will replace it with the **C** Key Event so that the player crouches when C is pressed. 
+	This AnimBlueprint was built with the expectation that there is an **Action Mapping** called **Crouch**, since it does not exist in a fresh project and to get rid of the yellow warning sign, we will replace it with the **C** Key Event so that the player crouches when C is pressed. 
 
 1.	Drag off the **Pressed** pin and add a **FlipFlop** node, then connect the **A** and **B** pins as shown below. 
 
 	![](BlendSpace18.png)(w:640)
 
-	The _FlipFlop_ node acts as a toggle causing the character to crouch the first time C is pressed, then stand the second time it is pressed. 
+	The **FlipFlop** node acts as a toggle causing the character to crouch the first time **C** is pressed, then stand the second time it is pressed. 
 
 1.	Inside **MyBlueprint**, click the **Show inherited variables** box, then hold **Ctrl** and drag the **CharacterMovement** into the graph. 
 
@@ -130,7 +130,7 @@ Once downloaded, you can add the Animation Starter Pack to your project via the
 
 	![](BlendSpace20.png)
 
-1.	Below the _Crouch_ section, **Right-click** in the graph and add the **Left Shift** Key Event. 
+1.	Below the Crouch section, **Right-click** in the graph and add the **Left Shift** Key Event. 
 
 1.	Drag off the **Pressed** pin and add a **Branch** node, then also add a **Branch** of the **Released** pin. 
 
@@ -138,19 +138,21 @@ Once downloaded, you can add the Animation Starter Pack to your project via the
 
 	![](BlendSpace21.png)
 
-1.	Duplicate the **CharacterMovement** and **Set Max Walk Speed** nodes from the _Crouch_ section. 
+1.	Duplicate the **CharacterMovement** and **Set Max Walk Speed** nodes from the Crouch section. 
 
 1.	Connect off the **False** pins of each **Branch** as shown below and set the **Max Walk Speed** values to **270** and **150**. 
 
 	![](BlendSpace22.png)
 
-1.	In the upper-right of the _Blueprint_, click the **Defaults** tab. 
+1.	In the upper-right of the Blueprint, click the **Defaults** tab. 
 
 1.	Enter **Max Walk Speed** in the search field, then set **Max Walk Speed** to **150** and **Max Walk Speed Crouched** to **125**. 
 
 	![](BlendSpace23.png)
 
-1.	Click **Compile** and **Save** then close the _Blueprint_, then click **Play** to play in the editor. 
+1.	Click **Compile** and **Save** then close the Blueprint, then delete the existing ThirdPersonCharacter from the scene.
+
+1. Click **Play** to play in the editor. 
 
 When playing in the Editor, using **WASD** to move around, the character will now walk by default and blend between walking animations when changing directions. If you press **C** to crouch, the character will now stay in a crouched position while walking at a slightly reduced speed than found when standing/walking. While standing, pressing and holding **Left Shift** while using WASD will cause the character to blend from the walk/aim animation to a lowered/aim jog animation, essentially creating a "sprint" for the character. 
 
