#UE4doc #docs:
Added engine version field
Added skill level field
Changed multiple file names due to updated content
Changed a couple images due to outdated UI
Replaced mentions of the "Yellow" axis/arrow as this was referring to the green and/or red axis/arrow indicators from the rotation/translation widgets.
Added multiple steps and images to illustrate how to add the state machine related to the animation created, as it would not play otherwise.
Changed font settings for multiple terms

diff --git a/Engine/Documentation/Source/Engine/Animation/AnimHowTo/LayerEditing/LayerEditingHowTo.INT.udn b/Engine/Documentation/Source/Engine/Animation/AnimHowTo/LayerEditing/LayerEditingHowTo.INT.udn
--- a/Engine/Documentation/Source/Engine/Animation/AnimHowTo/LayerEditing/LayerEditingHowTo.INT.udn
+++ b/Engine/Documentation/Source/Engine/Animation/AnimHowTo/LayerEditing/LayerEditingHowTo.INT.udn
@@ -8,8 +8,10 @@ Related: Engine\Animation\AnimBlueprints
 Related: Engine\Animation\AnimHowTo\AdditiveAnimations
 Related: Engine\Animation\Skeleton
 Related: Engine\Animation\Overview
+version: 4.9
+skilllevel: Advanced
 
-**Animation Layer Editing** allows you to modify any existing [](Engine\Animation\Sequences) inside Unreal Engine 4 as part of Additive Layer Tracks. With Animation Layer Editing you can take an existing animation and adjust it to produce a different take on that animation, or you can create a separate, new animation with its own motion you provide based off the original. 
+Animation Layer Editing allows you to modify any existing [](Engine\Animation\Sequences) inside Unreal Engine 4 as part of Additive Layer Tracks. With Animation Layer Editing you can take an existing animation and adjust it to produce a different take on that animation, or you can create a separate, new animation with its own motion you provide based off the original. 
 
 This page will show you how to use the Animation Layer Editing tools and take an Idle animation and repurpose it into a Reload animation. 
 
@@ -23,15 +25,15 @@ For this guide, we created a New Project using the **Blueprint First Person** te
 
 In this section, we will create a copy of the Idle animation that we can then modify, leaving the original intact. 
 
-1.  Inside the **Content Browser**, under the **Game/FirstPersonBP/Animations** folder, open the **Prone_Idle** animation. 
+1.  Inside the **Content Browser**, under the **Game/FirstPerson/Animations** folder, open the **FirstPerson_Idle** animation. 
 
-	![](LE_FPIdleM.jpg)
+	![](AnimEdit1.png)
 
 1.	With the animation open, use the **Playback Controls** to pause the animation and set it to frame 0 so it is at the beginning. 
 
 	![](AnimEdit2.png)(w:640)
 
-1.	Click the **Create Asset** button from the Toolbar Menu, then under _Create Animation_, select **From Current Animation**. 
+1.	Click the **Create Asset** button from the Toolbar Menu, then under **Create Animation**, select **From Current Animation**. 
 
 	![](LE_FromCurrentAnimation.jpg)
 
@@ -47,7 +49,7 @@ In this section, we will create a copy of the Idle animation that we can then mo
 
 With the **FirstPerson_Reload** animation created, we can begin editing it through Animation Layer Editing. 
 
-1.  With the **FirstPerson_Reload** animation open, click the **ToNext** button in the **Playback Controls** and advance to frame 15. 
+1.  With the **FirstPerson_Reload** animation open, click the **To Next** button in the **Playback Controls** and advance to frame 15. 
 
 	![](AnimEdit5.png)
 
@@ -71,11 +73,11 @@ With the **FirstPerson_Reload** animation created, we can begin editing it throu
 
 	Above, three separate tracks were added for the **upperarm_l** Bone: a **Translation** track (1), a **Rotation** track (2) and a **Scale** track (3).
 
-1.	Click the **ToNext** button to advance to frame 30. 
+1.	Click the **To Next** button to advance to frame 30. 
 
 	![](AnimEdit9.png)
 
-1.	Inside the viewport, **Left-click** on the yellow arch of the **Rotation** widget. 
+1.	Inside the viewport, **Left-click** on the green arch of the **Rotation** widget. 
 
 	![](LE_LeftArmRot.jpg)
 
@@ -91,9 +93,9 @@ With the **FirstPerson_Reload** animation created, we can begin editing it throu
 
 	The new keys will be added to the Additive Layer Tracks and a curve from the first key to the second key will be created. 
 
-1.	Click the **ToNext** button on the **Playback Controls** to advance to frame 45. 
+1.	Click the **To Next** button on the **Playback Controls** to advance to frame 45. 
 
-1.	In the viewport, **Left-click** on and drag the yellow arch of the **Rotation** widget, then rotate to the right to the **60.00** mark. 
+1.	In the viewport, **Left-click** on and drag the green arch of the **Rotation** widget, then rotate to the right to the **60.00** mark. 
 
 	![](LE_LeftArmUp60.jpg)
 
@@ -103,13 +105,13 @@ With the **FirstPerson_Reload** animation created, we can begin editing it throu
 
 	The new keys will be added and should be the same values as the first set of keys. 
 
-1.	Click the **ToNext** button on the **Playback Controls** to advance to frame 50, then press **S** to add a set of keys. 
+1.	Click the **To Next** button on the **Playback Controls** to advance to frame 50, then press **S** to add a set of keys. 
 
 	![](AnimEdit17.png)(w:640)
 
-1.	Click the **ToNext** button on the **Playback Controls** to advance to frame 60. 
+1.	Click the **To Next** button on the **Playback Controls** to advance to frame 60. 
 
-1.	In the viewport, press **W** to switch to **Translation** mode, then **Left-click** and drag the yellow arrow to the right slightly. 
+1.	In the viewport, press **W** to switch to **Translation** mode, then **Left-click** and drag the red arrow to the right slightly. 
 
 	![](LE_LeftArmPump.jpg)
 
@@ -121,7 +123,7 @@ With the **FirstPerson_Reload** animation created, we can begin editing it throu
 
 	The keys for the pose will be added to the Additive Layer Tracks along with a curve from the previous keys. 
 
-1.	Click the **ToNext** button on the **Playback Controls** to advance to frame 65, then press **S** to add a set of keys. 
+1.	Click the **To Next** button on the **Playback Controls** to advance to frame 65, then press **S** to add a set of keys. 
 
 	![](AnimEdit22.png)(w:640)
 
@@ -137,8 +139,34 @@ With the **FirstPerson_Reload** animation created, we can begin editing it throu
 
 	When keys and curves are added to an animation, the result is an additive motion applied to the base animation. By clicking the **Apply** button, the key and curve data will be baked into the final animation. 
 
-1.	Click the **Save** button as your animation is complete. 
+1.	Click the **Save** button as your animation is complete.
  
+1. Inside of the **Game/FirstPerson/Animations** folder, **Double-click** the **FirstPerson_AnimBP** to open it.
+
+	![](AnimEdit28.png)
+
+1. Inside of the **Animation Blueprint Editor**, open the **Anim Graph** by **Double-clicking** it in the **My Blueprint** panel
+
+	![](AnimEdit29.png)
+
+1.  **Double-click** the **New State Machine** node in the **Anim Graph**.
+
+	![](AnimEdit30.png)
+
+1. Find the node titled **FPP_Idle** and **Double-click** to open it.
+
+	![](AnimEdit31.png)
+
+1. **Right-click** in the graph and create a **Play FirstPerson_Reload** node.
+
+	![](AnimEdit32.png)
+
+1. Replace the existing **Play FirstPerson_Idle** node with your new **Play FirstPerson_Reload** node by holding **Ctrl** and dragging while holding **Left-click** from the output of the **Play FirstPerson_Idle** node to the **Play FirstPerson_Reload** node.
+
+	![](AnimEdit33.png)
+
+1. **Compile** and **Save** the Animation Blueprint before closing it and then click **Play** to test out your new animation.
+
 Below is an example of the animation playing on a character in-game. 
 
 [OBJECT:EmbeddedVideo]
