UEDOC-1011

Images have been updated and will show live after the next publishing

diff --git a/Engine/Documentation/Source/Engine/Physics/PhysicalMaterials/PhysMatUserGuide/PhysMatUserGuide.INT.udn b/Engine/Documentation/Source/Engine/Physics/PhysicalMaterials/PhysMatUserGuide/PhysMatUserGuide.INT.udn
--- a/Engine/Documentation/Source/Engine/Physics/PhysicalMaterials/PhysMatUserGuide/PhysMatUserGuide.INT.udn
+++ b/Engine/Documentation/Source/Engine/Physics/PhysicalMaterials/PhysMatUserGuide/PhysMatUserGuide.INT.udn
@@ -52,15 +52,15 @@ This document will cover the creation and usage of **Physical Materials** as wel
 
 1. From the Content Browser. **click** ![](Engine/Audio\SoundCues\Editor\button_CB_NewAsset.png) **New Asset** -> **Physics** -> **Physical Material** or **right-click** in the **Content Browser** -> **Physics** -> **Physical Material**.
 
-	![](PhysMat_Creation_01.png)
+	![](PhysMatCreate.jpg)
 
 1. **Double click** the new Physical Material to edit its properties.
 
-	![](PhysMat_Creation_02.png)
+	![](PhysMatProperties.jpg)(w:250)
 
 1. **Adjust Properties**.
 
-	![](PhysMat_Creation_03.png)
+	![](AdjustProperties.jpg)
 
 1. **Click Save** ![](Shared\Icons\icon_SaveAsset_40x.png)
 
@@ -83,11 +83,11 @@ By default, Unreal Engine 4 supports 62 Surface Types, which can be labeled anyt
 
 1. **Select** the main material node.
 
-	![](PhysMat_Material_02.png)
+	![](PhysMatNone.jpg)
 
 1. **Change** the Physical Material.
 
-	![](PhysMat_Material_03.png)
+	![](PhysMatGlass2.jpg)
 
 
 
@@ -99,7 +99,7 @@ By default, Unreal Engine 4 supports 62 Surface Types, which can be labeled anyt
 
 1. **Change** the Physical Material.
 
-	![](PhysMat_MaterialInst_02.png)
+	![](StoneMat.jpg)
 
 
 
@@ -110,15 +110,15 @@ When adjusting the Physical Material for a Physics Asset, it is best to start by
 [EXCERPT:PhysMatSkelMesh1]
 1. **Open** a Physics Asset in the Physics Asset Tool.
 
-	![](PhysMat_PhAT_02.png)
+	![](PhysicsAsset1.jpg)
 
 1. **Select** a Physical Material in the Content Browser.
 
-	![](PhysMat_PhAT_03.png)
+	![](M_Glass.jpg)
 
-1. In the Physics Asset Tool, **click** the **Apply PhysMat button**: ![](Shared\Icons\icon_PhAT_PhysMat_40x.png).
+1. In the Physics Asset Tool, **click** the **Apply PhysMat** button from the **Asset** dropdown in the toolbar: ![](AssetButton.jpg).
 	[REGION:raw]
-	![](PhysMat_PhAT_01.png)(w:860)
+	![](ApplyPhysMat.jpg)
 	[/REGION]
 [/EXCERPT:PhysMatSkelMesh1]
 
@@ -131,7 +131,7 @@ If specific Physics Bodies need differing Physical Materials, you can then adjus
 
 
 [REGION:raw]
-![](PhysMat_PhAT_04.png)(w:920)
+![](SimpleCollision.jpg)(w:920)
 [/REGION]
 
 The default behavior for Physics interactions with Skeletal Meshes is to only interact with its associated Physics Asset, so the Skeletal Mesh's material's Physical Material will 
