#UE4doc #docs:
Changed font settings for multiple terms

diff --git a/Engine/Documentation/Source/Engine/Animation/Persona/Viewport/PersonaViewport.INT.udn b/Engine/Documentation/Source/Engine/Animation/Persona/Viewport/PersonaViewport.INT.udn
--- a/Engine/Documentation/Source/Engine/Animation/Persona/Viewport/PersonaViewport.INT.udn
+++ b/Engine/Documentation/Source/Engine/Animation/Persona/Viewport/PersonaViewport.INT.udn
@@ -33,14 +33,14 @@ The **Viewport** panel allows Animation Sequences, Blendspaces, and other aspect
 
 ![Viewport Panel with Callouts](ViewportPanel.png)(w:720)
 
-1. **Movement Toolbar**
-1. **View Options Toolbar** 
-1. **Viewport** 
-1. **Timeline and Control Bar**
+1. Movement Toolbar
+1. View Options Toolbar 
+1. Viewport 
+1. Timeline and Control Bar
 
 ## Movement Toolbar
 [EXCERPT:movement]
-The Movement Toolbar contains a collection of tools for moving, rotating, or scaling bones as well as how they are moved, rotated, or scaled with Grid Snapping. This toolbar also has an option for setting the relative coordinate system for movement or rotation as well as an option to set the camera movement speed when moving around in the viewport. 
+The **Movement Toolbar** contains a collection of tools for moving, rotating, or scaling bones as well as how they are moved, rotated, or scaled with Grid Snapping. This toolbar also has an option for setting the relative coordinate system for movement or rotation as well as an option to set the camera movement speed when moving around in the viewport. 
 
 | Option | Description |
 | ------ | ----------- |
@@ -58,11 +58,11 @@ The Movement Toolbar contains a collection of tools for moving, rotating, or sca
 ## View Options Toolbar
 
 [EXCERPT:viewoptions]
-The View Options Toolbar contains a series of buttons that hold options for what is displayed in the viewport.
+The **View Options Toolbar** contains a series of buttons that hold options for what is displayed in the viewport.
 
 ### Preview Mode
 
-The Preview Mode Button ![](Persona_PreviewMode1.png) provides the following options: 
+The **Preview** Mode Button ![](Persona_PreviewMode1.png) provides the following options: 
 
 | Option | Description |
 | ------ | ----------- |
@@ -73,7 +73,7 @@ The Preview Mode Button ![](Persona_PreviewMode1.png) provides the following opt
 
 ### Perspective
 
-The Perspective Button ![](Persona_Perspective.png) provides the following options: 
+The **Perspective** Button ![](Persona_Perspective.png) provides the following options: 
 
 | Option | Description |
 | ------ | ----------- |
@@ -85,7 +85,7 @@ The Perspective Button ![](Persona_Perspective.png) provides the following optio
 
 ### View Mode
 
-The View Mode Button ![](Persona_LitOptions.png) provides the following options: 
+The **View** Mode Button ![](Persona_LitOptions.png) provides the following options: 
 
 | Option | Description |
 | ------ | ----------- |
@@ -104,7 +104,7 @@ The View Mode Button ![](Persona_LitOptions.png) provides the following options:
 
 ### Show
 
-The Show Button ![](Persona_Show.png) provides the following options: 
+The **Show** Button ![](Persona_Show.png) provides the following options: 
 
 | Option | Description |
 | ------ | ----------- |
@@ -168,7 +168,7 @@ The Show Button ![](Persona_Show.png) provides the following options:
 
 ### Preview LODs
 
-The Preview LODs Button ![](Persona_LODAuto.png) provides the following options: 
+The **Preview LODs** Button ![](Persona_LODAuto.png) provides the following options: 
 
 | Option | Description |
 | ------ | ----------- |
@@ -191,7 +191,7 @@ _The added LOD(s) can either be removed or an FBX asset can be imported for the
 
 ### Playback Speed
 
-The Playback Speed Button ![](Persona_PlaybackSpeed.png) provides the following options: 
+The **Playback Speed** Button ![](Persona_PlaybackSpeed.png) provides the following options: 
 
 | Option | Description |
 | ------ | ----------- |
@@ -206,14 +206,14 @@ The Playback Speed Button ![](Persona_PlaybackSpeed.png) provides the following
 
 ## Viewport
 
-The viewport displays a preview of the Skeletal Mesh, any options set in the Toolbar and menu bar, and any selected animation.
+The **Viewport** displays a preview of the Skeletal Mesh, any options set in the Toolbar and menu bar, and any selected animation.
 
 ![Viewport Panel](ViewportPanel_Viewport.png)
 
 ## Timeline and Control Bar
 
 [EXCERPT:PreviewControls]
-The timeline and control bar contain the following controls:
+The **Timeline** and **Control bar** contain the following controls:
 
 | Control | Description |
 | ------- | ----------- |
