#UE4doc #docs:
Added engine version field
Added skill level field
Changed font settings for multiple terms

diff --git a/Engine/Documentation/Source/Engine/Animation/AnimHowTo/PreviewingAnimations/PreviewingHowTo.INT.udn b/Engine/Documentation/Source/Engine/Animation/AnimHowTo/PreviewingAnimations/PreviewingHowTo.INT.udn
--- a/Engine/Documentation/Source/Engine/Animation/AnimHowTo/PreviewingAnimations/PreviewingHowTo.INT.udn
+++ b/Engine/Documentation/Source/Engine/Animation/AnimHowTo/PreviewingAnimations/PreviewingHowTo.INT.udn
@@ -8,6 +8,8 @@ Related: Engine/Animation/Sequences
 Related: Engine/Animation/AnimMontage
 Related: Engine/Animation/Blendspaces
 Related: Engine/Animation/AimOffset
+version: 4.9
+skilllevel: Beginner
 
 
 This page covers the different methods in which you can load and preview animation assets inside the [Persona Animation Editor](Engine/Animation/Persona), which is Unreal Engine 4's Animation Editing Tool. The controls used to manipulate the playback or information displayed about the animation currently being previewed is also covered and additional links to more information are provided. 
@@ -68,7 +70,7 @@ You can also get a smaller, looping preview directly inside the **Asset Browser*
 
 ## Persona Viewport Controls
 
-Below are the controls for navigating the Persona Viewport, allowing you to move the camera around. 
+Below are the controls for navigating the **Persona Viewport**, allowing you to move the camera around. 
 
 ### Standard
 
@@ -83,7 +85,7 @@ These controls represent the default behavior when clicking and dragging in the
 
 ### Game-style
 
-The WASD controls will feel natural to those who are used to playing shooter games on the PC. They are enabled by default and can be used whenever you are holding the **Right Mouse Button** (RMB). This means you will still use **RMB** to turn the camera while navigating in this way.
+The **WASD** keys will feel natural to those who are used to playing shooter games on the PC. They are enabled by default and can be used whenever you are holding the **Right Mouse Button** (**RMB**). This means you will still use the mouse to turn the camera while navigating in this way.
 
 These controls are mirrored on the arrow keys and num pad to provide alternate access to them.
 
@@ -109,9 +111,9 @@ This section illustrates the options contained within the **Persona Viewport**.
 
 ![](PA_PersonaViewportOptions.jpg)
 
-1.	**View Options Toolbar**
-2.	**Movement Toolbar**
-3.	**Animation Display Information**
+1.	View Options Toolbar
+2.	Movement Toolbar
+3.	Animation Display Information
 
 ### View Options Toolbar
 
