Updating Gif on the page and converting it into a YouTube video

diff --git a/Engine/Documentation/Source/Engine/OpenWorldTools/ProceduralFoliage/QuickStart/ProceduralFoliageQuickStart.INT.udn b/Engine/Documentation/Source/Engine/OpenWorldTools/ProceduralFoliage/QuickStart/ProceduralFoliageQuickStart.INT.udn
--- a/Engine/Documentation/Source/Engine/OpenWorldTools/ProceduralFoliage/QuickStart/ProceduralFoliageQuickStart.INT.udn
+++ b/Engine/Documentation/Source/Engine/OpenWorldTools/ProceduralFoliage/QuickStart/ProceduralFoliageQuickStart.INT.udn
@@ -24,7 +24,17 @@ Before you can use the Procedural Foliage tools in your project, you must first
 1. Inside of  the Editor Preferences **right - click** on the **Experimental** section.
 1. Enable the Procedural Foliage option by clicking on the checkmark box next to the word **Procedural Foliage**.
 
-	![](Engine\OpenWorldTools\Overview\T_Enable_PFT.gif)(convert:false)
+	[OBJECT:EmbeddedVideo]
+	[PARAMLITERAL:width]
+	640
+	[/PARAMLITERAL]
+	[PARAMLITERAL:height]
+	360
+	[/PARAMLITERAL]
+	[PARAMLITERAL:videoid]
+	DzR6UNAIYbs
+	[/PARAMLITERAL]
+	[/OBJECT]
 
 You will also need to download the **Open World Demo Collection** content pack from the **Unreal Engine Marketplace** as some of the content from the collection will be used in the following Quick Start.
 Once the Open World Demo Collection is downloaded add it to the project that you are using to follow along with this Quick Start by doing the following.
