#UE4doc #docs:
Removed overview header
Updated engine version
Changed multiple terms to fit the UI of the engine
Changed some lines that referenced outdated file names.

diff --git a/Engine/Documentation/Source/Engine/Animation/RetargetingDifferentSkeletons/RetargetingDifferentSkeletons.INT.udn b/Engine/Documentation/Source/Engine/Animation/RetargetingDifferentSkeletons/RetargetingDifferentSkeletons.INT.udn
--- a/Engine/Documentation/Source/Engine/Animation/RetargetingDifferentSkeletons/RetargetingDifferentSkeletons.INT.udn
+++ b/Engine/Documentation/Source/Engine/Animation/RetargetingDifferentSkeletons/RetargetingDifferentSkeletons.INT.udn
@@ -2,7 +2,7 @@ Availability:Public
 Title: Animation Retargeting (Different Skeletons)
 Crumbs: %ROOT%, Engine, Engine/Animation
 Description:Retargeting animations for use between different skeletons. 
-Version: 4.5
+Version: 4.9
 
 [REGION:fullwidth]
 ![](Retarget1_1.jpg)
@@ -14,7 +14,7 @@ Version: 4.5
 Note: This feature is not available before Engine Version 4.5
 [/REGION]
 
-## Overview
+
 **Animation Retargeting (Different Skeletons)** is a feature that allows animations to be reused between characters that **do not** share the same Skeleton asset, unlike [](Engine\Animation\AnimationRetargeting) which involves retargeting animations between characters that share the same Skeleton asset. 
 
 The process of retargeting animations between characters that have different Skeletons uses an asset called a **Rig** to pass Bone information from one Skeleton to another. After selecting your Rig, you will share it between your Source Skeleton and a Target Skeleton (the Skeleton you wish to retarget animations to). Once you have done so, you will then match up the Target Skeleton's **Bones** to the **Nodes** stored in the Rig. 
@@ -32,7 +32,7 @@ We are also going to use two asset packs that are available for free via the **M
 [REGION:note]
 In this example, we are using the following assets which are available for download via the **Marketplace**.
 
-* The Source Skeleton we are using is the **HeroTPP_Skeleton** that is included in the **Animation Starter Pack**. 
+* The Source Skeleton we are using is the **UE4_Mannequin_Skeleton** that is included in the **Animation Starter Pack**. 
 * The Target Skeleton we are using is the **SWAT_Skeleton** that is included in the **Mixamo Character Pack**.
 
 Once you have installed the packs above, you can add them to a project via the **Add to project** button in the Unreal Launcher. 
@@ -46,9 +46,9 @@ Also, your version of Unreal Engine 4 may vary from the one depicted above.
 
 The steps below will show you how to access the default Rig and retarget animations: 
 
-1.	Inside the **Content Browser** under the **Game/AnimStarterPack/Character** folder, open the **HeroTPP_Skeleton** asset. 
+1.	Inside the **Content Browser** under the **Game/AnimStarterPack/UE4_Mannequin/Mesh** folder, open the **UE4_Mannequin_Skeleton** asset. 
 
-1.	Click the **Retarget Manager** button, then under **Set up Rig**, click the **Pick Asset** icon and select the **Humanoid** Rig. 
+1.	Click the **Retarget Manager** button, then under **Set up Rig**, click the **Pick Asset** icon and select **Set Humanoid Rig**. 
 
 	[REGION:lightbox]
 	[![](RT_Manager.jpg)(w:720)](RT_Manager.jpg)
@@ -60,7 +60,7 @@ The steps below will show you how to access the default Rig and retarget animati
 
 1.	In the **Content Browser** under the **Game/MixamoAnimPack/Mixamo_SWAT/Anims** folder, open the **SWAT_Skeleton** asset.
 
-1.	Click the **Retarget Manager** button, then under **Set up Rig**, click the **Pick Asset** icon and select the **Humanoid** Rig. 
+1.	Click the **Retarget Manager** button, then under **Set up Rig**, click the **Pick Asset** icon and select **Set Humanoid Rig**. 
 
 	[REGION:lightbox]
 	[![](RetargetAdd3.png)(w:720)](RetargetAdd3.png)
@@ -86,7 +86,7 @@ The steps below will show you how to access the default Rig and retarget animati
 
 1.	Continue to specify each Bone from the Target Skeleton to a corresponding Node on the Rig. 
 
-	The table below demonstrates the setup for the **HeroTPP_Skeleton** and **SWAT_Skeleton**
+	The table below demonstrates the setup for the **UE4_Mannequin_Skeleton** and **SWAT_Skeleton**
 
 	| Node (Rig) | Bone (Skeleton) |
 	| ------- | ------------ |
@@ -173,7 +173,7 @@ The steps below will show you how to access the default Rig and retarget animati
 
 1.	In the **Content Browser**, under **Game/AnimStarterPack**, select an animation that you want to retarget to the new Skeleton.
 
-1.	**Right-click** on the animation and choose **Retarget Animation Assets -> Duplicate Anim Assets and Retarget**. 
+1.	**Right-click** on the animation and choose **Retarget Anim Assets -> Duplicate Anim Assets and Retarget**. 
 
 	![](RT_DuplicateAnim.jpg)
 
