#UE4doc #docs:
Some changes were missing from previous submit of this file

diff --git a/Engine/Documentation/Source/Engine/Animation/AnimHowTo/Retargeting/RetargetingHowTo.INT.udn b/Engine/Documentation/Source/Engine/Animation/AnimHowTo/Retargeting/RetargetingHowTo.INT.udn
--- a/Engine/Documentation/Source/Engine/Animation/AnimHowTo/Retargeting/RetargetingHowTo.INT.udn
+++ b/Engine/Documentation/Source/Engine/Animation/AnimHowTo/Retargeting/RetargetingHowTo.INT.udn
@@ -6,6 +6,8 @@ Related: Engine\Animation\AnimationRetargeting
 Related: Engine\Animation\RetargetingDifferentSkeletons
 Related: Engine\Animation\Skeleton
 Related: Engine\Animation\Overview
+version: 4.9
+skilllevel: Intermediate
 
 **Animation Retargeting** is the process of repurposing existing animations for use between multiple characters, eliminating the need to create entirely new animations as you can instead share your animation assets between characters. 
 
