UEDOC-1079

Updated the imgae on this page and removed NewVar and NewVar0

diff --git a/Engine/Documentation/Source/Resources/Showcases/BlueprintExamples/Foliage/Foliage.INT.udn b/Engine/Documentation/Source/Resources/Showcases/BlueprintExamples/Foliage/Foliage.INT.udn
--- a/Engine/Documentation/Source/Resources/Showcases/BlueprintExamples/Foliage/Foliage.INT.udn
+++ b/Engine/Documentation/Source/Resources/Showcases/BlueprintExamples/Foliage/Foliage.INT.udn
@@ -41,7 +41,7 @@ The Components List, Construction Script, and Event Graph for this Blueprint per
 As you might expect, most of the work of this Blueprint is done in the **Construction Script**, which is quite extensive. Below are the key exposed properties as well as their purposes.
 
 
-![](Variables_Foliage.png) ![](Properties_Foliage.png)
+![](MyBlueprintNoVar.jpg) ![](Properties_Foliage.png)
 
 
 | Property | Purpose |
