Corrected minor spelling error: SAssingNew to SAssignNew

diff --git a/Engine/Documentation/Source/Programming/Slate/InGameUI/SlateInGameUI.INT.udn b/Engine/Documentation/Source/Programming/Slate/InGameUI/SlateInGameUI.INT.udn
--- a/Engine/Documentation/Source/Programming/Slate/InGameUI/SlateInGameUI.INT.udn
+++ b/Engine/Documentation/Source/Programming/Slate/InGameUI/SlateInGameUI.INT.udn
@@ -77,7 +77,7 @@ Creating the widget using `SNew()` when passing it to `GameViewportClient::AddVi
 Or using `SAssignNew()` to create it, assign it to a `TSharedPtr` member, and pass it in:
 
 	GEngine->GameViewport->AddViewportWidgetContent(
-		SAssingNew(MyWidgetPtr, SWeakWidget)
+		SAssignNew(MyWidgetPtr, SWeakWidget)
 		.PossiblyNullContent(MyWidgetClass)
 	);
 
