#UE4doc #docs
Added sectional Landing Pages to the C++ First Person Shooter Tutorial.
Availability: Docs

diff --git a/Engine/Documentation/Source/Programming/Tutorials/FirstPersonShooter/1/FirstPersonShooter_1.INT.udn b/Engine/Documentation/Source/Programming/Tutorials/FirstPersonShooter/1/FirstPersonShooter_1.INT.udn
new file mode 100644
--- /dev/null
+++ b/Engine/Documentation/Source/Programming/Tutorials/FirstPersonShooter/1/FirstPersonShooter_1.INT.udn
@@ -0,0 +1,40 @@
+Availability: Docs
+Title: 1 - Setting up your Project
+Crumbs: %ROOT%, GettingStarted, Programming/Tutorials
+Description: How to set up your First Person Shooter project.
+SkillLevel: Beginner
+Version: 4.9
+
+<!--TODO: Insert Cover Image-->
+
+## Goals
+
+The purpose of this section is to show you how to set up your First Person Shooter project.
+
+## Objectives
+
+By the end of this section of the tutorial you'll be able to:
+
+* Set up a New Project
+* Use the C++ Class Wizard
+* Add Log Messaging to your Project
+* Compile your First C++ Class
+* Set the Default Game Mode
+
+[OBJECT:TopicButtonList]
+	[PARAM:icon]
+		![](%ROOT%/build_icon.png)(convert:false)
+	[/PARAM]
+	[PARAM:title]
+		Steps
+	[/PARAM]
+	[PARAM:description]
+	[/PARAM]
+	[PARAM:links]
+	* [](Programming/Tutorials/FirstPersonShooter/1/1 "%Programming/Tutorials/FirstPersonShooter/1/1:description%")
+	* [](Programming/Tutorials/FirstPersonShooter/1/2 "%Programming/Tutorials/FirstPersonShooter/1/2:description%")
+	* [](Programming/Tutorials/FirstPersonShooter/1/3 "%Programming/Tutorials/FirstPersonShooter/1/3:description%")
+	* [](Programming/Tutorials/FirstPersonShooter/1/4 "%Programming/Tutorials/FirstPersonShooter/1/4:description%")
+	* [](Programming/Tutorials/FirstPersonShooter/1/5 "%Programming/Tutorials/FirstPersonShooter/1/5:description%")
+	[/PARAM]
+[/OBJECT]
\ No newline at end of file
diff --git a/Engine/Documentation/Source/Programming/Tutorials/FirstPersonShooter/2/FirstPersonShooter_2.INT.udn b/Engine/Documentation/Source/Programming/Tutorials/FirstPersonShooter/2/FirstPersonShooter_2.INT.udn
new file mode 100644
--- /dev/null
+++ b/Engine/Documentation/Source/Programming/Tutorials/FirstPersonShooter/2/FirstPersonShooter_2.INT.udn
@@ -0,0 +1,46 @@
+Availability: Docs
+Title: 2 - Implementing your Character
+Crumbs: %ROOT%, GettingStarted, Programming/Tutorials
+Description: How to implement your First Person Shooter character.
+SkillLevel: Beginner
+Version: 4.9
+
+<!--TODO: Insert Cover Image-->
+
+## Goals
+
+The purpose of this section is to show you how to implement your First Person Shooter character.
+
+## Objectives
+
+By the end of this section of the tutorial you'll be able to:
+
+* Make a New Character
+* Set-up Axis Mapping
+* Implement Character Movement Functions
+* Implement Mouse Camera Control
+* Implement Character Jumping
+* Add a Mesh to your Character
+* Change the Camera View
+* Add a First Person Mesh to your Character
+
+[OBJECT:TopicButtonList]
+	[PARAM:icon]
+		![](%ROOT%/build_icon.png)(convert:false)
+	[/PARAM]
+	[PARAM:title]
+		Steps
+	[/PARAM]
+	[PARAM:description]
+	[/PARAM]
+	[PARAM:links]
+	* [](Programming/Tutorials/FirstPersonShooter/2/1"%Programming/Tutorials/FirstPersonShooter/2/1:description%")
+	* [](Programming/Tutorials/FirstPersonShooter/2/2"%Programming/Tutorials/FirstPersonShooter/2/2:description%")
+	* [](Programming/Tutorials/FirstPersonShooter/2/3"%Programming/Tutorials/FirstPersonShooter/2/3:description%")
+	* [](Programming/Tutorials/FirstPersonShooter/2/4"%Programming/Tutorials/FirstPersonShooter/2/4:description%")
+	* [](Programming/Tutorials/FirstPersonShooter/2/5"%Programming/Tutorials/FirstPersonShooter/2/5:description%")
+	* [](Programming/Tutorials/FirstPersonShooter/2/6"%Programming/Tutorials/FirstPersonShooter/2/6:description%")
+	* [](Programming/Tutorials/FirstPersonShooter/2/7"%Programming/Tutorials/FirstPersonShooter/2/7:description%")
+	* [](Programming/Tutorials/FirstPersonShooter/2/8"%Programming/Tutorials/FirstPersonShooter/2/8:description%")
+	[/PARAM]
+[/OBJECT]
\ No newline at end of file
diff --git a/Engine/Documentation/Source/Programming/Tutorials/FirstPersonShooter/3/FirstPersonShooter_3.INT.udn b/Engine/Documentation/Source/Programming/Tutorials/FirstPersonShooter/3/FirstPersonShooter_3.INT.udn
new file mode 100644
--- /dev/null
+++ b/Engine/Documentation/Source/Programming/Tutorials/FirstPersonShooter/3/FirstPersonShooter_3.INT.udn
@@ -0,0 +1,40 @@
+Availability: Docs
+Title: 3 - Implementing Projectiles
+Crumbs: %ROOT%, GettingStarted, Programming/Tutorials
+Description: How to implement projectiles for your First Person Shooter game.
+SkillLevel: Beginner
+Version: 4.9
+
+<!--TODO: Insert Cover Image-->
+
+## Goals
+
+The purpose of this section is to show you how to implement projectiles for your First Person Shooter game.
+
+## Objectives
+
+By the end of this section of the tutorial you'll be able to:
+
+* Add Projectiles to your Game
+* Implement Shooting
+* Set-up Projectile Collision and Lifetime
+* Get your Projectiles to Interact with the World
+* Add Crosshairs to your Viewport
+
+[OBJECT:TopicButtonList]
+	[PARAM:icon]
+		![](%ROOT%/build_icon.png)(convert:false)
+	[/PARAM]
+	[PARAM:title]
+		Steps
+	[/PARAM]
+	[PARAM:description]
+	[/PARAM]
+	[PARAM:links]
+	* [](Programming/Tutorials/FirstPersonShooter/3/1 "%Programming/Tutorials/FirstPersonShooter/3/1:description%")
+	* [](Programming/Tutorials/FirstPersonShooter/3/2 "%Programming/Tutorials/FirstPersonShooter/3/2:description%")
+	* [](Programming/Tutorials/FirstPersonShooter/3/3 "%Programming/Tutorials/FirstPersonShooter/3/3:description%")
+	* [](Programming/Tutorials/FirstPersonShooter/3/4 "%Programming/Tutorials/FirstPersonShooter/3/4:description%")
+	* [](Programming/Tutorials/FirstPersonShooter/3/5 "%Programming/Tutorials/FirstPersonShooter/3/5:description%")
+	[/PARAM]
+[/OBJECT]
\ No newline at end of file
diff --git a/Engine/Documentation/Source/Programming/Tutorials/FirstPersonShooter/4/FirstPersonShooter_4.INT.udn b/Engine/Documentation/Source/Programming/Tutorials/FirstPersonShooter/4/FirstPersonShooter_4.INT.udn
new file mode 100644
--- /dev/null
+++ b/Engine/Documentation/Source/Programming/Tutorials/FirstPersonShooter/4/FirstPersonShooter_4.INT.udn
@@ -0,0 +1,46 @@
+Availability: Docs
+Title: 4 - Adding Character Animation
+Crumbs: %ROOT%, GettingStarted, Programming/Tutorials
+Description: How to add animations to your First Person Shooter character.
+SkillLevel: Beginner
+Version: 4.9
+
+<!--TODO: Insert Cover Image-->
+
+## Goals
+
+The purpose of this section is to show you how to add animations for your First Person Shooter character.
+
+## Objectives
+
+By the end of this section of the tutorial you'll be able to:
+
+* Animate your Character
+* Set-up your Event Graph
+* Add an Animation State Machine
+* Add Transition States to your Animation
+* Associate the Animation Blueprint with the Character Blueprint
+
+[OBJECT:TopicButtonList]
+	[PARAM:icon]
+		![](%ROOT%/build_icon.png)(convert:false)
+	[/PARAM]
+	[PARAM:title]
+		Steps
+	[/PARAM]
+	[PARAM:description]
+	[/PARAM]
+	[PARAM:links]
+	* [](Programming/Tutorials/FirstPersonShooter/4/1 "%Programming/Tutorials/FirstPersonShooter/4/1:description%")
+	* [](Programming/Tutorials/FirstPersonShooter/4/2 "%Programming/Tutorials/FirstPersonShooter/4/2:description%")
+	* [](Programming/Tutorials/FirstPersonShooter/4/3 "%Programming/Tutorials/FirstPersonShooter/4/3:description%")
+	* [](Programming/Tutorials/FirstPersonShooter/4/4 "%Programming/Tutorials/FirstPersonShooter/4/4:description%")
+		* [](Programming/Tutorials/FirstPersonShooter/4/4/1 "%Programming/Tutorials/FirstPersonShooter/4/4/1:description%")
+		* [](Programming/Tutorials/FirstPersonShooter/4/4/2 "%Programming/Tutorials/FirstPersonShooter/4/4/2:description%")
+		* [](Programming/Tutorials/FirstPersonShooter/4/4/3 "%Programming/Tutorials/FirstPersonShooter/4/4/3:description%")
+		* [](Programming/Tutorials/FirstPersonShooter/4/4/4 "%Programming/Tutorials/FirstPersonShooter/4/4/4:description%")
+		* [](Programming/Tutorials/FirstPersonShooter/4/4/5 "%Programming/Tutorials/FirstPersonShooter/4/4/5:description%")
+		* [](Programming/Tutorials/FirstPersonShooter/4/4/6 "%Programming/Tutorials/FirstPersonShooter/4/4/6:description%")
+	* [](Programming/Tutorials/FirstPersonShooter/4/5 "%Programming/Tutorials/FirstPersonShooter/4/5:description%")
+	[/PARAM]
+[/OBJECT]
\ No newline at end of file
