#UE4doc #docs
Added links to C++ First Person Shooter Template.
Availability: Docs

diff --git a/Engine/Documentation/Source/Programming/Tutorials/FirstPersonShooter/FirstPersonShooter.INT.udn b/Engine/Documentation/Source/Programming/Tutorials/FirstPersonShooter/FirstPersonShooter.INT.udn
--- a/Engine/Documentation/Source/Programming/Tutorials/FirstPersonShooter/FirstPersonShooter.INT.udn
+++ b/Engine/Documentation/Source/Programming/Tutorials/FirstPersonShooter/FirstPersonShooter.INT.udn
@@ -25,14 +25,14 @@ By the end of this tutorial you'll be able to:
 		![](%ROOT%/build_icon.png)(convert:false)
 	[/PARAM]
 	[PARAM:title]
-		Steps
+		Sections
 	[/PARAM]
 	[PARAM:description]
 	[/PARAM]
 	[PARAM:links]
-	* **Project Setup**
-	* **Character Implementation**
-	* **Projectile Implementation**
-	* **Character Animation**
+	* [](Programming/Tutorials/FirstPersonShooter/1 "%Programming/Tutorials/FirstPersonShooter/1:description%")
+	* [](Programming/Tutorials/FirstPersonShooter/2 "%Programming/Tutorials/FirstPersonShooter/2:description%")
+	* [](Programming/Tutorials/FirstPersonShooter/3 "%Programming/Tutorials/FirstPersonShooter/3:description%")
+	* [](Programming/Tutorials/FirstPersonShooter/4 "%Programming/Tutorials/FirstPersonShooter/4:description%")
 	[/PARAM]
 [/OBJECT]
