#UE4doc #docs
Added Section 1 subpages without content.
Availability: Docs

diff --git a/Engine/Documentation/Source/Programming/Tutorials/FirstPersonShooter/1/1/FirstPersonShooter_1_1.INT.udn b/Engine/Documentation/Source/Programming/Tutorials/FirstPersonShooter/1/1/FirstPersonShooter_1_1.INT.udn
new file mode 100644
--- /dev/null
+++ b/Engine/Documentation/Source/Programming/Tutorials/FirstPersonShooter/1/1/FirstPersonShooter_1_1.INT.udn
@@ -0,0 +1,27 @@
+Availability: Docs
+Title: 1.1 - Setting up your Project
+Crumbs: %ROOT%, GettingStarted, Engine/Content/QuickStart
+Description: Learn how to set up your First Person Shooter project.
+SkillLevel: Beginner
+Version: 4.9
+
+[VAR:Steps]
+[OBJECT:Navigation]
+	[PARAM:previous]
+	[/PARAM]
+	[PARAM:current]
+	[/PARAM]
+	[PARAM:home]
+		[Section Home](Programming/Tutorials/FirstPersonShooter/1)
+	[/PARAM]
+	[PARAM:next]
+		[Next Step](Programming/Tutorials/FirstPersonShooter/1/2)
+	[/PARAM]
+[/OBJECT]
+[/VAR]
+
+%Steps%
+
+<!--TODO: Add tutorial steps here-->
+
+%Steps%
\ No newline at end of file
diff --git a/Engine/Documentation/Source/Programming/Tutorials/FirstPersonShooter/1/2/FirstPersonShooter_1_2.INT.udn b/Engine/Documentation/Source/Programming/Tutorials/FirstPersonShooter/1/2/FirstPersonShooter_1_2.INT.udn
new file mode 100644
--- /dev/null
+++ b/Engine/Documentation/Source/Programming/Tutorials/FirstPersonShooter/1/2/FirstPersonShooter_1_2.INT.udn
@@ -0,0 +1,28 @@
+Availability: Docs
+Title: 1.2 - Using the C++ Class Wizard
+Crumbs: %ROOT%, GettingStarted, Engine/Content/QuickStart
+Description: Learn how to create your first C++ class using the C++ Wizard.
+SkillLevel: Beginner
+Version: 4.9
+
+[VAR:Steps]
+[OBJECT:Navigation]
+	[PARAM:previous]
+		[Previous Step](Programming/Tutorials/FirstPersonShooter/1/1)
+	[/PARAM]
+	[PARAM:current]
+	[/PARAM]
+	[PARAM:home]
+		[Section Home](Programming/Tutorials/FirstPersonShooter/1)
+	[/PARAM]
+	[PARAM:next]
+		[Next Step](Programming/Tutorials/FirstPersonShooter/1/3)
+	[/PARAM]
+[/OBJECT]
+[/VAR]
+
+%Steps%
+
+<!--TODO: Add tutorial steps here-->
+
+%Steps%
\ No newline at end of file
diff --git a/Engine/Documentation/Source/Programming/Tutorials/FirstPersonShooter/1/3/FirstPersonShooter_1_3.INT.udn b/Engine/Documentation/Source/Programming/Tutorials/FirstPersonShooter/1/3/FirstPersonShooter_1_3.INT.udn
new file mode 100644
--- /dev/null
+++ b/Engine/Documentation/Source/Programming/Tutorials/FirstPersonShooter/1/3/FirstPersonShooter_1_3.INT.udn
@@ -0,0 +1,28 @@
+Availability: Docs
+Title: 1.3 - Adding Log Messaging
+Crumbs: %ROOT%, GettingStarted, Engine/Content/QuickStart
+Description: Learn how to add log messaging to your game.
+SkillLevel: Beginner
+Version: 4.9
+
+[VAR:Steps]
+[OBJECT:Navigation]
+	[PARAM:previous]
+		[Previous Step](Programming/Tutorials/FirstPersonShooter/1/2)
+	[/PARAM]
+	[PARAM:current]
+	[/PARAM]
+	[PARAM:home]
+		[Section Home](Programming/Tutorials/FirstPersonShooter/1)
+	[/PARAM]
+	[PARAM:next]
+		[Next Step](Programming/Tutorials/FirstPersonShooter/1/4)
+	[/PARAM]
+[/OBJECT]
+[/VAR]
+
+%Steps%
+
+<!--TODO: Add tutorial steps here-->
+
+%Steps%
\ No newline at end of file
diff --git a/Engine/Documentation/Source/Programming/Tutorials/FirstPersonShooter/1/4/FirstPersonShooter_1_4.INT.udn b/Engine/Documentation/Source/Programming/Tutorials/FirstPersonShooter/1/4/FirstPersonShooter_1_4.INT.udn
new file mode 100644
--- /dev/null
+++ b/Engine/Documentation/Source/Programming/Tutorials/FirstPersonShooter/1/4/FirstPersonShooter_1_4.INT.udn
@@ -0,0 +1,28 @@
+Availability: Docs
+Title: 1.4 - Compiling the Project
+Crumbs: %ROOT%, GettingStarted, Engine/Content/QuickStart
+Description: Learn how to compile your first C++ class.
+SkillLevel: Beginner
+Version: 4.9
+
+[VAR:Steps]
+[OBJECT:Navigation]
+	[PARAM:previous]
+		[Previous Step](Programming/Tutorials/FirstPersonShooter/1/3)
+	[/PARAM]
+	[PARAM:current]
+	[/PARAM]
+	[PARAM:home]
+		[Section Home](Programming/Tutorials/FirstPersonShooter/1)
+	[/PARAM]
+	[PARAM:next]
+		[Next Step](Programming/Tutorials/FirstPersonShooter/1/5)
+	[/PARAM]
+[/OBJECT]
+[/VAR]
+
+%Steps%
+
+<!--TODO: Add tutorial steps here-->
+
+%Steps%
\ No newline at end of file
diff --git a/Engine/Documentation/Source/Programming/Tutorials/FirstPersonShooter/1/5/FirstPersonShooter_1_5.INT.udn b/Engine/Documentation/Source/Programming/Tutorials/FirstPersonShooter/1/5/FirstPersonShooter_1_5.INT.udn
new file mode 100644
--- /dev/null
+++ b/Engine/Documentation/Source/Programming/Tutorials/FirstPersonShooter/1/5/FirstPersonShooter_1_5.INT.udn
@@ -0,0 +1,28 @@
+Availability: Docs
+Title: 1.5 - Setting the Default Game Mode
+Crumbs: %ROOT%, GettingStarted, Engine/Content/QuickStart
+Description: Learn how to set the default game mode for your First Person Shooter project.
+SkillLevel: Beginner
+Version: 4.9
+
+[VAR:Steps]
+[OBJECT:Navigation]
+	[PARAM:previous]
+		[Previous Step](Programming/Tutorials/FirstPersonShooter/1/4)
+	[/PARAM]
+	[PARAM:current]
+	[/PARAM]
+	[PARAM:home]
+		[Section Home](Programming/Tutorials/FirstPersonShooter/1)
+	[/PARAM]
+	[PARAM:next]
+		[Next Section](Programming/Tutorials/FirstPersonShooter/2)
+	[/PARAM]
+[/OBJECT]
+[/VAR]
+
+%Steps%
+
+<!--TODO: Add tutorial steps here-->
+
+%Steps%
\ No newline at end of file
