#UE4doc #docs
Added Section 2 subpages without content.
Availability: Docs

diff --git a/Engine/Documentation/Source/Programming/Tutorials/FirstPersonShooter/2/1/FirstPersonShooter_2_1.INT.udn b/Engine/Documentation/Source/Programming/Tutorials/FirstPersonShooter/2/1/FirstPersonShooter_2_1.INT.udn
new file mode 100644
--- /dev/null
+++ b/Engine/Documentation/Source/Programming/Tutorials/FirstPersonShooter/2/1/FirstPersonShooter_2_1.INT.udn
@@ -0,0 +1,27 @@
+Availability: Docs
+Title: 2.1 - Making a New Character
+Crumbs: %ROOT%, GettingStarted, Engine/Content/QuickStart
+Description: Learn how to make a new character for your First Person Shooter project.
+SkillLevel: Beginner
+Version: 4.9
+
+[VAR:Steps]
+[OBJECT:Navigation]
+	[PARAM:previous]
+	[/PARAM]
+	[PARAM:current]
+	[/PARAM]
+	[PARAM:home]
+		[Section Home](Programming/Tutorials/FirstPersonShooter/2)
+	[/PARAM]
+	[PARAM:next]
+		[Next Step](Programming/Tutorials/FirstPersonShooter/2/2)
+	[/PARAM]
+[/OBJECT]
+[/VAR]
+
+%Steps%
+
+<!--TODO: Add tutorial steps here-->
+
+%Steps%
\ No newline at end of file
diff --git a/Engine/Documentation/Source/Programming/Tutorials/FirstPersonShooter/2/2/FirstPersonShooter_2_2.INT.udn b/Engine/Documentation/Source/Programming/Tutorials/FirstPersonShooter/2/2/FirstPersonShooter_2_2.INT.udn
new file mode 100644
--- /dev/null
+++ b/Engine/Documentation/Source/Programming/Tutorials/FirstPersonShooter/2/2/FirstPersonShooter_2_2.INT.udn
@@ -0,0 +1,28 @@
+Availability: Docs
+Title: 2.2 - Setting up Axis Mapping
+Crumbs: %ROOT%, GettingStarted, Engine/Content/QuickStart
+Description: Learn how to set-up axis mapping for your First Person Shooter character.
+SkillLevel: Beginner
+Version: 4.9
+
+[VAR:Steps]
+[OBJECT:Navigation]
+	[PARAM:previous]
+		[Previous Step](Programming/Tutorials/FirstPersonShooter/2/1)
+	[/PARAM]
+	[PARAM:current]
+	[/PARAM]
+	[PARAM:home]
+		[Section Home](Programming/Tutorials/FirstPersonShooter/2)
+	[/PARAM]
+	[PARAM:next]
+		[Next Step](Programming/Tutorials/FirstPersonShooter/2/3)
+	[/PARAM]
+[/OBJECT]
+[/VAR]
+
+%Steps%
+
+<!--TODO: Add tutorial steps here-->
+
+%Steps%
\ No newline at end of file
diff --git a/Engine/Documentation/Source/Programming/Tutorials/FirstPersonShooter/2/3/FirstPersonShooter_2_3.INT.udn b/Engine/Documentation/Source/Programming/Tutorials/FirstPersonShooter/2/3/FirstPersonShooter_2_3.INT.udn
new file mode 100644
--- /dev/null
+++ b/Engine/Documentation/Source/Programming/Tutorials/FirstPersonShooter/2/3/FirstPersonShooter_2_3.INT.udn
@@ -0,0 +1,28 @@
+Availability: Docs
+Title: 2.3 - Implementing Character Movement Functions
+Crumbs: %ROOT%, GettingStarted, Engine/Content/QuickStart
+Description: Learn how to implement movement functions for your First Person Shooter character.
+SkillLevel: Beginner
+Version: 4.9
+
+[VAR:Steps]
+[OBJECT:Navigation]
+	[PARAM:previous]
+		[Previous Step](Programming/Tutorials/FirstPersonShooter/2/2)
+	[/PARAM]
+	[PARAM:current]
+	[/PARAM]
+	[PARAM:home]
+		[Section Home](Programming/Tutorials/FirstPersonShooter/2)
+	[/PARAM]	
+	[PARAM:next]
+		[Next Step](Programming/Tutorials/FirstPersonShooter/2/4)
+	[/PARAM]
+[/OBJECT]
+[/VAR]
+
+%Steps%
+
+<!--TODO: Add tutorial steps here-->
+
+%Steps%
\ No newline at end of file
diff --git a/Engine/Documentation/Source/Programming/Tutorials/FirstPersonShooter/2/4/FirstPersonShooter_2_4.INT.udn b/Engine/Documentation/Source/Programming/Tutorials/FirstPersonShooter/2/4/FirstPersonShooter_2_4.INT.udn
new file mode 100644
--- /dev/null
+++ b/Engine/Documentation/Source/Programming/Tutorials/FirstPersonShooter/2/4/FirstPersonShooter_2_4.INT.udn
@@ -0,0 +1,28 @@
+Availability: Docs
+Title: 2.4 - Implementing Mouse Camera Control
+Crumbs: %ROOT%, GettingStarted, Engine/Content/QuickStart
+Description: Learn how to implement mouse camera control for your First Person Shooter character.
+SkillLevel: Beginner
+Version: 4.9
+
+[VAR:Steps]
+[OBJECT:Navigation]
+	[PARAM:previous]
+		[Previous Step](Programming/Tutorials/FirstPersonShooter/2/3)
+	[/PARAM]
+	[PARAM:current]
+	[/PARAM]
+	[PARAM:home]
+		[Section Home](Programming/Tutorials/FirstPersonShooter/2)
+	[/PARAM]
+	[PARAM:next]
+		[Next Step](Programming/Tutorials/FirstPersonShooter/2/5)
+	[/PARAM]
+[/OBJECT]
+[/VAR]
+
+%Steps%
+
+<!--TODO: Add tutorial steps here-->
+
+%Steps%
\ No newline at end of file
diff --git a/Engine/Documentation/Source/Programming/Tutorials/FirstPersonShooter/2/5/FirstPersonShooter_2_5.INT.udn b/Engine/Documentation/Source/Programming/Tutorials/FirstPersonShooter/2/5/FirstPersonShooter_2_5.INT.udn
new file mode 100644
--- /dev/null
+++ b/Engine/Documentation/Source/Programming/Tutorials/FirstPersonShooter/2/5/FirstPersonShooter_2_5.INT.udn
@@ -0,0 +1,28 @@
+Availability: Docs
+Title: 2.5 - Implementing Character Jumping
+Crumbs: %ROOT%, GettingStarted, Engine/Content/QuickStart
+Description: Learn how to implement jumping for your First Person Shooter character.
+SkillLevel: Beginner
+Version: 4.9
+
+[VAR:Steps]
+[OBJECT:Navigation]
+	[PARAM:previous]
+		[Previous Step](Programming/Tutorials/FirstPersonShooter/2/4)
+	[/PARAM]
+	[PARAM:current]
+	[/PARAM]
+	[PARAM:home]
+		[Section Home](Programming/Tutorials/FirstPersonShooter/2)
+	[/PARAM]
+	[PARAM:next]
+		[Next Step](Programming/Tutorials/FirstPersonShooter/2/6)
+	[/PARAM]
+[/OBJECT]
+[/VAR]
+
+%Steps%
+
+<!--TODO: Add tutorial steps here-->
+
+%Steps%
\ No newline at end of file
diff --git a/Engine/Documentation/Source/Programming/Tutorials/FirstPersonShooter/2/6/FirstPersonShooter_2_6.INT.udn b/Engine/Documentation/Source/Programming/Tutorials/FirstPersonShooter/2/6/FirstPersonShooter_2_6.INT.udn
new file mode 100644
--- /dev/null
+++ b/Engine/Documentation/Source/Programming/Tutorials/FirstPersonShooter/2/6/FirstPersonShooter_2_6.INT.udn
@@ -0,0 +1,28 @@
+Availability: Docs
+Title: 2.6 - Adding a Mesh to your Character
+Crumbs: %ROOT%, GettingStarted, Engine/Content/QuickStart
+Description: Learn how to add a mesh to your First Person Shooter character.
+SkillLevel: Beginner
+Version: 4.9
+
+[VAR:Steps]
+[OBJECT:Navigation]
+	[PARAM:previous]
+		[Previous Step](Programming/Tutorials/FirstPersonShooter/2/5)
+	[/PARAM]
+	[PARAM:current]
+	[/PARAM]
+	[PARAM:home]
+		[Section Home](Programming/Tutorials/FirstPersonShooter/2)
+	[/PARAM]
+	[PARAM:next]
+		[Next Step](Programming/Tutorials/FirstPersonShooter/2/7)
+	[/PARAM]
+[/OBJECT]
+[/VAR]
+
+%Steps%
+
+<!--TODO: Add tutorial steps here-->
+
+%Steps%
\ No newline at end of file
diff --git a/Engine/Documentation/Source/Programming/Tutorials/FirstPersonShooter/2/7/FirstPersonShooter_2_7.INT.udn b/Engine/Documentation/Source/Programming/Tutorials/FirstPersonShooter/2/7/FirstPersonShooter_2_7.INT.udn
new file mode 100644
--- /dev/null
+++ b/Engine/Documentation/Source/Programming/Tutorials/FirstPersonShooter/2/7/FirstPersonShooter_2_7.INT.udn
@@ -0,0 +1,28 @@
+Availability: Docs
+Title: 2.7 - Changing the Camera View
+Crumbs: %ROOT%, GettingStarted, Engine/Content/QuickStart
+Description: Learn how to change the camera view for your First Person Shooter character.
+SkillLevel: Beginner
+Version: 4.9
+
+[VAR:Steps]
+[OBJECT:Navigation]
+	[PARAM:previous]
+		[Previous Step](Programming/Tutorials/FirstPersonShooter/2/6)
+	[/PARAM]
+	[PARAM:current]
+	[/PARAM]
+	[PARAM:home]
+		[Section Home](Programming/Tutorials/FirstPersonShooter/2)
+	[/PARAM]
+	[PARAM:next]
+		[Next Step](Programming/Tutorials/FirstPersonShooter/2/8)
+	[/PARAM]
+[/OBJECT]
+[/VAR]
+
+%Steps%
+
+<!--TODO: Add tutorial steps here-->
+
+%Steps%
\ No newline at end of file
diff --git a/Engine/Documentation/Source/Programming/Tutorials/FirstPersonShooter/2/8/FirstPersonShooter_2_8.INT.udn b/Engine/Documentation/Source/Programming/Tutorials/FirstPersonShooter/2/8/FirstPersonShooter_2_8.INT.udn
new file mode 100644
--- /dev/null
+++ b/Engine/Documentation/Source/Programming/Tutorials/FirstPersonShooter/2/8/FirstPersonShooter_2_8.INT.udn
@@ -0,0 +1,28 @@
+Availability: Docs
+Title: 2.8 - Add a First Person Mesh to your Character
+Crumbs: %ROOT%, GettingStarted, Engine/Content/QuickStart
+Description: Learn how to add a first person mesh to your First Person Shooter character.
+SkillLevel: Beginner
+Version: 4.9
+
+[VAR:Steps]
+[OBJECT:Navigation]
+	[PARAM:previous]
+		[Previous Step](Programming/Tutorials/FirstPersonShooter/2/7)
+	[/PARAM]
+	[PARAM:current]
+	[/PARAM]
+	[PARAM:home]
+		[Section Home](Programming/Tutorials/FirstPersonShooter/2)
+	[/PARAM]
+	[PARAM:next]
+		[Next Section](Programming/Tutorials/FirstPersonShooter/3)
+	[/PARAM]
+[/OBJECT]
+[/VAR]
+
+%Steps%
+
+<!--TODO: Add tutorial steps here-->
+
+%Steps%
\ No newline at end of file
