#UE4doc #docs
Added Section 3 subpages without content.
Availability: Docs

diff --git a/Engine/Documentation/Source/Programming/Tutorials/FirstPersonShooter/3/1/FirstPersonShooter_3_1.INT.udn b/Engine/Documentation/Source/Programming/Tutorials/FirstPersonShooter/3/1/FirstPersonShooter_3_1.INT.udn
new file mode 100644
--- /dev/null
+++ b/Engine/Documentation/Source/Programming/Tutorials/FirstPersonShooter/3/1/FirstPersonShooter_3_1.INT.udn
@@ -0,0 +1,27 @@
+Availability: Docs
+Title: 3.1 - Adding Projectiles to your Game
+Crumbs: %ROOT%, GettingStarted, Engine/Content/QuickStart
+Description: Learn how to add projectiles to your First Person Shooter project.
+SkillLevel: Beginner
+Version: 4.9
+
+[VAR:Steps]
+[OBJECT:Navigation]
+	[PARAM:previous]
+	[/PARAM]
+	[PARAM:current]
+	[/PARAM]
+	[PARAM:home]
+		[Section Home](Programming/Tutorials/FirstPersonShooter/3)
+	[/PARAM]
+	[PARAM:next]
+		[Next Step](Programming/Tutorials/FirstPersonShooter/3/2)
+	[/PARAM]
+[/OBJECT]
+[/VAR]
+
+%Steps%
+
+<!--TODO: Add tutorial steps here-->
+
+%Steps%
\ No newline at end of file
diff --git a/Engine/Documentation/Source/Programming/Tutorials/FirstPersonShooter/3/2/FirstPersonShooter_3_2.INT.udn b/Engine/Documentation/Source/Programming/Tutorials/FirstPersonShooter/3/2/FirstPersonShooter_3_2.INT.udn
new file mode 100644
--- /dev/null
+++ b/Engine/Documentation/Source/Programming/Tutorials/FirstPersonShooter/3/2/FirstPersonShooter_3_2.INT.udn
@@ -0,0 +1,28 @@
+Availability: Docs
+Title: 3.2 - Implementing Shooting
+Crumbs: %ROOT%, GettingStarted, Engine/Content/QuickStart
+Description: Learn how to shooting for your First Person Shooter character.
+SkillLevel: Beginner
+Version: 4.9
+
+[VAR:Steps]
+[OBJECT:Navigation]
+	[PARAM:previous]
+		[Previous Step](Programming/Tutorials/FirstPersonShooter/3/1)
+	[/PARAM]
+	[PARAM:current]
+	[/PARAM]
+	[PARAM:home]
+		[Section Home](Programming/Tutorials/FirstPersonShooter/3)
+	[/PARAM]
+	[PARAM:next]
+		[Next Step](Programming/Tutorials/FirstPersonShooter/3/3)
+	[/PARAM]
+[/OBJECT]
+[/VAR]
+
+%Steps%
+
+<!--TODO: Add tutorial steps here-->
+
+%Steps%
\ No newline at end of file
diff --git a/Engine/Documentation/Source/Programming/Tutorials/FirstPersonShooter/3/3/FirstPersonShooter_3_3.INT.udn b/Engine/Documentation/Source/Programming/Tutorials/FirstPersonShooter/3/3/FirstPersonShooter_3_3.INT.udn
new file mode 100644
--- /dev/null
+++ b/Engine/Documentation/Source/Programming/Tutorials/FirstPersonShooter/3/3/FirstPersonShooter_3_3.INT.udn
@@ -0,0 +1,28 @@
+Availability: Docs
+Title: 3.3 - Setting up Projectile Collision and Lifetime
+Crumbs: %ROOT%, GettingStarted, Engine/Content/QuickStart
+Description: Learn how to set-up collision and lifetime for your First Person Shooter projectiles.
+SkillLevel: Beginner
+Version: 4.9
+
+[VAR:Steps]
+[OBJECT:Navigation]
+	[PARAM:previous]
+		[Previous Step](Programming/Tutorials/FirstPersonShooter/3/2)
+	[/PARAM]
+	[PARAM:current]
+	[/PARAM]
+	[PARAM:home]
+		[Section Home](Programming/Tutorials/FirstPersonShooter/3)
+	[/PARAM]	
+	[PARAM:next]
+		[Next Step](Programming/Tutorials/FirstPersonShooter/3/4)
+	[/PARAM]
+[/OBJECT]
+[/VAR]
+
+%Steps%
+
+<!--TODO: Add tutorial steps here-->
+
+%Steps%
\ No newline at end of file
diff --git a/Engine/Documentation/Source/Programming/Tutorials/FirstPersonShooter/3/4/FirstPersonShooter_3_4.INT.udn b/Engine/Documentation/Source/Programming/Tutorials/FirstPersonShooter/3/4/FirstPersonShooter_3_4.INT.udn
new file mode 100644
--- /dev/null
+++ b/Engine/Documentation/Source/Programming/Tutorials/FirstPersonShooter/3/4/FirstPersonShooter_3_4.INT.udn
@@ -0,0 +1,28 @@
+Availability: Docs
+Title: 3.4 - Getting your Projectiles to Interact with the World
+Crumbs: %ROOT%, GettingStarted, Engine/Content/QuickStart
+Description: Learn how to get projectiles to interact with the world for your First Person Shooter project.
+SkillLevel: Beginner
+Version: 4.9
+
+[VAR:Steps]
+[OBJECT:Navigation]
+	[PARAM:previous]
+		[Previous Step](Programming/Tutorials/FirstPersonShooter/3/3)
+	[/PARAM]
+	[PARAM:current]
+	[/PARAM]
+	[PARAM:home]
+		[Section Home](Programming/Tutorials/FirstPersonShooter/3)
+	[/PARAM]
+	[PARAM:next]
+		[Next Step](Programming/Tutorials/FirstPersonShooter/3/5)
+	[/PARAM]
+[/OBJECT]
+[/VAR]
+
+%Steps%
+
+<!--TODO: Add tutorial steps here-->
+
+%Steps%
\ No newline at end of file
diff --git a/Engine/Documentation/Source/Programming/Tutorials/FirstPersonShooter/3/5/FirstPersonShooter_3_5.INT.udn b/Engine/Documentation/Source/Programming/Tutorials/FirstPersonShooter/3/5/FirstPersonShooter_3_5.INT.udn
new file mode 100644
--- /dev/null
+++ b/Engine/Documentation/Source/Programming/Tutorials/FirstPersonShooter/3/5/FirstPersonShooter_3_5.INT.udn
@@ -0,0 +1,28 @@
+Availability: Docs
+Title: 3.5 - Adding Crosshairs to your Viewport
+Crumbs: %ROOT%, GettingStarted, Engine/Content/QuickStart
+Description: Learn how to add crosshairs to the viewport for your First Person Shooter project.
+SkillLevel: Beginner
+Version: 4.9
+
+[VAR:Steps]
+[OBJECT:Navigation]
+	[PARAM:previous]
+		[Previous Step](Programming/Tutorials/FirstPersonShooter/3/4)
+	[/PARAM]
+	[PARAM:current]
+	[/PARAM]
+	[PARAM:home]
+		[Section Home](Programming/Tutorials/FirstPersonShooter/3)
+	[/PARAM]
+	[PARAM:next]
+		[Next Section](Programming/Tutorials/FirstPersonShooter/4)
+	[/PARAM]
+[/OBJECT]
+[/VAR]
+
+%Steps%
+
+<!--TODO: Add tutorial steps here-->
+
+%Steps%
\ No newline at end of file
