#UE4doc #docs
Added Section 4 subpages without content.
Availability: Docs

diff --git a/Engine/Documentation/Source/Programming/Tutorials/FirstPersonShooter/4/1/FirstPersonShooter_4_1.INT.udn b/Engine/Documentation/Source/Programming/Tutorials/FirstPersonShooter/4/1/FirstPersonShooter_4_1.INT.udn
new file mode 100644
--- /dev/null
+++ b/Engine/Documentation/Source/Programming/Tutorials/FirstPersonShooter/4/1/FirstPersonShooter_4_1.INT.udn
@@ -0,0 +1,27 @@
+Availability: Docs
+Title: 4.1 - Animating your Character
+Crumbs: %ROOT%, GettingStarted, Engine/Content/QuickStart
+Description: Learn how to animate your First Person Shooter character.
+SkillLevel: Beginner
+Version: 4.9
+
+[VAR:Steps]
+[OBJECT:Navigation]
+	[PARAM:previous]
+	[/PARAM]
+	[PARAM:current]
+	[/PARAM]
+	[PARAM:home]
+		[Section Home](Programming/Tutorials/FirstPersonShooter/4)
+	[/PARAM]
+	[PARAM:next]
+		[Next Step](Programming/Tutorials/FirstPersonShooter/4/2)
+	[/PARAM]
+[/OBJECT]
+[/VAR]
+
+%Steps%
+
+<!--TODO: Add tutorial steps here-->
+
+%Steps%
\ No newline at end of file
diff --git a/Engine/Documentation/Source/Programming/Tutorials/FirstPersonShooter/4/2/FirstPersonShooter_4_2.INT.udn b/Engine/Documentation/Source/Programming/Tutorials/FirstPersonShooter/4/2/FirstPersonShooter_4_2.INT.udn
new file mode 100644
--- /dev/null
+++ b/Engine/Documentation/Source/Programming/Tutorials/FirstPersonShooter/4/2/FirstPersonShooter_4_2.INT.udn
@@ -0,0 +1,28 @@
+Availability: Docs
+Title: 4.2 - Setting up your Event Graph
+Crumbs: %ROOT%, GettingStarted, Engine/Content/QuickStart
+Description: Learn how to set-up an Event Graph for your First Person Shooter character.
+SkillLevel: Beginner
+Version: 4.9
+
+[VAR:Steps]
+[OBJECT:Navigation]
+	[PARAM:previous]
+		[Previous Step](Programming/Tutorials/FirstPersonShooter/4/1)
+	[/PARAM]
+	[PARAM:current]
+	[/PARAM]
+	[PARAM:home]
+		[Section Home](Programming/Tutorials/FirstPersonShooter/4)
+	[/PARAM]
+	[PARAM:next]
+		[Next Step](Programming/Tutorials/FirstPersonShooter/4/3)
+	[/PARAM]
+[/OBJECT]
+[/VAR]
+
+%Steps%
+
+<!--TODO: Add tutorial steps here-->
+
+%Steps%
\ No newline at end of file
diff --git a/Engine/Documentation/Source/Programming/Tutorials/FirstPersonShooter/4/3/FirstPersonShooter_4_3.INT.udn b/Engine/Documentation/Source/Programming/Tutorials/FirstPersonShooter/4/3/FirstPersonShooter_4_3.INT.udn
new file mode 100644
--- /dev/null
+++ b/Engine/Documentation/Source/Programming/Tutorials/FirstPersonShooter/4/3/FirstPersonShooter_4_3.INT.udn
@@ -0,0 +1,28 @@
+Availability: Docs
+Title: 4.3 - Adding an Animation State Machine
+Crumbs: %ROOT%, GettingStarted, Engine/Content/QuickStart
+Description: Learn how to add an animation state machine for your First Person character.
+SkillLevel: Beginner
+Version: 4.9
+
+[VAR:Steps]
+[OBJECT:Navigation]
+	[PARAM:previous]
+		[Previous Step](Programming/Tutorials/FirstPersonShooter/4/2)
+	[/PARAM]
+	[PARAM:current]
+	[/PARAM]
+	[PARAM:home]
+		[Section Home](Programming/Tutorials/FirstPersonShooter/4)
+	[/PARAM]	
+	[PARAM:next]
+		[Next Step](Programming/Tutorials/FirstPersonShooter/4/4)
+	[/PARAM]
+[/OBJECT]
+[/VAR]
+
+%Steps%
+
+<!--TODO: Add tutorial steps here-->
+
+%Steps%
\ No newline at end of file
diff --git a/Engine/Documentation/Source/Programming/Tutorials/FirstPersonShooter/4/4/FirstPersonShooter_4_4.INT.udn b/Engine/Documentation/Source/Programming/Tutorials/FirstPersonShooter/4/4/FirstPersonShooter_4_4.INT.udn
new file mode 100644
--- /dev/null
+++ b/Engine/Documentation/Source/Programming/Tutorials/FirstPersonShooter/4/4/FirstPersonShooter_4_4.INT.udn
@@ -0,0 +1,28 @@
+Availability: Docs
+Title: 4.4 - Adding Animation Transition States
+Crumbs: %ROOT%, GettingStarted, Programming/Tutorials
+Description: How to add animation transition states to your First Person Shooter project.
+SkillLevel: Beginner
+Version: 4.9
+
+[VAR:Steps]
+[OBJECT:Navigation]
+	[PARAM:previous]
+		[Previous Step](Programming/Tutorials/FirstPersonShooter/4/3)
+	[/PARAM]
+	[PARAM:current]
+	[/PARAM]
+	[PARAM:home]
+		[Section Home](Programming/Tutorials/FirstPersonShooter/4)
+	[/PARAM]
+	[PARAM:next]
+		[Next Step](Programming/Tutorials/FirstPersonShooter/4/5)
+	[/PARAM]
+[/OBJECT]
+[/VAR]
+
+%Steps%
+
+<!--TODO: Add tutorial steps here-->
+
+%Steps%
\ No newline at end of file
diff --git a/Engine/Documentation/Source/Programming/Tutorials/FirstPersonShooter/4/5/FirstPersonShooter_4_5.INT.udn b/Engine/Documentation/Source/Programming/Tutorials/FirstPersonShooter/4/5/FirstPersonShooter_4_5.INT.udn
new file mode 100644
--- /dev/null
+++ b/Engine/Documentation/Source/Programming/Tutorials/FirstPersonShooter/4/5/FirstPersonShooter_4_5.INT.udn
@@ -0,0 +1,27 @@
+Availability: Docs
+Title: 4.5 - Associating Animation & Character Blueprints
+Crumbs: %ROOT%, GettingStarted, Engine/Content/QuickStart
+Description: Learn how to associate the Animation Blueprint with the Character Blueprint for your First Person Shooter project.
+SkillLevel: Beginner
+Version: 4.9
+
+[VAR:Steps]
+[OBJECT:Navigation]
+	[PARAM:previous]
+		[Previous Step](Programming/Tutorials/FirstPersonShooter/4/4)
+	[/PARAM]
+	[PARAM:current]
+	[/PARAM]
+	[PARAM:home]
+		[Section Home](Programming/Tutorials/FirstPersonShooter/4)
+	[/PARAM]
+	[PARAM:next]
+	[/PARAM]
+[/OBJECT]
+[/VAR]
+
+%Steps%
+
+<!--TODO: Add tutorial steps here-->
+
+%Steps%
\ No newline at end of file
