#UE4doc #docs
Added Section 4.4 subpages without content.
Availability: Docs

diff --git a/Engine/Documentation/Source/Programming/Tutorials/FirstPersonShooter/4/4/1/FirstPersonShooter_4_4_1.INT.udn b/Engine/Documentation/Source/Programming/Tutorials/FirstPersonShooter/4/4/1/FirstPersonShooter_4_4_1.INT.udn
new file mode 100644
--- /dev/null
+++ b/Engine/Documentation/Source/Programming/Tutorials/FirstPersonShooter/4/4/1/FirstPersonShooter_4_4_1.INT.udn
@@ -0,0 +1,28 @@
+Availability: Docs
+Title: 4.4.1 - Add Idle to/from Run Transitions
+Crumbs: %ROOT%, GettingStarted, Engine/Content/QuickStart
+Description: Learn how to add idle to/from run state transitions your First Person Shooter character's animation state machine.
+SkillLevel: Beginner
+Version: 4.9
+
+[VAR:Steps]
+[OBJECT:Navigation]
+	[PARAM:previous]
+		[Previous Step](Programming/Tutorials/FirstPersonShooter/4/4)
+	[/PARAM]
+	[PARAM:current]
+	[/PARAM]
+	[PARAM:home]
+		[Section Home](Programming/Tutorials/FirstPersonShooter/4)
+	[/PARAM]
+	[PARAM:next]
+		[Next Step](Programming/Tutorials/FirstPersonShooter/4/4/2)
+	[/PARAM]
+[/OBJECT]
+[/VAR]
+
+%Steps%
+
+<!--TODO: Add tutorial steps here-->
+
+%Steps%
\ No newline at end of file
diff --git a/Engine/Documentation/Source/Programming/Tutorials/FirstPersonShooter/4/4/2/FirstPersonShooter_4_4_2.INT.udn b/Engine/Documentation/Source/Programming/Tutorials/FirstPersonShooter/4/4/2/FirstPersonShooter_4_4_2.INT.udn
new file mode 100644
--- /dev/null
+++ b/Engine/Documentation/Source/Programming/Tutorials/FirstPersonShooter/4/4/2/FirstPersonShooter_4_4_2.INT.udn
@@ -0,0 +1,28 @@
+Availability: Docs
+Title: 4.4.2 - Add Idle to Jump Start Transition
+Crumbs: %ROOT%, GettingStarted, Engine/Content/QuickStart
+Description: Learn how to add idle to jump start state transitions your First Person Shooter character's animation state machine.
+SkillLevel: Beginner
+Version: 4.9
+
+[VAR:Steps]
+[OBJECT:Navigation]
+	[PARAM:previous]
+		[Previous Step](Programming/Tutorials/FirstPersonShooter/4/4/1)
+	[/PARAM]
+	[PARAM:current]
+	[/PARAM]
+	[PARAM:home]
+		[Section Home](Programming/Tutorials/FirstPersonShooter/4)
+	[/PARAM]
+	[PARAM:next]
+		[Next Step](Programming/Tutorials/FirstPersonShooter/4/4/3)
+	[/PARAM]
+[/OBJECT]
+[/VAR]
+
+%Steps%
+
+<!--TODO: Add tutorial steps here-->
+
+%Steps%
\ No newline at end of file
diff --git a/Engine/Documentation/Source/Programming/Tutorials/FirstPersonShooter/4/4/3/FirstPersonShooter_4_4_3.INT.udn b/Engine/Documentation/Source/Programming/Tutorials/FirstPersonShooter/4/4/3/FirstPersonShooter_4_4_3.INT.udn
new file mode 100644
--- /dev/null
+++ b/Engine/Documentation/Source/Programming/Tutorials/FirstPersonShooter/4/4/3/FirstPersonShooter_4_4_3.INT.udn
@@ -0,0 +1,28 @@
+Availability: Docs
+Title: 4.4.3 - Add Run to Jump Start Transition
+Crumbs: %ROOT%, GettingStarted, Engine/Content/QuickStart
+Description: Learn how to add run to jump start state transitions your First Person Shooter character's animation state machine.
+SkillLevel: Beginner
+Version: 4.9
+
+[VAR:Steps]
+[OBJECT:Navigation]
+	[PARAM:previous]
+		[Previous Step](Programming/Tutorials/FirstPersonShooter/4/4/2)
+	[/PARAM]
+	[PARAM:current]
+	[/PARAM]
+	[PARAM:home]
+		[Section Home](Programming/Tutorials/FirstPersonShooter/4)
+	[/PARAM]
+	[PARAM:next]
+		[Next Step](Programming/Tutorials/FirstPersonShooter/4/4/4)
+	[/PARAM]
+[/OBJECT]
+[/VAR]
+
+%Steps%
+
+<!--TODO: Add tutorial steps here-->
+
+%Steps%
\ No newline at end of file
diff --git a/Engine/Documentation/Source/Programming/Tutorials/FirstPersonShooter/4/4/4/FirstPersonShooter_4_4_4.INT.udn b/Engine/Documentation/Source/Programming/Tutorials/FirstPersonShooter/4/4/4/FirstPersonShooter_4_4_4.INT.udn
new file mode 100644
--- /dev/null
+++ b/Engine/Documentation/Source/Programming/Tutorials/FirstPersonShooter/4/4/4/FirstPersonShooter_4_4_4.INT.udn
@@ -0,0 +1,28 @@
+Availability: Docs
+Title: 4.4.4 - Add Jump Start to Jump Loop Transition
+Crumbs: %ROOT%, GettingStarted, Engine/Content/QuickStart
+Description: Learn how to add jump start to jump loop state transitions your First Person Shooter character's animation state machine.
+SkillLevel: Beginner
+Version: 4.9
+
+[VAR:Steps]
+[OBJECT:Navigation]
+	[PARAM:previous]
+		[Previous Step](Programming/Tutorials/FirstPersonShooter/4/4/3)
+	[/PARAM]
+	[PARAM:current]
+	[/PARAM]
+	[PARAM:home]
+		[Section Home](Programming/Tutorials/FirstPersonShooter/4)
+	[/PARAM]
+	[PARAM:next]
+		[Next Step](Programming/Tutorials/FirstPersonShooter/4/4/5)
+	[/PARAM]
+[/OBJECT]
+[/VAR]
+
+%Steps%
+
+<!--TODO: Add tutorial steps here-->
+
+%Steps%
\ No newline at end of file
diff --git a/Engine/Documentation/Source/Programming/Tutorials/FirstPersonShooter/4/4/5/FirstPersonShooter_4_4_5.INT.udn b/Engine/Documentation/Source/Programming/Tutorials/FirstPersonShooter/4/4/5/FirstPersonShooter_4_4_5.INT.udn
new file mode 100644
--- /dev/null
+++ b/Engine/Documentation/Source/Programming/Tutorials/FirstPersonShooter/4/4/5/FirstPersonShooter_4_4_5.INT.udn
@@ -0,0 +1,28 @@
+Availability: Docs
+Title: 4.4.5 - Add Jump Loop to Jump End Transition
+Crumbs: %ROOT%, GettingStarted, Engine/Content/QuickStart
+Description: Learn how to add jump loop to jump end state transitions your First Person Shooter character's animation state machine.
+SkillLevel: Beginner
+Version: 4.9
+
+[VAR:Steps]
+[OBJECT:Navigation]
+	[PARAM:previous]
+		[Previous Step](Programming/Tutorials/FirstPersonShooter/4/4/4)
+	[/PARAM]
+	[PARAM:current]
+	[/PARAM]
+	[PARAM:home]
+		[Section Home](Programming/Tutorials/FirstPersonShooter/4)
+	[/PARAM]
+	[PARAM:next]
+		[Next Step](Programming/Tutorials/FirstPersonShooter/4/4/6)
+	[/PARAM]
+[/OBJECT]
+[/VAR]
+
+%Steps%
+
+<!--TODO: Add tutorial steps here-->
+
+%Steps%
\ No newline at end of file
diff --git a/Engine/Documentation/Source/Programming/Tutorials/FirstPersonShooter/4/4/6/FirstPersonShooter_4_4_6.INT.udn b/Engine/Documentation/Source/Programming/Tutorials/FirstPersonShooter/4/4/6/FirstPersonShooter_4_4_6.INT.udn
new file mode 100644
--- /dev/null
+++ b/Engine/Documentation/Source/Programming/Tutorials/FirstPersonShooter/4/4/6/FirstPersonShooter_4_4_6.INT.udn
@@ -0,0 +1,28 @@
+Availability: Docs
+Title: 4.4.6 - Add Jump End to Idle Transition
+Crumbs: %ROOT%, GettingStarted, Engine/Content/QuickStart
+Description: Learn how to add jump end to idle state transitions your First Person Shooter character's animation state machine.
+SkillLevel: Beginner
+Version: 4.9
+
+[VAR:Steps]
+[OBJECT:Navigation]
+	[PARAM:previous]
+		[Previous Step](Programming/Tutorials/FirstPersonShooter/4/4/5)
+	[/PARAM]
+	[PARAM:current]
+	[/PARAM]
+	[PARAM:home]
+		[Section Home](Programming/Tutorials/FirstPersonShooter/4)
+	[/PARAM]
+	[PARAM:next]
+		[Next Step](Programming/Tutorials/FirstPersonShooter/4/5)
+	[/PARAM]
+[/OBJECT]
+[/VAR]
+
+%Steps%
+
+<!--TODO: Add tutorial steps here-->
+
+%Steps%
\ No newline at end of file
