#UE4doc #docs - Added MenuItemAsset for Size Map to be referenced when needed.

diff --git a/Engine/Documentation/Source/Shared/Editors/Common/EditorMenuItems/EditorMenuItems.INT.udn b/Engine/Documentation/Source/Shared/Editors/Common/EditorMenuItems/EditorMenuItems.INT.udn
--- a/Engine/Documentation/Source/Shared/Editors/Common/EditorMenuItems/EditorMenuItems.INT.udn
+++ b/Engine/Documentation/Source/Shared/Editors/Common/EditorMenuItems/EditorMenuItems.INT.udn
@@ -40,6 +40,10 @@ Crumbs:
 * **View Reference** - Launches the reference viewer showing the selected assets' references. 
 [/EXCERPT:MenuAssetViewReference]
 
+[EXCERPT:MenuAssetSizeMap]
+* **Size Map** - Displays an interactive map showing the approximate size of this asset and everything it references. 
+[/EXCERPT:MenuAssetSizeMap]
+
 [EXCERPT:MenuAssetReimport]
 * **Reimport _filename_** - Reimports the current asset from its original location on disk.
 [/EXCERPT:MenuAssetReimport]
