#UE4doc #docs - Added MenuWindow items that reference the Destructible Mesh editor for Destrucitble Settings, Fracture Settings, and Chunk Parameters.

diff --git a/Engine/Documentation/Source/Shared/Editors/Common/EditorMenuItems/EditorMenuItems.INT.udn b/Engine/Documentation/Source/Shared/Editors/Common/EditorMenuItems/EditorMenuItems.INT.udn
--- a/Engine/Documentation/Source/Shared/Editors/Common/EditorMenuItems/EditorMenuItems.INT.udn
+++ b/Engine/Documentation/Source/Shared/Editors/Common/EditorMenuItems/EditorMenuItems.INT.udn
@@ -68,6 +68,17 @@ Crumbs:
 * **Details** - Toggles the display of the **Details** panel.
 [/EXCERPT:MenuWindowDetails]
 
+[EXCERPT:MenuWindowDestructibleSettings]
+* **Destructible Settings** - Toggles the display of the **Destructible Settings** panel.
+[/EXCERPT:MenuWindowDestructibleSettings]
+
+[EXCERPT:MenuWindowFractureSettings]
+* **Fracture Settings** - Toggles the display of the **Fracture Settings** panel.
+[/EXCERPT:MenuWindowFractureSettings]
+
+[EXCERPT:MenuWindowChunkParameters]
+* **Chunk Parameters** - Toggles the display of the **Chunk Parameters** panel.
+[/EXCERPT:MenuWindowChunkParameters]
 
 [EXCERPT:MenuWindow]
 * **Content Browser** - Brings up the **Content Browser** in a separate window.
