#UE4doc #docs - Per UEDOC-2242. Page added for Destructible Mesh Editor UI.

diff --git a/Engine/Documentation/Source/Engine/Physics/Destructibles/Editor/DestructibleMeshEditorInterface.INT.udn b/Engine/Documentation/Source/Engine/Physics/Destructibles/Editor/DestructibleMeshEditorInterface.INT.udn
new file mode 100644
--- /dev/null
+++ b/Engine/Documentation/Source/Engine/Physics/Destructibles/Editor/DestructibleMeshEditorInterface.INT.udn
@@ -0,0 +1,131 @@
+Availability: Docs
+Title:Destructible Mesh Editor UI
+Crumbs:%ROOT%, Engine, Engine/Content, Engine/Physics/Destructibles/Editor
+Description:Tool used to preview the look, preview depth levels, and materials of DestructibleMesh assets.
+
+[TOC(start:2)]
+
+
+The **Destructible Mesh Editor** consists of six areas:
+
+[REGION:fullwidth]
+![](DM_UI.png)
+[/REGION]
+1. [Menu Bar](#MenuBar)
+1. [Toolbar](#ToolBar)
+1. [Viewport Panel](#ViewportPanel)
+1. [Destructible Settings Panel](#DestructibleSettingsPanel)
+1. [Fracture Settings Panel](#FractureSettingsPanel)
+1. [Chunk Parameters Panel](#ChunkParametersPanel)
+
+[Include:Shared/Editors/Common/EditorMenuItems/#UITabTips]
+
+<!-- For more information about the Destructible Mesh Editor, see [](Engine/Content\Types\APEX\Editor). -->
+
+## Menu Bar
+
+### File
+
+[INCLUDE:Shared/Editors/Common/EditorMenuItems/#MenuFileSave]
+[INCLUDE:Shared/Editors/Common/EditorMenuItems/#MenuFileSaveOptions]
+
+### Edit
+
+[INCLUDE:Shared/Editors/Common/EditorMenuItems/#MenuEditUndo]
+[INCLUDE:Shared/Editors/Common/EditorMenuItems/#MenuEditSettings]
+
+### Asset
+
+[INCLUDE:Shared/Editors/Common/EditorMenuItems/#MenuAssetContentBrowser]
+[INCLUDE:Shared/Editors/Common/EditorMenuItems/#MenuAssetViewReference]
+[INCLUDE:Shared/Editors/Common/EditorMenuItems/#MenuAssetSizeMap]
+[INCLUDE:Shared/Editors/Common/EditorMenuItems/#MenuAssetReimport]
+
+
+### Window
+[INCLUDE:Shared/Editors/Common/EditorMenuItems/#MenuWindowToolbar]
+[INCLUDE:Shared/Editors/Common/EditorMenuItems/#MenuWindowViewport]
+[INCLUDE:Shared/Editors/Common/EditorMenuItems/#MenuWindowDetails]
+[INCLUDE:Shared/Editors/Common/EditorMenuItems/#MenuWindowDestructibleSettings]
+[INCLUDE:Shared/Editors/Common/EditorMenuItems/#MenuWindowFractureSettings]
+[INCLUDE:Shared/Editors/Common/EditorMenuItems/#MenuWindowChunkParameters]
+[INCLUDE:Shared/Editors/Common/EditorMenuItems/#MenuWindow]
+
+
+### Help
+
+[INCLUDE:Shared/Editors/Common/EditorMenuItems/#MenuHelp]
+
+
+## Toolbar
+
+| Item | Description |
+| ---- | ----------- |
+| **![Fracture Mesh](DM_FractureMeshIcon.png)** | This will fracture the meshes root chunk(s) based upon the Fracture Settings. |
+| **![Refresh](DM_RefreshIcon.png)** | This will refresh the DestructibleMesh from the StaticMesh it was created from. |
+| **![Import FBX Chunks](DM_FBXChunksIcon.png)** | This will import a FBX as depth level 1 chunks. |
+| **![Preview Depth](DM_PreviewDepthIcon.png)** | This is a drop down selection to preview a specific depth level of destruction. |
+| **![Explode Amount](DM_ExplodeAmtIcon.png)** | This slider will push apart the chunks of the specified Preview Depth. |
+
+
+## Viewport Panel
+
+The **Viewport** panel shows a rendered view of the Destructible Mesh asset. This allows you to examine the Destructible Mesh as it would be rendered in-game. 
+The viewport will also allow you to use different Preview Depths to see your different levels of destruction. This will allow different chunks to be selected which can be used in conjunction with the Chunk Parameters panel settings as well.  
+
+![Destructible Mesh Editor Preview panel](DM_Preview.png)
+
+
+
+## Destructible Settings Panel
+
+The **Destructible Settings** panel shows specific properties pertaining to the Destructible Mesh Actor, such as the Materials applied to the surface and interior, Damage settings, and flags that can enable specific abilities. 
+
+For information about the basics of using this panel, see the [**Destructibles Properties Matrix**](Engine/Physics/Destructibles/DestructibleProperties) documentation.
+
+![Editor Destructible Settings Panel](DM_Settings.png)
+
+
+## Fracture Settings Panel
+
+The **Fracture Settings** panel shows settings to control how the Destructible Mesh Actor is broken into chunks. There are properties that pertain to the number of chunks the fracture should contain, control for the interior material UVs, and the random pattern seed. 
+
+For information about the basics of using this panel, see the [**Destructibles Properties Matrix**](Engine/Physics/Destructibles/DestructibleProperties) documentation.
+
+![Fracture Settings Panel](DM_FractureSettings.png)
+
+
+## Chunk Parameters Panel
+
+The **Chunk Parameters** panel contains settings that directly affect specific chunks that are selected in the preview panel when viewing Preview Depth 1 or higher. The settings here can control whether a chunk can be damaged, fractured, crumbled, or even considered as support.
+
+For information about the basics of using this panel, see the [**Destructibles Properties Matrix**](Engine/Physics/Destructibles/DestructibleProperties) documentation.
+
+![Chunk Parameters Panel](DM_ChunkParameters.png)
+
+
+## Controls
+
+
+### Mouse Controls
+
+**Viewport Panel**
+
+
+* **LMB + Drag** - If the camera is locked, rotates the mesh around its Z axis and moves towards or away from the origin. Otherwise, rotates the camera around its Z axis and moves the camera along its local X axis.
+* **RMB + Drag** - If the camera is locked, rotates the mesh. Otherwise, rotates the camera.
+* **LMB + RMB + Drag** - If the camera is not locked, moves the camera along its local YZ plane.
+
+
+### Keyboard Controls
+
+* **L + Mouse Move** - Rotates the preview light in the **Preview** panel.
+
+
+
+
+
+
+
+
+
