#UE4doc #docs - fixed issues associated with UEDOC-2240 addressing some of the Animation BP UI page. Updated all images to be more unified and included PSD file for easier editing in the future.

diff --git a/Engine/Documentation/Source/Engine/Animation/AnimBlueprints/Interface/AnimationBlueprintUI.INT.udn b/Engine/Documentation/Source/Engine/Animation/AnimBlueprints/Interface/AnimationBlueprintUI.INT.udn
--- a/Engine/Documentation/Source/Engine/Animation/AnimBlueprints/Interface/AnimationBlueprintUI.INT.udn
+++ b/Engine/Documentation/Source/Engine/Animation/AnimBlueprints/Interface/AnimationBlueprintUI.INT.udn
@@ -16,9 +16,9 @@ The following sections describe the areas of the Animation Blueprint UI.
 1. Menu Bar
 1. Toolbar
 1. Viewport
+1. Graph
 1. Anim Preview Editor
 1. Asset Override Editor
-1. Graph
 1. Asset Browser tab
 1. My Blueprint tab
 1. Details panel
@@ -65,20 +65,13 @@ In addition to those options, the Animation Blueprint toolbar also contains tabs
 ## Viewport
 The **Viewport** displays the Skeletal Mesh with any animations activated by current variable values.
 
-![](AnimGraphUI_Viewport.png)
-
-## Class Defaults
-
-The **Class Defaults** panel displays all of your Animation Blueprint's variables, in their assigned categories. You can manually modify the variable values here to preview and test different animations in the viewport.
-
-![](AnimGraphUI_BPDefaults.png)
-
-For more information, see [](Engine/Animation/Persona\BlueprintDefaults).
+![](AnimGraphUI_Viewport.png)(w:500)
 
 ## Graph
 
 The tabbed **Graph** panel displays whichever graph is currently selected.
 
+
 ### EventGraph Tab
 
 The **EventGraph** tab contains any animation event nodes that trigger the updates of the Skeletal Mesh poses.  It can also contain animation nodes that set variable values that determine which animation asset is played when.
@@ -91,12 +84,25 @@ The **AnimGraph** tab contains the State Machine that hooks up all the actual an
 
 ![](AnimGraphUI_AnimGraph.png)
 
-For more information about State Machines, see [](Engine/Animation/StateMachines).
+For more information about State Machines, see **[](Engine/Animation/StateMachines)**.
 
 [REGION:note]
 Although you could put all your animation nodes directly into the AnimGraph tab, it is more efficient to organize them into a State Machine, States, and Transitions (Rules), for ease of upkeep and simplicity.
 [/REGION]
 
+## Anim Preview Editor
+
+The **Anim Preview Editor** panel displays all of your Animation Blueprint's variables, in their assigned categories. You can manually modify the variable values here to preview and test different animations in the viewport.
+
+![](AnimGraphUI_BPDefaults.png)
+
+For more information, see **[](Engine/Animation/Persona\BlueprintDefaults)**.
+
+
+## Asset Override Editor
+
+The **Asset Override Editor** panel displays all of your animation assets that are referenced by the parent animation graph. If there are no nodes referencing assets in the parent class or this is not a derived blueprint there will be nothing displayed in this panel.
+
 ## Asset Browser Tab
 
 The Asset Browser tab contains a list of all the animation assets in your Animation Blueprint.
@@ -123,7 +129,7 @@ The **Details** panel displays the properties of selected elements.
 
 ## Palette
 
-The **Palette** panel contains all the nodes available for use in the Animation Blueprint's various graphs.
+The **Palette** panel contains all the nodes available for use in the Animation Blueprint's various graphs. In the top panel any nodes marked as a "favorite" will be displayed here.
 
 ![](AnimGraphUI_Palette.png)
 
