Replacing the '&' in page titles with 'and'

diff --git a/Engine/Documentation/Source/Engine/Actors/Actors.INT.udn b/Engine/Documentation/Source/Engine/Actors/Actors.INT.udn
--- a/Engine/Documentation/Source/Engine/Actors/Actors.INT.udn
+++ b/Engine/Documentation/Source/Engine/Actors/Actors.INT.udn
@@ -1,5 +1,5 @@
 Availability:Public
-Title: Actors & Geometry
+Title: Actors and Geometry
 Crumbs:%ROOT%, Engine
 Description:A definition of Actors and how they are used in level design, including placement and manipulation techniques.
 Navigation:topic
diff --git a/Engine/Documentation/Source/Engine/Audio/Audio.INT.udn b/Engine/Documentation/Source/Engine/Audio/Audio.INT.udn
--- a/Engine/Documentation/Source/Engine/Audio/Audio.INT.udn
+++ b/Engine/Documentation/Source/Engine/Audio/Audio.INT.udn
@@ -1,5 +1,5 @@
 Availability: Public
-Title: Audio & Sound
+Title: Audio and Sound
 Crumbs: %ROOT%, Engine
 Description:System for playing in-game sounds, including the use of SoundCue node-based audio assets.
 Navigation:topic
diff --git a/Engine/Documentation/Source/Engine/Deployment/Deployment.INT.udn b/Engine/Documentation/Source/Engine/Deployment/Deployment.INT.udn
--- a/Engine/Documentation/Source/Engine/Deployment/Deployment.INT.udn
+++ b/Engine/Documentation/Source/Engine/Deployment/Deployment.INT.udn
@@ -1,5 +1,5 @@
 Availability:Public
-Title:Packaging & Cooking Games
+Title:Packaging and Cooking Games
 Crumbs:%ROOT%, Engine
 Description:Generating platform-specific content for assets used in the game.
 
diff --git a/Engine/Documentation/Source/Engine/Matinee/MatineeAndCinematics.INT.udn b/Engine/Documentation/Source/Engine/Matinee/MatineeAndCinematics.INT.udn
--- a/Engine/Documentation/Source/Engine/Matinee/MatineeAndCinematics.INT.udn
+++ b/Engine/Documentation/Source/Engine/Matinee/MatineeAndCinematics.INT.udn
@@ -1,5 +1,5 @@
 Availability: Public
-Title:Matinee & Cinematics
+Title:Matinee and Cinematics
 Crumbs: %ROOT%, Engine
 Description:Animation tool that provides the ability to animate the properties of Actors over time, either to create dynamic gameplay or cinematic in-game sequences.
 
diff --git a/Engine/Documentation/Source/Engine/Rendering/LightingAndShadows/QuickStart/2/LQS_2.INT.udn b/Engine/Documentation/Source/Engine/Rendering/LightingAndShadows/QuickStart/2/LQS_2.INT.udn
--- a/Engine/Documentation/Source/Engine/Rendering/LightingAndShadows/QuickStart/2/LQS_2.INT.udn
+++ b/Engine/Documentation/Source/Engine/Rendering/LightingAndShadows/QuickStart/2/LQS_2.INT.udn
@@ -1,5 +1,5 @@
 Availability: Public
-Title:2. Atmospheric & Directional Lighting
+Title:2. Atmospheric and Directional Lighting
 Crumbs:%ROOT%, GettingStarted, Engine\Rendering\LightingAndShadows\QuickStart
 Description:In this step we will create the basic level lighting using Atmospheric Fog and a Directional Light. 
 
diff --git a/Engine/Documentation/Source/Engine/Rendering/Rendering.INT.udn b/Engine/Documentation/Source/Engine/Rendering/Rendering.INT.udn
--- a/Engine/Documentation/Source/Engine/Rendering/Rendering.INT.udn
+++ b/Engine/Documentation/Source/Engine/Rendering/Rendering.INT.udn
@@ -1,5 +1,5 @@
 Availability:Public
-Title: Rendering & Graphics
+Title: Rendering and Graphics
 Crumbs: %ROOT%, Engine
 Description:Rendering subsystem including lighting and shadowing, materials and textures, particles and effects, and post processing.
 Related:Programming/Rendering
diff --git a/Engine/Documentation/Source/Engine/UI/Tools/EditorTools.INT.udn b/Engine/Documentation/Source/Engine/UI/Tools/EditorTools.INT.udn
--- a/Engine/Documentation/Source/Engine/UI/Tools/EditorTools.INT.udn
+++ b/Engine/Documentation/Source/Engine/UI/Tools/EditorTools.INT.udn
@@ -1,5 +1,5 @@
 Availability:Public
-Title:Tools & Systems
+Title:Tools and Systems
 Crumbs:%ROOT%, Engine, Engine/UI
 Description:Individual tools and systems built into Unreal Editor.
 Navigation:topic
diff --git a/Engine/Documentation/Source/Gameplay/AI/AI.INT.udn b/Engine/Documentation/Source/Gameplay/AI/AI.INT.udn
--- a/Engine/Documentation/Source/Gameplay/AI/AI.INT.udn
+++ b/Engine/Documentation/Source/Gameplay/AI/AI.INT.udn
@@ -1,5 +1,5 @@
 Availability:Public
-Title: AI & Behavior Trees
+Title: AI and Behavior Trees
 Crumbs:%ROOT%, Engine, Gameplay
 Description:Information over Artificial Intelligence including the use of Behavior Trees and Navmesh.
 
diff --git a/Engine/Documentation/Source/Gameplay/Networking/Networking.INT.udn b/Engine/Documentation/Source/Gameplay/Networking/Networking.INT.udn
--- a/Engine/Documentation/Source/Gameplay/Networking/Networking.INT.udn
+++ b/Engine/Documentation/Source/Gameplay/Networking/Networking.INT.udn
@@ -1,5 +1,5 @@
 Availability:Public
-Title: Networking & Multiplayer
+Title: Networking and Multiplayer
 Crumbs:%ROOT%, Gameplay, Gameplay/Framework
 Description:Setting up networked games for multiplayer.
 
diff --git a/Engine/Documentation/Source/GettingStarted/HowTo/FindingHelp/FindingHelp.INT.udn b/Engine/Documentation/Source/GettingStarted/HowTo/FindingHelp/FindingHelp.INT.udn
--- a/Engine/Documentation/Source/GettingStarted/HowTo/FindingHelp/FindingHelp.INT.udn
+++ b/Engine/Documentation/Source/GettingStarted/HowTo/FindingHelp/FindingHelp.INT.udn
@@ -1,5 +1,5 @@
 Availability:Public
-Title:Find Help & Answers
+Title:Find Help and Answers
 Crumbs: %ROOT%, Gameplay, GettingStarted\HowTo
 Description:Where to find in-editor help and tutorials or external help and answers to your Unreal Engine 4 questions. 
 Related: Resources
diff --git a/Engine/Documentation/Source/GettingStarted/HowTo/PIE/PIE.INT.udn b/Engine/Documentation/Source/GettingStarted/HowTo/PIE/PIE.INT.udn
--- a/Engine/Documentation/Source/GettingStarted/HowTo/PIE/PIE.INT.udn
+++ b/Engine/Documentation/Source/GettingStarted/HowTo/PIE/PIE.INT.udn
@@ -1,5 +1,5 @@
 Availability:Public
-Title:Playing & Simulating
+Title:Playing and Simulating
 Crumbs: %ROOT%, Gameplay, GettingStarted\HowTo
 Description:Playtesting and Simulating your game with the Play and Simulate options. 
 Related: Engine/UI/LevelEditor/InEditorTesting
diff --git a/Engine/Documentation/Source/Resources/Resources.INT.udn b/Engine/Documentation/Source/Resources/Resources.INT.udn
--- a/Engine/Documentation/Source/Resources/Resources.INT.udn
+++ b/Engine/Documentation/Source/Resources/Resources.INT.udn
@@ -1,5 +1,5 @@
 Availability: Public
-Title:Samples & Tutorials
+Title:Samples and Tutorials
 Crumbs:%ROOT%
 Description:Links to various example scenes, sample games, and tutorials.
 Navigation:topic
diff --git a/Engine/Documentation/Source/Resources/Showcases/BlueprintExamples/FogSheet/FogSheet.INT.udn b/Engine/Documentation/Source/Resources/Showcases/BlueprintExamples/FogSheet/FogSheet.INT.udn
--- a/Engine/Documentation/Source/Resources/Showcases/BlueprintExamples/FogSheet/FogSheet.INT.udn
+++ b/Engine/Documentation/Source/Resources/Showcases/BlueprintExamples/FogSheet/FogSheet.INT.udn
@@ -1,5 +1,5 @@
 Availability: Public
-Title:Fog Sheet & Light Beams
+Title:Fog Sheet and Light Beams
 Crumbs:%ROOT%, Resources, Resources/Showcases/BlueprintExamples
 Description:An overview of the Fog Sheet Blueprint.
 Related: Engine/Blueprints/Overview
diff --git a/Engine/Documentation/Source/Resources/Showcases/BlueprintExamples/SecurityDoor/SecurityDoor.INT.udn b/Engine/Documentation/Source/Resources/Showcases/BlueprintExamples/SecurityDoor/SecurityDoor.INT.udn
--- a/Engine/Documentation/Source/Resources/Showcases/BlueprintExamples/SecurityDoor/SecurityDoor.INT.udn
+++ b/Engine/Documentation/Source/Resources/Showcases/BlueprintExamples/SecurityDoor/SecurityDoor.INT.udn
@@ -1,5 +1,5 @@
 Availability: Public
-Title:Security Camera, Alarm, & Door
+Title:Security Camera, Alarm, and Door
 Crumbs:%ROOT%, Resources, Resources/Showcases/BlueprintExamples
 Description:An overview of the Security Door Blueprints.
 Related: Engine/Blueprints/Overview
