Updating titles on the pages listed. They will show live after the next publishing

diff --git a/Engine/Documentation/Source/Engine/Basics/SourceControl/InEditor/SourceControlInsideUE4.INT.udn b/Engine/Documentation/Source/Engine/Basics/SourceControl/InEditor/SourceControlInsideUE4.INT.udn
--- a/Engine/Documentation/Source/Engine/Basics/SourceControl/InEditor/SourceControlInsideUE4.INT.udn
+++ b/Engine/Documentation/Source/Engine/Basics/SourceControl/InEditor/SourceControlInsideUE4.INT.udn
@@ -1,5 +1,5 @@
 Availability: Public
-Title: Source Control inside Unreal Engine 4
+Title: Source Control Inside Unreal Editor
 Crumbs: %ROOT%, Engine, Engine\Basics\SourceControl
 Description: How to use Source Control software inside UE4. 
 Related: Engine/UI/SourceControl
diff --git a/Engine/Documentation/Source/Engine/Basics/SourceControl/Perforce/P4SourceControl.INT.udn b/Engine/Documentation/Source/Engine/Basics/SourceControl/Perforce/P4SourceControl.INT.udn
--- a/Engine/Documentation/Source/Engine/Basics/SourceControl/Perforce/P4SourceControl.INT.udn
+++ b/Engine/Documentation/Source/Engine/Basics/SourceControl/Perforce/P4SourceControl.INT.udn
@@ -1,5 +1,5 @@
 Availability: Public
-Title: Using Perforce as Source Control inside UE4
+Title: Using Perforce as Source Control
 Crumbs: %ROOT%, Engine, Engine\Basics\SourceControl
 Description: How to setup Perforce so that you can share assets with other on your team. 
 Related: Engine/UI/SourceControl
diff --git a/Engine/Documentation/Source/Engine/Basics/SourceControl/SVN/SVNSourceControl.INT.udn b/Engine/Documentation/Source/Engine/Basics/SourceControl/SVN/SVNSourceControl.INT.udn
--- a/Engine/Documentation/Source/Engine/Basics/SourceControl/SVN/SVNSourceControl.INT.udn
+++ b/Engine/Documentation/Source/Engine/Basics/SourceControl/SVN/SVNSourceControl.INT.udn
@@ -1,5 +1,5 @@
 Availability: Public
-Title: Using SVN as Source Control inside UE4
+Title: Using SVN as Source Control
 Crumbs: %ROOT%, Engine, Engine\Basics\SourceControl
 Description: How to setup SVN so that you can share assets with other on your team. 
 Related: Engine/UI/SourceControl
