Fixing page titles

diff --git a/Engine/Documentation/Source/Engine/Matinee/HowTo/MHT_1/MHT_1.INT.udn b/Engine/Documentation/Source/Engine/Matinee/HowTo/MHT_1/MHT_1.INT.udn
--- a/Engine/Documentation/Source/Engine/Matinee/HowTo/MHT_1/MHT_1.INT.udn
+++ b/Engine/Documentation/Source/Engine/Matinee/HowTo/MHT_1/MHT_1.INT.udn
@@ -1,5 +1,5 @@
 Availability: Public
-Title:Set Up a Door
+Title:Setting Up a Door
 Crumbs: %ROOT%, Engine, Engine/Matinee, Engine/Matinee/HowTo
 Description:Setting up a door that opens/closes automatically through Matinee. 
 
diff --git a/Engine/Documentation/Source/Engine/Matinee/HowTo/MHT_2/MHT_2.INT.udn b/Engine/Documentation/Source/Engine/Matinee/HowTo/MHT_2/MHT_2.INT.udn
--- a/Engine/Documentation/Source/Engine/Matinee/HowTo/MHT_2/MHT_2.INT.udn
+++ b/Engine/Documentation/Source/Engine/Matinee/HowTo/MHT_2/MHT_2.INT.udn
@@ -1,5 +1,5 @@
 Availability: Public
-Title:Set Up Lifts
+Title:Setting Up Lifts
 Crumbs: %ROOT%, Engine, Engine/Matinee, Engine/Matinee/HowTo
 Description:Create two types of lifts that move from point A to point B automatically.
 
diff --git a/Engine/Documentation/Source/Engine/Matinee/HowTo/MHT_3/MHT_3.INT.udn b/Engine/Documentation/Source/Engine/Matinee/HowTo/MHT_3/MHT_3.INT.udn
--- a/Engine/Documentation/Source/Engine/Matinee/HowTo/MHT_3/MHT_3.INT.udn
+++ b/Engine/Documentation/Source/Engine/Matinee/HowTo/MHT_3/MHT_3.INT.udn
@@ -1,5 +1,5 @@
 Availability: Public
-Title:Use Multiple Cameras
+Title:Using Multiple Cameras
 Crumbs: %ROOT%, Engine, Engine/Matinee, Engine/Matinee/HowTo
 Description:This scene shows how to switch between multiple cameras and apply camera movement.
 
diff --git a/Engine/Documentation/Source/Engine/Matinee/HowTo/MHT_4/MHT_4.INT.udn b/Engine/Documentation/Source/Engine/Matinee/HowTo/MHT_4/MHT_4.INT.udn
--- a/Engine/Documentation/Source/Engine/Matinee/HowTo/MHT_4/MHT_4.INT.udn
+++ b/Engine/Documentation/Source/Engine/Matinee/HowTo/MHT_4/MHT_4.INT.udn
@@ -1,5 +1,5 @@
 Availability: Public
-Title:Use Camera Effects
+Title:Using Camera Effects
 Crumbs: %ROOT%, Engine, Engine/Matinee, Engine/Matinee/HowTo
 Description:A simple scene that fades in-and-out, uses multiple FOV angles and Post Process Effects. 
 
diff --git a/Engine/Documentation/Source/Engine/Matinee/HowTo/MHT_5/MHT_5.INT.udn b/Engine/Documentation/Source/Engine/Matinee/HowTo/MHT_5/MHT_5.INT.udn
--- a/Engine/Documentation/Source/Engine/Matinee/HowTo/MHT_5/MHT_5.INT.udn
+++ b/Engine/Documentation/Source/Engine/Matinee/HowTo/MHT_5/MHT_5.INT.udn
@@ -1,5 +1,5 @@
 Availability: Public
-Title:Adjust Lights
+Title:Adjusting Lights
 Crumbs: %ROOT%, Engine, Engine/Matinee, Engine/Matinee/HowTo
 Description:Light Color, Brightness and other Light Settings are adjusted in this scene. 
 
diff --git a/Engine/Documentation/Source/Engine/Matinee/HowTo/MHT_6/MHT_6.INT.udn b/Engine/Documentation/Source/Engine/Matinee/HowTo/MHT_6/MHT_6.INT.udn
--- a/Engine/Documentation/Source/Engine/Matinee/HowTo/MHT_6/MHT_6.INT.udn
+++ b/Engine/Documentation/Source/Engine/Matinee/HowTo/MHT_6/MHT_6.INT.udn
@@ -1,5 +1,5 @@
 Availability: Public
-Title:Animate Materials
+Title:Animating Materials
 Crumbs: %ROOT%, Engine, Engine/Matinee, Engine/Matinee/HowTo
 Description:The Material of an Actor changes colors when the player is near and another Actor's Material blinks. 
 
diff --git a/Engine/Documentation/Source/Engine/Matinee/HowTo/MHT_7/MHT_7.INT.udn b/Engine/Documentation/Source/Engine/Matinee/HowTo/MHT_7/MHT_7.INT.udn
--- a/Engine/Documentation/Source/Engine/Matinee/HowTo/MHT_7/MHT_7.INT.udn
+++ b/Engine/Documentation/Source/Engine/Matinee/HowTo/MHT_7/MHT_7.INT.udn
@@ -1,5 +1,5 @@
 Availability: Public
-Title:Play Animations and Sounds
+Title:Playing Animations and Sounds
 Crumbs: %ROOT%, Engine, Engine/Matinee, Engine/Matinee/HowTo
 Description:How to animate a Skeletal Mesh jumping onto a platform and applying sounds in Matinee. 
 
diff --git a/Engine/Documentation/Source/Engine/Matinee/HowTo/MHT_8/MHT_8.INT.udn b/Engine/Documentation/Source/Engine/Matinee/HowTo/MHT_8/MHT_8.INT.udn
--- a/Engine/Documentation/Source/Engine/Matinee/HowTo/MHT_8/MHT_8.INT.udn
+++ b/Engine/Documentation/Source/Engine/Matinee/HowTo/MHT_8/MHT_8.INT.udn
@@ -1,5 +1,5 @@
 Availability: Public
-Title:Trigger Events and Toggle Visibility
+Title:Triggering Events and Toggle Visibility
 Crumbs: %ROOT%, Engine, Engine/Matinee, Engine/Matinee/HowTo
 Description:Using Events and Visibility, how to create a rock that explodes after a countdown. 
 
diff --git a/Engine/Documentation/Source/Engine/Matinee/HowTo/MatineeHowTo.INT.udn b/Engine/Documentation/Source/Engine/Matinee/HowTo/MatineeHowTo.INT.udn
--- a/Engine/Documentation/Source/Engine/Matinee/HowTo/MatineeHowTo.INT.udn
+++ b/Engine/Documentation/Source/Engine/Matinee/HowTo/MatineeHowTo.INT.udn
@@ -1,5 +1,5 @@
 Availability: Public
-Title:Matinee - How To...
+Title:Matinee How-To's
 Crumbs: %ROOT%, Engine, Engine/Matinee
 Description:A How To Guide for using different aspects of Matinee in Unreal Engine 4. 
 
