#UE4doc #docs:
Removed Overview header
Added engine version field
Replaced picture to reflect modern UI
Changed font settings for multiple terms

diff --git a/Engine/Documentation/Source/Engine/Animation/AnimMontage/AnimMontage.INT.udn b/Engine/Documentation/Source/Engine/Animation/AnimMontage/AnimMontage.INT.udn
--- a/Engine/Documentation/Source/Engine/Animation/AnimMontage/AnimMontage.INT.udn
+++ b/Engine/Documentation/Source/Engine/Animation/AnimMontage/AnimMontage.INT.udn
@@ -9,24 +9,25 @@ Related: Engine/Animation/StateMachines
 Related: Gameplay/AI
 Related: Engine/Content/Tools/MayaRiggingTool
 Related: Gameplay/Input
+version: 4.9
 
 
 [TOC(start:2 end:3)]
 
-## Overview
+
 
 [EXCERPT:Intro]
-**AnimMontages** (or just Montages for short) are a multipurpose tool that allows for a wide variety of animation effects, primarily related to exposing animation controls within code or Blueprint. It can also be used to create a wide variety of animation effects including intelligent loops of animation, logic-based animation switching, root motion handling, and much more.
+**Animation Montages** (or just Montages for short) are a multipurpose tool that allows for a wide variety of animation effects, primarily related to exposing animation controls within code or Blueprint. It can also be used to create a wide variety of animation effects including intelligent loops of animation, logic-based animation switching, root motion handling, and much more.
 [/EXCERPT:Intro]
 
-_AnimMontages_ are animation assets that can be created and exist within the **Content Browser**. These assets can then be plugged directly into the AnimGraph in run-time and you can modify any state of it. For example, you can jump to different sections or you can re-link different sections. This is mostly for code driven animations or one-off animations such as melee attacks, allowing you to control triggers, stop when you want, or change state (looping or non-looping) between (see the image below). 
+AnimMontages are animation assets that can be created and exist within the **Content Browser**. These assets can then be plugged directly into the AnimGraph in run-time and you can modify any state of it. For example, you can jump to different sections or you can re-link different sections. This is mostly for code driven animations or one-off animations such as melee attacks, allowing you to control triggers, stop when you want, or change state (looping or non-looping) between (see the image below). 
 
 ![](Montage_Screen2.png)
 
-The image above is a melee attack with 3 sections [start, loop, and end]. When the player **left-clicks**, it triggers the start section by default when you ask to play the montage. Now the middle section is set to loop. When the start is done, it will transition to loop, and will loop there forever. If the player lets the mouse button go, it will stop, but you do not want to stop right away since the animation will pop in the middle of the loop. You would like to relink the loop to the end, where it will then transition out to the end and finish the animation. 
+The image above is a melee attack with 3 sections [start, loop, and end]. When the player **Left-clicks**, it triggers the start section by default when you ask to play the montage. Now the middle section is set to loop. When the start is done, it will transition to loop, and will loop there forever. If the player lets the mouse button go, it will stop, but you do not want to stop right away since the animation will pop in the middle of the loop. You would like to relink the loop to the end, where it will then transition out to the end and finish the animation. 
 
 Some additional uses for Montages include:
-* The ability to play an animation from within an _AnimBlueprint's_ **EventGraph**.
+* The ability to play an animation from within an AnimBlueprint's **EventGraph**.
 * Chaining together a complex sequence of animations that you want to think of as a single animation.
 * Looping only a specific portion of animation or animations based on code or Blueprint script.
 * Handling event-based switching of multiple animations based on code or Blueprint script.
@@ -37,13 +38,13 @@ Some additional uses for Montages include:
 So, as you can see, there are a lot of different things you can do with Montages. This can easily make Montages seem overwhelming as a system. However, if you really boil it down, the key point of Montages is to expose animation controls to code or Blueprints. Nearly everything else the system can do follows from that single point.
 
 [REGION:tip]
-An example of an **Animation Montage** applied to a character can also be seen on the [Animation Content Examples](Resources\ContentExamples\Animation) page under section 1.5.
+An example of an Animation Montage applied to a character can also be seen on the [Animation Content Examples](Resources\ContentExamples\Animation) page under section 1.5.
 [/REGION]
 
 
 ## System Caveats
 
-There are a few details to keep in mind regarding _AnimMontages_:
+There are a few details to keep in mind regarding AnimMontages:
 
 * Only one Montage can play back at a time. If you start playing a second one, it will stop the first.
 * [INCLUDE:#MontageReplication]
@@ -52,7 +53,7 @@ There are a few details to keep in mind regarding _AnimMontages_:
 
 ## Montage Properties
 
-The following is a breakdown of the _AnimMontage_ asset properties. These are available in **Persona** when looking at a Montage, and can also be accessed by **right-clicking** on a Montage within the **Content Browser** and choosing _Properties_ from the context menu.
+The following is a breakdown of the AnimMontage asset properties. These are available in **Persona** when looking at a Montage, and can also be accessed by **Right-clicking** on a Montage within the **Content Browser** and choosing Properties from the context menu.
 
 ![](MontageProperties.png)
 
@@ -82,6 +83,7 @@ When looking at a Montage in Persona, it is useful to know what each area is and
 
 1. [Montage Area](#MontageArea)
 1. [Sections Area](#SectionsArea)
+1. Element Timing Area
 1. [Notifies Area](#NotifiesArea)
 1. [Curves Area](#CurvesArea)
 
@@ -108,15 +110,15 @@ Montage Sections provide a way to break a Slot up into multiple portions of anim
 
 It may help to think of Sections like songs on a music playlist, with Slots being the albums. Just like with many modern media players, you can choose which song will play next when the current one finishes, or just jump to the one you want to hear right now.
 
-Sections are created by **right-clicking** on the Section track and choosing **New Montage Section**. 
+Sections are created by **Right-clicking** on the Section track and choosing **New Montage Section**. 
 
 
 ### Slots
 
 [EXCERPT:Slots]
-Within a Montage, a **Slot** is just a single track that can hold any number of animations. You can name this Slot, and then blend to those specific animations by calling the Slot name. A great example is having a character with a weapon reload animation. You may have different versions of the reload for when the player is standing, crouching, and lying prone. So long as all 3 of the animations used the same timing, you could place each one within a separate Slot in your Montage; the Slots could be named _Standing_, _Crouching_, and _Prone_. In your _Animation Blueprint's_ **AnimGraph**, you can use the _Slot_ node to determine which one you want to play based on your character's current state. When they are standing, you can use the result of the animation in the _Standing_ Slot. When they are prone, you can see the result of the _Prone_ Slot.
+Within a Montage, a **Slot** is just a single track that can hold any number of animations. You can name this Slot, and then blend to those specific animations by calling the Slot name. A great example is having a character with a weapon reload animation. You may have different versions of the reload for when the player is standing, crouching, and lying prone. So long as all 3 of the animations used the same timing, you could place each one within a separate Slot in your Montage; the Slots could be named Standing, Crouching, and Prone. In your Animation Blueprint's **AnimGraph**, you can use the Slot node to determine which one you want to play based on your character's current state. When they are standing, you can use the result of the animation in the Standing Slot. When they are prone, you can see the result of the Prone Slot.
 
-It is an important point to keep in mind that although much of your Montage control will take place in the _Animation Blueprint's_ **Event Graph**, Slots are actually handled within the **Anim Graph**. This is done by way of the Slot node, which simply takes in the name for a Slot. By positioning this node at a strategic point along your AnimGraph's execution, you can have multiple Montages that utilize the same Slot name.
+It is an important point to keep in mind that although much of your Montage control will take place in the Animation Blueprint's **Event Graph**, Slots are actually handled within the **Anim Graph**. This is done by way of the Slot node, which simply takes in the name for a Slot. By positioning this node at a strategic point along your AnimGraph's execution, you can have multiple Montages that utilize the same Slot name.
 [/EXCERPT:Slots]
 
 
@@ -124,17 +126,17 @@ It is an important point to keep in mind that although much of your Montage cont
 
 Branch Points allow you to create events that coincide with animation playback. These Branch Point Events can be used in code or Blueprint to cause other things to happen, but specifically, Branch Points are useful for switching to other animation Sections within a Montage. 
 
-If you are already familiar with [Animation Notifies](Engine/Animation/Sequences\Notifies), you may notice a strong similarity between the two systems, as they both expose events that can be leveraged in script. The key difference is that Notifies are _asynchronous_, while Branch Points are _synchronous_. What this means to the end user is that Branch Points come with a much higher degree of precision for where they will take place along the animation timeline.
+If you are already familiar with [Animation Notifies](Engine/Animation/Sequences\Notifies), you may notice a strong similarity between the two systems, as they both expose events that can be leveraged in script. The key difference is that Notifies are asynchronous, while Branch Points are synchronous. What this means to the end user is that Branch Points come with a much higher degree of precision for where they will take place along the animation timeline.
 
-High precision is important when you need to jump to a specific animation at a very precise moment in time. While you _could_ use a Notify to do the same job, the asynchronous nature of Notifies means that the Notify Event could be fired at the incorrect animation frame, which can lead to hitches and jumps in your motion.
+High precision is important when you need to jump to a specific animation at a very precise moment in time. While you could use a Notify to do the same job, the asynchronous nature of Notifies means that the Notify Event could be fired at the incorrect animation frame, which can lead to hitches and jumps in your motion.
 
 [REGION:warning]
-Due to their synchronous nature and the precision resulting from it, Branch Points are more performance expensive than Notifies. You should only use them when an event _must_ be fired at a precise moment along the animation timeline, such as jumping directly to another animation that matches up frame-to-frame. If being off by a frame - or some percentage of one - is not important, you should use Notifies instead.
+Due to their synchronous nature and the precision resulting from it, Branch Points are more performance expensive than Notifies. You should only use them when an event must be fired at a precise moment along the animation timeline, such as jumping directly to another animation that matches up frame-to-frame. If being off by a frame (or some percentage of one) is not important, you should use Notifies instead.
 [/REGION]
 
 ## Sections Area
 
-The Sections area is where you can establish relationships between the Sections you define in the Montage Area. For instance, you may want a certain Section of animation - or group of Sections - to play in a specific sequence, or even loop. 
+The Sections area is where you can establish relationships between the Sections you define in the Montage Area. For instance, you may want a certain Section of animation (or group of Sections) to play in a specific sequence, or even loop. 
 
 [REGION:raw]
 ![](SectionsArea.png)
@@ -142,7 +144,7 @@ The Sections area is where you can establish relationships between the Sections
 
 
 1. **Create Default and Clear Buttons** - **Create Default** creates default associations between all Sections, stringing them together one after the other. **Clear** wipes out all associations.
-1. **Section Buttons** - In this area, you will see a button for each of the Sections you define in the Montage area. By choosing an existing Section and then clicking the one of these buttons, you associate the Section for that button with the selected track. For instance, in the image above, we have associated _Swing2_ with _Swing1_. The order actually goes _Swing1_, _Swing2_, and then _Swing1_ again, which causes a loop. See the [Looping](#Looping) section below for details.
+1. **Section Buttons** - In this area, you will see a button for each of the Sections you define in the Montage area. By choosing an existing Section and then clicking the one of these buttons, you associate the Section for that button with the selected track. For instance, in the image above, we have associated Swing2 with Swing1. The order actually goes Swing1, Swing2, and then Swing1 again, which causes a loop. See the [Looping](#Looping) section below for details.
 1. **Section Association Tracks** - This is where you can visualize and preview the relationships between animation Sections. By clicking on the **Preview** buttons, you can see the result of each individual track, or click the **Preview All Sections** button and see all Sections play back one after the other.
 
 
@@ -183,13 +185,13 @@ However, there are few things we have set up in advance:
 
 ### Creating the Montage
 
-Making a Montage is as easy as **right-clicking** in the **Content Browser** and choosing **Animation > Animation Montage**. You can also **right-click** on an existing Animation Sequence and choose **Create Montage** from the context menu. This will automatically create a new Montage with the selected AnimSequence already set up in the default Slot.
+Making a Montage is as easy as **Right-clicking** in the **Content Browser** and choosing **Animation > Animation Montage**. You can also **Right-click** on an existing Animation Sequence and choose **Create Montage** from the context menu. This will automatically create a new Montage with the selected AnimSequence already set up in the default Slot.
 
 ![](MakeMontage.png)
 
 ### Montage Setup
 
-Our first order of business was to name our Slot. We only need one in this Montage, and since we want our attack to only affect the upper body, the name _Upper Body_ seemed perfect. We then drag/dropped the animations we would need into this slot. The animations we have do the following:
+Our first order of business was to name our Slot. We only need one in this Montage, and since we want our attack to only affect the upper body, the name Upper Body seemed perfect. We then drag/dropped the animations we would need into this slot. The animations we have do the following:
 
 * Right-to-left hammer swing
 * Left-to-right hammer swing
@@ -207,9 +209,9 @@ While order is not extremely important, having the first two animations at the b
 
 ### Section Creation
 
-Our next step is to section out the animations in our Montage so that we can query them and call them up when needed in our Blueprint code. This is as easy as **right-clicking** and choosing **Add New Montage Section** while clicking on the Section Track at the top of the Montage Area. 
+Our next step is to section out the animations in our Montage so that we can query them and call them up when needed in our Blueprint code. This is as easy as **Right-clicking** and choosing **Add New Montage Section** while clicking on the Section Track at the top of the Montage Area. 
 
-We used fairly straightforward names for each Section. Note that we replaced the Default section that came along with the Montage (made a new one and deleted the Default by **right-clicking** and choosing **Delete Montage Section**). You can drag these Sections along the Section Track if you need to, and you will notice that they snap a bit at the boundary between two animation segments when you release the mouse. Use this to your advantage.
+We used fairly straightforward names for each Section. Note that we replaced the Default section that came along with the Montage (made a new one and deleted the Default by **Right-clicking** and choosing **Delete Montage Section**). You can drag these Sections along the Section Track if you need to, and you will notice that they snap a bit at the boundary between two animation segments when you release the mouse. Use this to your advantage.
 
 ![](AddNewSections.png)
 
@@ -218,17 +220,17 @@ We used fairly straightforward names for each Section. Note that we replaced the
 ### Defining Section Relationships
 
 
-Now that our Sections are created, we can now define any special relationships between them using the Sections Area. For instance, we can define a relationship between the _Swing1_ and _Swing2_ sections so that they play back to back in a loop. This is very useful for us. We start by clicking the **Clear** button to wipe out any default relationships. Then it is just a matter of selecting the track with _Swing1_ in it, and clicking the green button labeled _Swing2_ near the top of the Sections Area. This will remove the _Swing2_ track and add _Swing2_ to the _Swing1_ track.
+Now that our Sections are created, we can now define any special relationships between them using the Sections Area. For instance, we can define a relationship between the Swing1 and Swing2 sections so that they play back to back in a loop. This is very useful for us. We start by clicking the **Clear** button to wipe out any default relationships. Then it is just a matter of selecting the track with Swing1 in it, and clicking the green button labeled Swing2 near the top of the Sections Area. This will remove the Swing2 track and add Swing2 to the Swing1 track.
 
 ![](Swing2Track.png)
 
-If we repeat the process, clicking the new _Swing2_ segment and then clicking the _Swing1_ button, the system detects that you are creating a loop and the track turns blue. This means that the Swing1 and Swing2 segments are considered a looping Section. They will play back to back, repeating indefinitely.
+If we repeat the process, clicking the new Swing2 segment and then clicking the Swing1 button, the system detects that you are creating a loop and the track turns blue. This means that the Swing1 and Swing2 segments are considered a looping Section. They will play back to back, repeating indefinitely.
 
 ![](LoopingTracks.png)
 
 ### Setting Up Branch Points
 
-We will now set up some Branch Points to test whether to continue with the loop or jump to one of two possible endings for our attack animation. All we have to do is **right-click** in the Branch Point Track and choose **New Branch Point**. We chose the names _Swing_1_End_ and _Swing_2_End_ for our Branch Points. We also zoomed very closely in with the mouse wheel to make sure that each one fires right at the last possible moment of its corresponding section. This means it is placed very slightly to the left of the border between the two.
+We will now set up some Branch Points to test whether to continue with the loop or jump to one of two possible endings for our attack animation. All we have to do is **Right-click** in the Branch Point Track and choose **New Branch Point**. We chose the names Swing_1_End and Swing_2_End for our Branch Points. We also zoomed very closely in with the mouse wheel to make sure that each one fires right at the last possible moment of its corresponding section. This means it is placed very slightly to the left of the border between the two.
 
 ![](BranchPointsSetUp.png)
 
@@ -246,20 +248,20 @@ At this point, our Montage is all set up. We now need to establish our AnimGraph
 ![](StateMachineResult.png)
 
 
-We only want the Montage to play back from the Spine_01 bone (the waist) up, so we will be using a **Layered Blend per Bone** node. We add a Blend Pose to it set the weight to 1. We also associate this Blend Pose with Spine_01 in the properties for the node. We then bring in a **Slot** node and set it to _UpperBody_, the name of our slot. But now we run into a problem:
+We only want the Montage to play back from the Spine_01 bone (the waist) up, so we will be using a **Layered Blend per Bone** node. We add a Blend Pose to it set the weight to 1. We also associate this Blend Pose with Spine_01 in the properties for the node. We then bring in a **Slot** node and set it to UpperBody, the name of our slot. But now we run into a problem:
 
 ![](SlotAndBlend.png)
 
-The Slot node needs a Source connection to fall back on once the Montage is done playing. Without it, the character will return to the T-pose from the waist up after the Montage completes. However, we cannot connect the State Machine to both the Base Pose of the **Layered Blend per Bone** as well as the Source for the **Slot** node. The solution? Use a **Cache** node! We can store the result of the State Machine into a Cache node, and then create **Cached Pose** nodes to connect to both our required inputs. This is somewhat similar to storing the result of the State Machine into a variable so you can use it in multiple places. For this example, we named the Cache _LocoCache_.
+The Slot node needs a Source connection to fall back on once the Montage is done playing. Without it, the character will return to the T-pose from the waist up after the Montage completes. However, we cannot connect the State Machine to both the Base Pose of the **Layered Blend per Bone** as well as the Source for the **Slot** node. The solution? Use a **Cache** node! We can store the result of the State Machine into a Cache node, and then create **Cached Pose** nodes to connect to both our required inputs. This is somewhat similar to storing the result of the State Machine into a variable so you can use it in multiple places. For this example, we named the Cache LocoCache.
 
 ![](CachedLocomotion.png)
 
-Our AnimGraph is done. As soon as the UpperBody Slot node receives data from _any_ _AnimMontage_ (you can use any Montage, so long as it has a Slot named _UpperBody_), it will blend it in. As soon as it is no longer receiving data, it will fall back on the result of the State Machine.
+Our AnimGraph is done. As soon as the UpperBody Slot node receives data from any AnimMontage (you can use any Montage, so long as it has a Slot named UpperBody), it will blend it in. As soon as it is no longer receiving data, it will fall back on the result of the State Machine.
 
 
 ### Setting Up the Event Graph
 
-Our Event Graph setup is very basic. Using the **Event Blueprint Update Animation** node and dragging off the out pin of a **Get Player Character** node, we can **Cast To** our **Character Blueprint** (_MyCharacter_ in this example) to access the variables and functions from that Blueprint (see [Blueprint Communications](Engine\Blueprints\UserGuide\BlueprintComms) for more info on communicating between Blueprints). 
+Our Event Graph setup is very basic. Using the **Event Blueprint Update Animation** node and dragging off the out pin of a **Get Player Character** node, we can **Cast To** our **Character Blueprint** (MyCharacter in this example) to access the variables and functions from that Blueprint (see [Blueprint Communications](Engine\Blueprints\UserGuide\BlueprintComms) for more info on communicating between Blueprints). 
 
 _Click for full view_
 [REGION:lightbox]
@@ -268,7 +270,7 @@ _Click for full view_
 
 The first thing we do is check if the **IsAttacking** variable from our Character Blueprint is TRUE and if it is, we then check if the Montage is already playing. If the Montage is playing, we do not want to play it again; however if it is not playing, then we play the Montage. This prevents the system from restarting the animation in the middle of playback, which would look bad. 
 
-Dragging off the **IsAttacking** node which extends from the **Cast To** node, we see if the mouse button is still down, and if _not_, we jump to the appropriate ending animation, depending on which half of the loop is playing. Creating the Branch Point Events is done by **right-clicking** and choosing the appropriate event under **Add Montage Branching Point Event**.
+Dragging off the **IsAttacking** node which extends from the **Cast To** node, we see if the mouse button is still down, and if not, we jump to the appropriate ending animation, depending on which half of the loop is playing. Creating the Branch Point Events is done by **right-clicking** and choosing the appropriate event under **Add Montage Branching Point Event**.
 
 ![](EndLoopEventGraph.png)
 
@@ -278,12 +280,12 @@ And that is it! If we compile, we will now see that the character continues to s
 [![](HammerMontage.png)](HammerMontage.gif)
 [/REGION]
 
-**Click to watch playback**
+_Click to watch playback_
 
 <!--------------------
 
 
-[EXCERPT:MontageReplication]At this time, root motion is replicated across the server. However the _AnimMontage_ is not. That means that users utilizing root motion will need to make sure they are replicating the state of the _AnimMontage_ across the network. (Position, play rate, etc.).[/EXCERPT:MontageReplication]
+[EXCERPT:MontageReplication]At this time, root motion is replicated across the server. However the AnimMontage is not. That means that users utilizing root motion will need to make sure they are replicating the state of the AnimMontage across the network. (Position, play rate, etc.).[/EXCERPT:MontageReplication]
 
 
 -------------------->
