#UE4doc #docs:
Removed Overview header
Added engine version field
Changed font settings for multiple terms
Fixed some grammar
Replaced a picture due to color problems

diff --git a/Engine/Documentation/Source/Engine/Animation/Blendspaces/Blendspaces.INT.udn b/Engine/Documentation/Source/Engine/Animation/Blendspaces/Blendspaces.INT.udn
--- a/Engine/Documentation/Source/Engine/Animation/Blendspaces/Blendspaces.INT.udn
+++ b/Engine/Documentation/Source/Engine/Animation/Blendspaces/Blendspaces.INT.udn
@@ -10,6 +10,7 @@ Related: Engine/Animation/NodeReference/Blend
 Related: Engine/Animation/StateMachines
 Related: Engine/Animation/AnimBlueprints/AnimGraph
 Related: Engine/Content/Tools/MayaRiggingTool
+version: 4.9
 
 [VAR:Topic]
 [OBJECT:Topic]
@@ -33,7 +34,7 @@ Related: Engine/Content/Tools/MayaRiggingTool
 
 [TOC (start:2)]
 
-## Overview
+
 
 [EXCERPT:Intro]
 **Blend Spaces** are special assets which can be sampled in AnimGraphs that allow for blending of 
@@ -63,20 +64,20 @@ An example of a Locomotion **Blendspace** is also shown in the [Animation Conten
 
 ## How Blend Spaces Work
 
-The goal of _Blend Spaces_ is to reduce the need for creating individual, hard-coded nodes to perform 
+The goal of Blend Spaces is to reduce the need for creating individual, hard-coded nodes to perform 
 blending based on specific properties or conditions. By allowing the animator or programmer to specify 
 the inputs, the animations, and how the inputs are used to blend between animations, virtually any 
-type of blending can be performed using the generic _Blend Space_. This is completely contrary to how 
+type of blending can be performed using the generic Blend Space. This is completely contrary to how 
 AnimTrees in Unreal Engine 3 handled the same tasks. Any complex blending required the creation of a 
-new node to handle that blending. _Blend Spaces_ are completely generic allowing the factors that determine 
-the blending to be specified on each individual _Blend Space_. Each _Blend Space_ has inputs that simply take 
-in values. These values can be calculated during the update cycle of an _Animation Blueprint_ (via the **EventGraph**), 
+new node to handle that blending. Blend Spaces are completely generic allowing the factors that determine 
+the blending to be specified on each individual Blend Space. Each Blend Space has inputs that simply take 
+in values. These values can be calculated during the update cycle of an Animation Blueprint (via the **EventGraph**), 
 driven by gameplay code, or any other means. This makes them extremely flexible and puts the power into 
-the hands of the person creating the _AnimGraph_ to blend animations in any way they see fit.
+the hands of the person creating the AnimGraph to blend animations in any way they see fit.
 
 ![Input Flow](bs_flow.png)
 
-You can think of a _Blend Space_ like a 2D graph with each input value along one axis and animations plotted 
+You can think of a Blend Space like a 2D graph with each input value along one axis and animations plotted 
 on the graph at various points. The blended animation is calculated by blending between the animations 
 on the graph based on the location designated by the current input values.
 
@@ -98,23 +99,23 @@ perform the blending were hidden away in the code, so no modifications could be
 except by a programmer; and if you modified the value in the code, it affected every instance of the node 
 in every AnimTree making use of it. This is far from ideal.
 
-With _Blend Spaces_, this same blend can be performed using a single _Blend Space_. The direction (e.g., -180 to 
+With Blend Spaces, this same blend can be performed using a single Blend Space. The direction (e.g., -180 to 
 180 representing degrees) and speed (e.g., 0 to 250 representing units per second) of the player would be 
-values passed to the _Blend Space_ as inputs and the animations would be set up to blend appropriately based 
+values passed to the Blend Space as inputs and the animations would be set up to blend appropriately based 
 on the values of these inputs. 
 
 ![](BlendSpaceDirection.png)(w:720)
 
 In addition to only requiring a single node, the values used as inputs can be calculated directly in the 
-_Animation Blueprint's_ **EventGraph** and then passed to the _Blend Space_ in the **AnimGraph**. The input values can be 
+Animation Blueprint's **EventGraph** and then passed to the Blend Space in the **AnimGraph**. The input values can be 
 tweaked easily. The thresholds for animation blending can be quickly and intuitively adjusted using a 
 graphical editor. The flexibility and ease of use represent a huge win over previous systems.
 
 ## One-Dimensional Blend Spaces
 
-_Blend Spaces_ can also be created in a one-dimensional format, known as a **_Blend Space 1D_**. These can blend between any number of poses or animations, but do so based on a single input value. An example that we use is to make a character lean when sprinting. We have two poses, one leaning left and one leaning right, and use an input driven by how much the player is turning the character while sprinting. By blending into the additive leaning poses, we can give the running animation a bit more life as the character seems to lean into hard turns.
+Blend Spaces can also be created in a one-dimensional format, known as a **Blend Space 1D**. These can blend between any number of poses or animations, but do so based on a single input value. An example that we use is to make a character lean when sprinting. We have two poses, one leaning left and one leaning right, and use an input driven by how much the player is turning the character while sprinting. By blending into the additive leaning poses, we can give the running animation a bit more life as the character seems to lean into hard turns.
 
-Aside from having only one axis instead of two, a _Blend Space 1D_ is exactly the same as a standard two-axis _Blend Space_. As such, this document will focus on the 2D version. But you should know that all the rules and properties still carry over.
+Aside from having only one axis instead of two, a Blend Space 1D is exactly the same as a standard two-axis Blend Space. As such, this document will focus on the 2D version. But you should know that all the rules and properties still carry over.
 
 [OBJECT:EmbeddedVideo]
 [PARAMLITERAL:width]
@@ -130,26 +131,26 @@ Aside from having only one axis instead of two, a _Blend Space 1D_ is exactly th
 
 ## Architecture
 
-The underlying architecture of _Blend Spaces_ is fairly simple. The main parts were touched on previously: 
+The underlying architecture of Blend Spaces is fairly simple. The main parts were touched on previously: 
 inputs values and sampled animations.
 
 ### Skeleton
 
-Each _Blend Space_, as with other animation asset types, is targeted to a specific **Skeleton**. The Skeleton 
-associated with the _Blend Space_ is set when the _Blend Space_ is created and determines what animations can be 
-sampled within the _Blend Space_. It also determines which _Animation Blueprints_ the _Blend Space_ can be sampled within 
+Each Blend Space, as with other animation asset types, is targeted to a specific **Skeleton**. The Skeleton 
+associated with the Blend Space is set when the Blend Space is created and determines what animations can be 
+sampled within the Blend Space. It also determines which Animation Blueprints the Blend Space can be sampled within 
 as they both must target the same Skeleton.
 
 ### Blend Parameters
 
 [EXCERPT:BlendParameters]
-The **Blend Parameters** of a _Blend Space_ represent the input values used to drive the blending between 
-sample poses. Each _Blend Space_ has two Blend Parameters - **X** and **Y** - and each Blend Parameter has 
+The **Blend Parameters** of a Blend Space represent the input values used to drive the blending between 
+sample poses. Each Blend Space has two Blend Parameters (**X** and **Y**) and each Blend Parameter has 
 the following properties:
 
 | Property | Description |
 | -------- | ----------- |
-| **Label** | The readable name displayed in the **BlendSpace Editor** along the axis for this input and on the data pin for this input on the _BlendSpace_ node in the _AnimGraph_ of a _Animation Blueprint_. |
+| **Label** | The readable name displayed in the **BlendSpace Editor** along the axis for this input and on the data pin for this input on the BlendSpace node in the AnimGraph of a Animation Blueprint. |
 | **Range** | The minimum and maximum values to which any input data will be clamped. |
 | **Divisions** | The number of sampled points from the animation data used for interpolating. Higher values give more precise results. Displayed as grid lines in the **BlendSpace Editor** along the axis for this input. |
 [/EXCERPT:BlendParameters]
@@ -157,7 +158,7 @@ the following properties:
 ### Sample Data
 
 [EXCERPT:SampleData]
-The **Sample Data** for the _Blend Space_ is a collection of animations and sample values. The sample values are 
+The **Sample Data** for the Blend Space is a collection of animations and sample values. The sample values are 
 used to determine the weight of that particular animation. For instance, using the directional movement example 
 above, a **_Run_Fwd_** animation with a sample value of (0.0, 250.0, 0.0) would be fully blended to when the 
 direction input was 0.0 and the speed input was 250.0.
@@ -168,7 +169,7 @@ Each item in the Sample Data array has the following properties:
 | Property | Description |
 | -------- | ----------- |
 | **Animation** | The AnimationSequence asset to blend to. |
-| **Sample Value** | The input values (X and Y corresponding to the Blend Parameters) where the **Animation** has full influence. |
+| **Sample Value** | The input values (X and Y corresponding to the Blend Parameters) where the Animation has full influence. |
 
 
 
