#UE4doc #docs:
Added engine version field
Changed multiple images to match updated UI
Changed font settings for multiple terms

diff --git a/Engine/Documentation/Source/Engine/Animation/Blendspaces/Creation/BlendspaceCreation.INT.udn b/Engine/Documentation/Source/Engine/Animation/Blendspaces/Creation/BlendspaceCreation.INT.udn
--- a/Engine/Documentation/Source/Engine/Animation/Blendspaces/Creation/BlendspaceCreation.INT.udn
+++ b/Engine/Documentation/Source/Engine/Animation/Blendspaces/Creation/BlendspaceCreation.INT.udn
@@ -2,6 +2,7 @@ Availability:Public
 Title: Creating Blend Spaces
 Crumbs: %ROOT%, Engine, Engine/Animation, Engine/Animation/Blendspaces
 Description:Guide to creating new BlendSpace multi-input animation blending assets for use in AnimGraphs of Animation Blueprint assets.
+version: 4.9
 
 [VAR:Topic]
 [OBJECT:Topic]
@@ -24,23 +25,23 @@ Description:Guide to creating new BlendSpace multi-input animation blending asse
 [/VAR]
 
 
-_Blend Spaces_ are regular assets and must be created through the **Content Browser** like any other asset.
+Blend Spaces are regular assets and must be created through the **Content Browser** like any other asset.
 
-**To create a new _Blend Space_:**
+To create a new Blend Space:
 
-1. In the **Content Browser**, navigate to the folder within which you want to create your new _Blend Space_, then click ![](Shared/EditorButtons/button_NewAsset.png) or **right-click** and choose **Animation** > **Blend Space**:
+1. In the **Content Browser**, navigate to the folder within which you want to create your new Blend Space, then click ![](Shared/EditorButtons/button_NewAsset.png) or **Right-click** and choose **Animation** > **Blend Space**:
 
 	![](NewBlendSpace.png)
 
-1. In the **Choose Skeleton** dialog, select the Skeleton the _Blend Space_ should target:
+1. In the **Pick Skeleton** dialog, select the Skeleton the Blend Space should target:
 
 	![](SkeletonList.png)
 
-1. The new _Blend Space_ asset is displayed in the **Content Browser** at the location specified:
+1. The new Blend Space asset is displayed in the **Content Browser** at the location specified:
 
 	![Create Blueprint - Asset Created](bs_create_asset.png)
 
-	The _Blend Space_ is also available to be used in the **Anim Graphs** of _Animation Blueprints_ targeting the same Skeleton as the _Blend Space_:
+	The Blend Space is also available to be used in the **Anim Graphs** of Animation Blueprints targeting the same Skeleton as the Blend Space:
 
 	[REGION:imagetable]
 	| ![](AssetBrowserBlendList.png)(h:400) | ![](ContextMenuBlendList.png)(h:400) |
