#UE4doc #docs:
Added engine version field
Changed font settings for multiple terms
Updated/replaced multiple images to match doc guidelines and to show updated UI

diff --git a/Engine/Documentation/Source/Engine/Animation/Blendspaces/Editor/BlendspaceEditor.INT.udn b/Engine/Documentation/Source/Engine/Animation/Blendspaces/Editor/BlendspaceEditor.INT.udn
--- a/Engine/Documentation/Source/Engine/Animation/Blendspaces/Editor/BlendspaceEditor.INT.udn
+++ b/Engine/Documentation/Source/Engine/Animation/Blendspaces/Editor/BlendspaceEditor.INT.udn
@@ -2,6 +2,7 @@ Availability:Public
 Title: Editing Blend Spaces
 Crumbs: %ROOT%, Engine, Engine/Animation, Engine/Animation/Blendspaces
 Description:Guide to setting up sample parameter and animation samples in Blendspaces for complex, multi-input animation blending.
+version: 4.9
 
 
 [VAR:topic]
@@ -31,7 +32,7 @@ Description:Guide to setting up sample parameter and animation samples in Blends
 [/REGION]
 
 [EXCERPT:overview]
-The **Blend Space Editor** provides the ability to set up the Blend Parameters of the _BlendSpace_ as 
+The **Blend Space Editor** provides the ability to set up the Blend Parameters of the BlendSpace as 
 well as graphically edit the positions of Sample poses and interactively see the results of blending 
 between those Sample poses.
 [/EXCERPT:overview]
@@ -42,22 +43,22 @@ between those Sample poses.
 
 ![](TestBlendShape.png)(w:720)
 
-1. **Parameters** - Blend Parameters property section for setting up the _BlendSpace_.
-1. **Samples** - Editable list of all Samples in the _BlendSpace_.
-1. **Options** - Preview options for the _BlendSpace_.
+1. **Parameters** - Blend Parameters property section for setting up the BlendSpace.
+1. **Samples** - Editable list of all Samples in the BlendSpace.
+1. **Options** - Preview options for the BlendSpace.
 1. **Graph** - 2D graph for placing Samples and previewing blending between them.
 
 
 ## Setting up Parameters
 
-Parameters for the _BlendSpace_ are set up in the **Parameters** section of the UI.
+Parameters for the BlendSpace are set up in the **Parameters** section of the UI.
 
 ![Blendspace Parameters](bs_params.png)
 
 [INCLUDE:Engine/Animation/Blendspaces#BlendParameters]
 
 Once the appropriate values have been set for each of the Blend Parameters, pressing the ![Apply Parameter Changes](button_apply_changes.png) 
-button will apply the values to the _BlendSpace_, updating the **Graph**.
+button will apply the values to the BlendSpace, updating the **Graph**.
 
 ![](BlendSpaceSetup.png)
 
@@ -67,12 +68,12 @@ button will apply the values to the _BlendSpace_, updating the **Graph**.
 
 ### Adding Samples
 
-Samples are added to a _BlendSpace_ by dragging an _AnimationSequence_ from the **Asset Browser** 
+Samples are added to a BlendSpace by dragging an AnimationSequence from the **Asset Browser** 
 into the **Graph**.
 
 [REGION:note]
-When working with _AimOffset BlendSpaces_, the _AnimationSequences_ must have **Additive Anim Type** 
-set to _AAT_RotationOffsetMeshSpace_ in order to be added as Samples.
+When working with AimOffset BlendSpaces, the AnimationSequences must have **Additive Anim Type** 
+set to AAT_RotationOffsetMeshSpace in order to be added as Samples.
 [/REGION]
 
 ![](BlendShapeDragDrop.png)(w:720)
@@ -86,13 +87,13 @@ A node is placed at the location in the **Graph**:
 Samples can be moved to different positions interactively or manually. Moving Samples interactively 
 is very intuitive, but does not provide the preciseness of manually setting the Sample's position.
 
-**To interactively move a sample:**
+To interactively move a sample:
 
 Click and dragging the node in the **Graph**:
 
 ![Move Sample by Dragging](bs_sample_move_drag.png)(convert:false)
  
-**To move a sample manually:**
+To move a sample manually:
 
 Click the node and then specify the  **X** and **Y** values to place the node at:
 
@@ -100,12 +101,12 @@ Click the node and then specify the  **X** and **Y** values to place the node at
 
 ### Replacing Samples
 
-Samples placed in the _BlendSpace_ can be replaced with another Sample in the **Graph** via _drag and drop_ or using the 
+Samples placed in the BlendSpace can be replaced with another Sample in the **Graph** via drag and drop or using the 
 **Samples** section of the UI.
 
-**To replace by _drag and drop_:**
+To replace by drag and drop:
 
-To replace a sample in the **Graph**, drag a new _AnimationSequence_ from the **Asset Browser** directly on 
+To replace a sample in the **Graph**, drag a new AnimationSequence from the **Asset Browser** directly on 
 to the existing node.
 
 [REGION:fullwidth]
@@ -118,13 +119,13 @@ to the existing node.
 | Original Sample | Replaced Sample |
 [/REGION]
 
-**To replace manually:**
+To replace manually:
 
-To replace a sample in the **Samples** section, select a new _AnimationSequence_ in the **Content Browser** 
+To replace a sample in the **Samples** section, select a new AnimationSequence in the **Content Browser** 
 and press the ![Use Selected Asset](button_use_selected.png) button for the Sample in the list.
 
 [REGION:note]
-	In this case, you must select your _AnimationSequence_ from the **Content Browser**. This operation will _not_ read any assets selected in Persona's **Asset Browser**.
+	In this case, you must select your AnimationSequence from the **Content Browser**. This operation will not read any assets selected in Persona's **Asset Browser**.
 [/REGION]
 
 ![Select New Animation](bs_sample_replace_manual_select.png)(convert:false)
@@ -135,14 +136,14 @@ and press the ![Use Selected Asset](button_use_selected.png) button for the Samp
 
 Samples can be deleted directly in the **Graph** or using the **Samples** section of the UI.
 
-**To delete in the Graph:**
+To delete in the Graph:
 
 Select the Sample node in the **Graph** to display the edit dialog and press the ![Delete](button_delete.png) 
 button to remove the Sample.
 
 ![Remove Sample in Graph](bs_sample_remove.png)(convert:false)
 
-**To delete in the Samples section:**
+To delete in the Samples section:
 
 Press the ![Remove Sample](button_clear_reference.png) button for the Sample to be removed from the list.
 
@@ -154,7 +155,7 @@ The result of the blending between poses can be previewed interactively in the *
 using the **Graph**. The position of the cursor determines the values of the **X** and **Y** parameters used for 
 performing the blending and the resulting pose is displayed in the **Viewport** panel.
 
-To enable previewing for the _BlendSpace_, check the **Enable Preview BlendSpace** option in the **Options** 
+To enable previewing for the BlendSpace, check the **Enable Preview BlendSpace** option in the **Options** 
 section of the UI.
 
 ![Enable Preview Blendspace Option](bs_options_preview.png)
