#UE4doc #docs:
Removed Overview header
Added engine version field
Changed color of arrow in an image to meet doc guidelines
Changed font settings for multiple terms

diff --git a/Engine/Documentation/Source/Engine/Animation/Blendspaces/UserGuide/BlendspaceUserGuide.INT.udn b/Engine/Documentation/Source/Engine/Animation/Blendspaces/UserGuide/BlendspaceUserGuide.INT.udn
--- a/Engine/Documentation/Source/Engine/Animation/Blendspaces/UserGuide/BlendspaceUserGuide.INT.udn
+++ b/Engine/Documentation/Source/Engine/Animation/Blendspaces/UserGuide/BlendspaceUserGuide.INT.udn
@@ -2,6 +2,7 @@ Availability:Public
 Title: Using Blend Spaces
 Crumbs: %ROOT%, Engine, Engine/Animation, Engine/Animation/Blendspaces
 Description:Guide to using Blend Spaces in Animation Blueprints to input data and sample output animation poses.
+version: 4.9
 
 [VAR:Topic]
 [OBJECT:Topic]
@@ -25,7 +26,6 @@ Description:Guide to using Blend Spaces in Animation Blueprints to input data an
 
 [TOC(start:2)]
 
-## Overview
 
 [EXCERPT:Intro]
 Blend Spaces are used by placing nodes in the **AnimGraphs** of Animation Blueprints. These nodes take in value data that is used to drive which animations the Blend Space uses to create the final blended pose.
@@ -33,11 +33,11 @@ Blend Spaces are used by placing nodes in the **AnimGraphs** of Animation Bluepr
 
 ## Adding Blend Space Nodes
 
-_Blend Space_ nodes provide access to the sample parameter data and output pose of the _Blend Space_.
+Blend Space nodes provide access to the sample parameter data and output pose of the Blend Space.
  
-**To place a _Blend Space_ node using drag and drop:**
+To place a Blend Space node using drag and drop:
 
-1. In **Persona**, find the _Blend Space_ you want to sample in the **Asset Browser**.
+1. In **Persona**, find the Blend Space you want to sample in the **Asset Browser**.
 
 	![](SelectBlendSpace.png)
 
@@ -45,37 +45,37 @@ _Blend Space_ nodes provide access to the sample parameter data and output pose
 
 	![](DragDropBlendSpace.png)(w:720)
 
-1. Release the **Left Mouse Button** to place the _Blend Space_ node.
+1. Release the **Left Mouse Button** to place the Blend Space node.
 
 	![](BlendSpaceNodeAdded.png)
 
 
-**To place a _Blend Space_ node using the context menu:**
+**To place a Blend Space node using the context menu:**
 
-1. **Right-click** in the **AnimGraph** and expand **Animation** to see a list of all compatible _Blend Spaces_.
+1. **Right-click** in the **AnimGraph** and expand **Animation** to see a list of all compatible Blend Spaces.
 
 	![](AnimationContextMenu.png)
 
-1. Click on the _Blend Space_ you want to sample.
+1. Click on the Blend Space you want to sample.
 
 	![](BlendSpaceContextSelection.png)
 
-1. The _Blend Space_ node is placed.
+1. The Blend Space node is placed.
 
 	![](BlendSpaceNodeAdded.png)
 
 
 ## Inputting Data
 
-_Blend Spaces_ rely on data being passed into the sample parameters in order to perform the animation blending. Each _Blend Space_ node has input data pins for each sample parameter. Variables or outputs from other nodes can be wired to these pins to pass the necessary data to the _Blend Space_.
+Blend Spaces rely on data being passed into the sample parameters in order to perform the animation blending. Each Blend Space node has input data pins for each sample parameter. Variables or outputs from other nodes can be wired to these pins to pass the necessary data to the Blend Space.
 
 ![Variable Input Data](bs_input_variables.png)
 
 [REGION:note]
-Since _Blend Space_ editing is currently limited to using two inputs, the third input pin is not used.
+Since Blend Space editing is currently limited to using two inputs, the third input pin is not used.
 [/REGION]
 
-The values passed to the _Blend Space_ are generally variables calculated during the update cycle in the **EventGraph** or through gameplay code. An example of the Direction and Speed variables being calculated in the **EventGraph** of an _Animation Blueprint_ is shown below.
+The values passed to the Blend Space are generally variables calculated during the update cycle in the **EventGraph** or through gameplay code. An example of the Direction and Speed variables being calculated in the **EventGraph** of an Animation Blueprint is shown below.
 
 [REGION:fullwidth]
 ![Variable Input Data - Update Cycle](bs_input_calc.png)
@@ -83,7 +83,7 @@ The values passed to the _Blend Space_ are generally variables calculated during
 
 ## Output Pose
 
-The **Pose** output pin of the _Blend Space_ node can be wired to the **Pose** input pin of the **Result** node in the **AnimGraph** in the simplest case.
+The **Pose** output pin of the Blend Space node can be wired to the **Pose** input pin of the **Result** node in the **AnimGraph** in the simplest case.
 
 ![Output Pose - Result](bs_output_result.png)
 
