UEDOC-2243

Images at the bottom of the page has been converted into a YouTube video

diff --git a/Engine/Documentation/Source/Engine/Animation/AnimMontage/AnimMontage.INT.udn b/Engine/Documentation/Source/Engine/Animation/AnimMontage/AnimMontage.INT.udn
--- a/Engine/Documentation/Source/Engine/Animation/AnimMontage/AnimMontage.INT.udn
+++ b/Engine/Documentation/Source/Engine/Animation/AnimMontage/AnimMontage.INT.udn
@@ -276,11 +276,17 @@ Dragging off the **IsAttacking** node which extends from the **Cast To** node, w
 
 And that is it! If we compile, we will now see that the character continues to swing as long as the mouse button is down, and performs an intelligent end animation when the mouse button is released!
 
-[REGION:lightbox]
-[![](HammerMontage.png)](HammerMontage.gif)
-[/REGION]
-
-_Click to watch playback_
+[OBJECT:EmbeddedVideo]
+[PARAMLITERAL:width]
+640
+[/PARAMLITERAL]
+[PARAMLITERAL:height]
+360
+[/PARAMLITERAL]
+[PARAMLITERAL:videoid]
+6eMxjbkn6hs
+[/PARAMLITERAL]
+[/OBJECT]
 
 <!--------------------
 
