#UE4doc #docs:
Removed Overview header
Added engine version field
Replaced multiple images to reflect the updated UI
Changed font settings for multiple terms

diff --git a/Engine/Documentation/Source/Engine/Animation/CharacterSetupOverview/CharacterSetupOverview.INT.udn b/Engine/Documentation/Source/Engine/Animation/CharacterSetupOverview/CharacterSetupOverview.INT.udn
--- a/Engine/Documentation/Source/Engine/Animation/CharacterSetupOverview/CharacterSetupOverview.INT.udn
+++ b/Engine/Documentation/Source/Engine/Animation/CharacterSetupOverview/CharacterSetupOverview.INT.udn
@@ -10,6 +10,7 @@ Related: Engine/Animation/AnimBlueprints
 Related: Engine/Animation/AnimBlueprints/EventGraph
 Related: Engine/Animation/StateMachines
 Related: Gameplay/AI
+version: 4.9
 
 [REGION:fullwidth]
 ![](Character.png)
@@ -17,7 +18,7 @@ Related: Gameplay/AI
 
 [TOC(start:2 end:2)]
 
-## Overview
+
 
 No matter your game project or genre, it is likely that at some point you are going to need some kind of animated character to move around in your environment. This may be a character that the player controls, or may be some AI-driven entity that interacts with the world in some way. Regardless, you are going to need to know how to set such characters up so that they can properly animate in your world. The purpose of this document is to give you a very high-level overview of how this is done, while guiding you to dedicated documents and examples for specific details. For our purposes, we will assume that you want to create a character that is controllable by the player in some way.
 
@@ -37,7 +38,7 @@ The primary workflow for character setup in UE4 is as follows:
 1. Create your art assets (Skeletal Meshes) and animations, using a 3rd party digital content creation (DCC) package such as 3ds Max or Maya.
 1. Import your Skeletal Meshes and animations into UE4, either creating a new Skeleton asset for new Skeletal Meshes, or reusing an existing Skeleton asset for identical or similar Skeletal Meshes.
 1. Create a PlayerController script or Blueprint to handle inputs from the player.
-1. Create a Character or Pawn script or Blueprint to parse inputs and control the actual movement (_not skeletal animation_) of the character.
+1. Create a Character or Pawn script or Blueprint to parse inputs and control the actual movement (not skeletal animation) of the character.
 1. Construct the Animation Blueprint for the character.
 1. Create a GameMode script or Blueprint that utilizes your custom PlayerController and any other custom script assets.
 1. Play your game!
@@ -84,24 +85,24 @@ For an example of a custom Blueprint PlayerController, you can start a new proje
 
 Once within the new project, you can search within the **Class Viewer** for PlayerController, turning off the filters in the **Class Viewer**. **Double-clicking** on this asset will open it up and you can see the setup for yourself.
 
-You can also see a PlayerController in C++ script by creating a new project (**File > New Project**) and choosing the **Blueprint Top Down** template. 
+You can also see a PlayerController in C++ script by creating a new project (**File > New Project**) and choosing the **C++ Top Down** template. 
 
 
 ## Creating a Pawn or Character Blueprint
 
 ![](CharacterBP.png)(w:600)
 
-Once you have set up your PlayerController, your system is now prepared to handle inputs from the player. Now, however, you have to translate those inputs into something that can in turn drive a character around on the screen. That means those inputs need to be _translated_ (or parsed) into actions. This is where the Pawn or Character classes come into play. 
+Once you have set up your PlayerController, your system is now prepared to handle inputs from the player. Now, however, you have to translate those inputs into something that can in turn drive a character around on the screen. That means those inputs need to be translated (or parsed) into actions. This is where the Pawn or Character classes come into play. 
 
 ### Choosing Pawn or Character
 
-You will notice that we mention two potential classes here: Pawn and Character. Both are used for entities in the game that are either controlled by the player or by in-game AI. The key difference is that the Character class is an _extension_ of the Pawn class, adding in player physics, support for a specific mesh, and the general types of handling needed when creating a playable in-game character. For our purposes, we will be using the Character class. For simpler elements that would just need to be driven around the scene by AI, for example, you can generally get away with using a Pawn.
+You will notice that we mention two potential classes here: Pawn and Character. Both are used for entities in the game that are either controlled by the player or by in-game AI. The key difference is that the Character class is an extension of the Pawn class, adding in player physics, support for a specific mesh, and the general types of handling needed when creating a playable in-game character. For our purposes, we will be using the Character class. For simpler elements that would just need to be driven around the scene by AI, for example, you can generally get away with using a Pawn.
 
 ### Character Class Setup
 
 ![](CharacterBPGraph.png)(w:800)
 
-Your Character Class is going to start with events that are triggered from the PlayerController, and use scripting (including Blueprint visual scripting) to control what to actually _do_ with those inputs and how they can be used to control the character. For instance, where the PlayerController simply creates a basic event for moving the analog stick on a controller in an upward direction, the Character class is responsible for receiving that event and using it to drive the character forward.
+Your Character Class is going to start with events that are triggered from the PlayerController, and use scripting (including Blueprint visual scripting) to control what to actually do with those inputs and how they can be used to control the character. For instance, where the PlayerController simply creates a basic event for moving the analog stick on a controller in an upward direction, the Character class is responsible for receiving that event and using it to drive the character forward.
 
 The Character Class also holds a reference to a Skeletal Mesh, which will be the basis for what the player sees while playing the game. In the case of a first-person game, this is often just a pair of floating arms, though there may be a full body if you need that body to shadow the environment properly. For third-person games, the mesh will be the Skeletal Mesh that represents the character. 
 
@@ -132,7 +133,7 @@ Animation Blueprints are by far the most sophisticated aspect of character setup
 
 These are just the tip of the iceberg. You would do well to look at the [](Engine/Animation/AnimBlueprints) documentation, and to also look at some of the Animation Blueprint included in some of our example content, such as the First and Third Person Templates and those found within the Content Examples project.
 
-Once you have created the Animation Blueprint that defines the motion of your character, you will need to make sure you assign it to the **Anim Blueprint Generated Class** property, found under the _Mesh Component_ of the Character Blueprint. This is necessary because you may have multiple Animation Blueprints for each Skeletal Mesh, and the Character Blueprint needs to know the one into which it will be sending the necessary animation and variable data.
+Once you have created the Animation Blueprint that defines the motion of your character, you will need to make sure you assign it to the **Anim Blueprint Generated Class** property, found under the Mesh Component of the Character Blueprint. This is necessary because you may have multiple Animation Blueprints for each Skeletal Mesh, and the Character Blueprint needs to know the one into which it will be sending the necessary animation and variable data.
 
 ## GameMode Setup
 
@@ -144,6 +145,7 @@ A GameMode is a special type of class that is used to define your game. Generall
 * **HUD Class** - This contains any special heads-up display (HUD) classes created, which are not covered in this document.
 * **PlayerController Class** - This will hold the custom PlayerController you set up for your in-game character.
 * **Spectator Class** - This holds any special classes used to control spectators, or passive players that just watch the action. These are not covered in this document.
+* **Replay Spectator Class** - This holds any special classes used to control spectators when in a replay. These are not covered in this document.
 * **Game State Class** - This class controls any special rules or setup for how your game will be played, but is not covered in this document.
 * **Player State Class** - This class defines any special rules for replicating relevant player information to all clients.
 
@@ -153,9 +155,9 @@ For testing out your character, you need to at the very least set up the Default
 
 ![](button_Toolbar_WorldSettings.png)
 
-Once you have set up your GameMode, the last step to being able to play with your custom Character is to make sure that the current Level is using your GameMode. This is done using the World Settings tab, accessible from the World Settings button located on the main toolbar.
+Once you have set up your GameMode, the last step to being able to play with your custom Character is to make sure that the current Level is using your GameMode. This is done using the **World Settings** tab, accessible from the **Settings** button located on the main toolbar.
 
-Within the World Settings, you will want to make sure you have set the GameMode Override to the name of your GameMode class. Once you have done that, you are ready to save and test out your new character!
+Within the **World Settings**, you will want to make sure you have set the **GameMode Override** to the name of your GameMode class. Once you have done that, you are ready to save and test out your new character!
 
 ![](GameModeSettings.png)
 
@@ -200,7 +202,7 @@ Note that each of these Templates is available in a code form or a Blueprint for
 ![](AnimationCE.png)
 [/REGION]
 
-Content Examples are specialized versions of content designed by artists and technicians here at Epic. They are found within a project named **ContentExamples**, which can be downloaded by users via the **Marketplace**. Of particular importance would be the assets found within the _Animation.umap_ level, which shows a variety of uses for Skeletal Mesh animation on a character.
+Content Examples are specialized versions of content designed by artists and technicians here at Epic. They are found within a project named **ContentExamples**, which can be downloaded by users via the **Marketplace**. Of particular importance would be the assets found within the Animation.umap level, which shows a variety of uses for Skeletal Mesh animation on a character.
 
 
 
