#UE4doc #docs:
UEDOC-2244 - Added a section for the new Element Timing feature along with pictures to demonstrate the tool

diff --git a/Engine/Documentation/Source/Engine/Animation/AnimMontage/AnimMontage.INT.udn b/Engine/Documentation/Source/Engine/Animation/AnimMontage/AnimMontage.INT.udn
--- a/Engine/Documentation/Source/Engine/Animation/AnimMontage/AnimMontage.INT.udn
+++ b/Engine/Documentation/Source/Engine/Animation/AnimMontage/AnimMontage.INT.udn
@@ -83,7 +83,7 @@ When looking at a Montage in Persona, it is useful to know what each area is and
 
 1. [Montage Area](#MontageArea)
 1. [Sections Area](#SectionsArea)
-1. Element Timing Area
+1. [Element Timing Area](#elementtimingarea)
 1. [Notifies Area](#NotifiesArea)
 1. [Curves Area](#CurvesArea)
 
@@ -152,6 +152,28 @@ The Sections area is where you can establish relationships between the Sections
 
 Sections can be set up to loop indefinitely, which is extremely useful for any action that you need to repeat. By associating the same Section more than once in a Section Association Track, you cause that association to run in a loop. The section will turn blue to show this. As an example, consider an animation in which a character is reloading a shotgun, one shell at a time. You could take just the section in which the character inserts a shell, and loop it. Then, by using Notifies, you could create Notify Events in Blueprint that increment the ammo count each time the animation plays through. Once that count reaches a set number (full ammo), you could then switch to an animation of the character closing the receiver and returning to idle.
 
+## Element Timing Area
+
+The Element Timing area takes information from the Montage area and the Notifies area to help with the timing of different sections.
+
+![](ElementTiming1.png)
+
+Each node in the track is given a number denoting the trigger order of that object in the entire montage, and are coloured to match what they describe. The default colours are:
+* Red - Notifies, Notify states and notify state ending markers
+* Light Blue - Branching point notifies
+* Green - Montage sections
+
+[REGION:note]
+These colors can be edited in the **Animation Editor** preferences, which can be found in the **Editor Preferences**.
+[/REGION]
+
+Hovering over any of the nodes in the track will cause a tooltip to appear that explains what that node represents.
+
+![](ElementTiming2.png)
+
+There is dropdown menu on the right side of the track to toggle different elements' visibility.
+
+![](ElementTiming3.png)
 
 ## Notifies Area
 
