#UE4doc #docs:
Removed Overview header
Added engine version field
Changed font settings for multiple terms

diff --git a/Engine/Documentation/Source/Engine/Animation/IKSetups/IKSetups.INT.udn b/Engine/Documentation/Source/Engine/Animation/IKSetups/IKSetups.INT.udn
--- a/Engine/Documentation/Source/Engine/Animation/IKSetups/IKSetups.INT.udn
+++ b/Engine/Documentation/Source/Engine/Animation/IKSetups/IKSetups.INT.udn
@@ -8,6 +8,7 @@ Related: Engine/Animation/PhysicallyDrivenAnimation
 Related: Engine/Content/Tools/MayaRiggingTool
 Related: Engine/Animation/AnimBlueprints
 Related: Engine/Animation/AnimBlueprints/EventGraph
+version: 4.9
 
 
 [REGION:fullwidth]
@@ -19,14 +20,14 @@ On the left, the character is not using IK setups. In the middle, IK is used to
 [/REGION]
 
 
-## Overview
 
-Inverse Kinematics (IK) provide a way to handle joint rotation from the location of an end-effector rather than via direct joint rotation. In practice, you provide an effector location and the IK solution then solves the rotation so that the final joint coincides with that location as best it can. This can be used to keep a character's feet planted on uneven ground, and in other ways to produce believable interactions with the world.
+
+**Inverse Kinematics** (**IK**) provide a way to handle joint rotation from the location of an end-effector rather than via direct joint rotation. In practice, you provide an effector location and the IK solution then solves the rotation so that the final joint coincides with that location as best it can. This can be used to keep a character's feet planted on uneven ground, and in other ways to produce believable interactions with the world.
 
 Unreal Engine uses a 2-bone IK system that is ideal for things such as arms and legs.
 
 [REGION:tip]
-For examples of Hand and Foot **IK**, you can also reference the [Animation Content Examples](Resources\ContentExamples\Animation) page under section 1.8.
+For examples of Hand and Foot IK, you can also reference the [Animation Content Examples](Resources\ContentExamples\Animation) page under section 1.8.
 [/REGION]
 
 
@@ -34,13 +35,13 @@ For examples of Hand and Foot **IK**, you can also reference the [Animation Cont
 
 If you are already familiar with what IK is and how it works, you can skip this section!
 
-Most animated skeletons in Unreal are driven by direct rotational data fed straight into the bones of the character or Skeletal Mesh. This can be thought of as _forward kinematics_, or direct application of rotation to joints or bones. Below is a diagram illustrating the concept:
+Most animated skeletons in Unreal are driven by direct rotational data fed straight into the bones of the character or Skeletal Mesh. This can be thought of as forward kinematics, or direct application of rotation to joints or bones. Below is a diagram illustrating the concept:
 
 [REGION:raw]
 ![](diagram_FK.png)(w:600)
 [/REGION]
 
-As its name implies, _inverse kinematics_ works in the other direction. Instead of applying rotation to bones, we instead give the bone chain a target (also known as an _end effector_), providing a position that the end of the chain should try to achieve. The user or animator moves the effector and the IK solver (the algorithm that drives rotation in an IK system) rotates the bones such that the final bone in the chain ends at the location of the target. In the image below, the end effector is designated by the red cross. 
+As its name implies, inverse kinematics works in the other direction. Instead of applying rotation to bones, we instead give the bone chain a target (also known as an end effector), providing a position that the end of the chain should try to achieve. The user or animator moves the effector and the IK solver (the algorithm that drives rotation in an IK system) rotates the bones such that the final bone in the chain ends at the location of the target. In the image below, the end effector is designated by the red cross. 
 
 [REGION:raw]
 ![](diagramIK.png)(w:600)
@@ -62,7 +63,7 @@ One of the more important considerations for IK use is that it generally require
 
 ### Planning
 
-As with all things, a little bit of planning goes a long way. Make sure you have an idea if _what_ you need your IK setup to do. Is it for a character's feet? Their hands? What will they be doing where they will need to react? The more of these questions you can answer early, the easier the setup will be. Fortunately, with the power of UE4's Blueprint visual scripting, it will be easy enough to add functionality later.
+As with all things, a little bit of planning goes a long way. Make sure you have an idea if what you need your IK setup to do. Is it for a character's feet? Their hands? What will they be doing where they will need to react? The more of these questions you can answer early, the easier the setup will be. Fortunately, with the power of UE4's Blueprint visual scripting, it will be easy enough to add functionality later.
 
 
 ## Foot Placement IK Setup
@@ -70,7 +71,7 @@ As with all things, a little bit of planning goes a long way. Make sure you have
 For the first example, we will give an overview of setting up simple IK on a character to help their feet remain planted on uneven ground.
 
 [REGION:note]
-This example can be found in the Content Examples project. Just open the map named Animation.umap and look at example 1.8!
+This example can be found in the Content Examples project. Just open the map named Animation.umap and look at example 1.8.
 [/REGION]
 
 ### Pawn or Character Blueprint Setup
@@ -95,7 +96,7 @@ The Construction Script of the Character Blueprint really just sets up two criti
 [/REGION]
 
 [REGION:caption]
-Click the image for a larger view or **right-click** and Save As.
+Click the image for a larger view or **Right-click** and Save As.
 [/REGION]
 #### Event Graph
 
@@ -111,7 +112,7 @@ One of the important points about this graph is the use of a custom function cal
 [/REGION]
 
 [REGION:caption]
-In this image, you can see the IKFoottrace function. Click the image for a larger view or **right-click** and Save As.
+In this image, you can see the IKFoottrace function. Click the image for a larger view or **Right-click** and Save As.
 [/REGION]
 
 
@@ -125,7 +126,7 @@ And here is the event graph. With the help of the above function, you can see th
 [/REGION]
 
 [REGION:caption]
-Here is the base level of the Event Graph. Click the image for a larger view or **right-click** and Save As.
+Here is the base level of the Event Graph. Click the image for a larger view or **Right-click** and Save As.
 [/REGION]
 
 The result of this is that during each tick of the game, there is a downward trace taking place, looking for an impact point which would designate some uneven piece of ground to be accounted for. When found, the distance to that point is stored as an IK offset to be used later on in the Animation Blueprint.
@@ -157,13 +158,13 @@ With the IK offset data that was stored in the Pawn Blueprint, we can generate l
 [/REGION]
 
 [REGION:caption]
-Click the image for a larger view or **right-click** and Save As.
+Click the image for a larger view or **Right-click** and Save As.
 [/REGION]
 
 
 #### AnimGraph
 
-The AnimGraph culminates our setup by applying the information assembled thus far and using it to update the existing animation created for the character. For this example, the AnimGraph is very simple in that it is really only playing a single animation. In most cases, the _Play JumpingJacks_ node would be replaced by any number of other nodes to produce the desired motion. 
+The AnimGraph culminates our setup by applying the information assembled thus far and using it to update the existing animation created for the character. For this example, the AnimGraph is very simple in that it is really only playing a single animation. In most cases, the Play JumpingJacks node would be replaced by any number of other nodes to produce the desired motion. 
 
 The important part is where we switch our current space from Local to Component. All animations that are played on the character are done in Local space. That is the fastest way to calculate it, since local transformations are relative to each bone's parent in the hierarchy. However, bone manipulations, such as applying 2-bone IK, must be done in Component Space, which is relative to the Root bone. 
 
@@ -179,7 +180,7 @@ For more information on coordinate spaces for animation, please see [](Engine/Ba
 [/REGION]
 
 [REGION:caption]
-Click the image for a larger view or **right-click** and Save As.
+Click the image for a larger view or **Right-click** and Save As.
 [/REGION]
 
 
