#UE4doc #docs - Fixed 404 page error for BasicUsage. Doc was not made public. Added engine version tag as well. This fixes UEDOC-2246.

diff --git a/Engine/Documentation/Source/Engine/Blueprints/BP_HowTo/BasicUsage/BasicUsage.INT.udn b/Engine/Documentation/Source/Engine/Blueprints/BP_HowTo/BasicUsage/BasicUsage.INT.udn
--- a/Engine/Documentation/Source/Engine/Blueprints/BP_HowTo/BasicUsage/BasicUsage.INT.udn
+++ b/Engine/Documentation/Source/Engine/Blueprints/BP_HowTo/BasicUsage/BasicUsage.INT.udn
@@ -1,4 +1,4 @@
-Availability: docs
+Availability: Public
 Title:Blueprint Basic User Guide
 Crumbs: %ROOT%, Engine, Engine\Blueprints, Engine\Blueprints\BP_HowTo
 Description:A quick guide on the elements needed to begin working with Blueprints. 
@@ -7,10 +7,11 @@ Related: Engine/Blueprints/GettingStarted
 Related: Engine/Blueprints/UserGuide
 Related: Engine/Blueprints/UserGuide/CheatSheet
 Related: Gameplay
+Version: 4.9
 
 [TOC(start:2 end:2)]
 
-## Overview
+
 This page aims to quickly get you started using **Blueprints** by covering the most basic use cases and commonly used actions. 
 
 For more detailed information on Blueprints, refer to the [](Engine/Blueprints) documentation. 
