#UE4 Docs: Submitting the first pass on the new docs as well as moving the old Android doc to a new location.

diff --git a/Engine/Documentation/Source/Platforms/Android/GettingStarted/1/AndroidQS_1.INT.udn b/Engine/Documentation/Source/Platforms/Android/GettingStarted/1/AndroidQS_1.INT.udn
--- a/Engine/Documentation/Source/Platforms/Android/GettingStarted/1/AndroidQS_1.INT.udn
+++ b/Engine/Documentation/Source/Platforms/Android/GettingStarted/1/AndroidQS_1.INT.udn
@@ -1,9 +1,9 @@
 Availability:Public
 Title:1. Install the Android SDK
 Crumbs:%ROOT%, Platforms, Platforms/Android, Platforms/Android/GettingStarted
-Description: Install the TADP Android SDK.
-
-
+Description: Install the Android Works SDK.
+Version:4.9
+Skilllevel:Intermediate
 
 [VAR:Steps]
 [OBJECT:Navigation]
@@ -20,6 +20,8 @@ Description: Install the TADP Android SDK.
 [/OBJECT]
 [/VAR]
 
+
+
 %Steps%
 
 %Globals:OSSelect%
@@ -53,169 +55,172 @@ Description: Install the TADP Android SDK.
 	[/PARAM]
 [/OBJECT]
 
+
+
 <!--
 [EXCERPT:WindowsSetup]
 
-[REGION:tip]
-If you already have the Android SDK installed, but your environment variables have been lost, you can skip all of this and simply point the editor to the install locations using the **Edit Menu** -> **Project Settings** -> **Android SDK** properties in the Editor.
+[REGION:note]
+If you already have Android Works installed, but your environment variables have been lost, you can skip all of this and simply point the editor to the install locations using the **Edit Menu** > **Project Settings** > **Android SDK** properties in the Editor.
 [/REGION]
 
-1. Navigate to the TADP installation folder:
-| Install Type | Location |
-| ------------ | ---------|
-| Binary | [EngineInstallLocation]\Unreal Engine\[EngineVersion]\Engine\Extras\Android |
-| GitHub | [EngineSyncLocation\UE4\Engine\Extras\Android |
+[region:tip]
+It is highly recommended that you install Android Works while the editor and launcher are both closed to ensure that there are not problems with the install.
+[/region]
 
-1. **Double-click** on `tadp-2.0r8-windows.exe`.
+1. To install Android Works navigate to the **Engine** > **Extras** > **Android** folder. 
 
-1. Click **Next**.
-    
-	![](tadpinstaller_1.png)(w:450)
+	| Install Type | Location |
+	| ------------ | ---------|
+	| Binary | [EngineInstallLocation]\Unreal Engine\[EngineVersion]\Engine\Extras\Android |
+	| GitHub | [EngineSyncLocation\UE4\Engine\Extras\Android |
 
-1.	Select **I accept the agreement**, then click **Next**.
+	[region:note]
+	If you downloaded UE4 from **Github** you will need to perform the following steps before you can run the Android Works installer. 
+	Please note that the follow applies to both Windows and Mac based PCs that download the UE4 source code from Github.
+	If you are using the binary build downloaded from the launcher then you can skip this part and go right into installing Android Works. 
 
-	![](tadpinstaller_2.png)(w:450)
+	1. Once UE4 has been downloaded from Github and extracted, run the **Setup.bat** file to download the required files as well as the Android Works installer.
 
-1.	Click **Next**.
+	1. After  **Setup.bat** has finished running locate and run the Android Works installer following the instructions below.
 
-	![](tadpinstaller_3.png)(w:450)
+	1. When Android Works has finished installing, run the **Setup.bat** file on more time to download the required Android Development files.
+	[/region]
 
-1.	Select the **Express** install, then click **Next**.
-	
-	![](tadpinstaller_4.png)(w:450)
+1. **Double-click** on the **AndroidWorks-1R1-windows.exe** file to initiate the install process.
 
-1.	Make sure that the install location is `C:\NVPACK`.
-	
-	![](tadpinstaller_5.png)(w:450) 
+	![](AndroidWorks_11.png)
 
-1.	Verify that the list of components to install matches this list:
-	
-	![](tadpinstaller_6.png)(w:450)
+1. Click the **Next** button on the first window to begin the install process.
+    
+	![](AndroidWorks_00.png)
 
-1.	Enter a proxy if needed:
-	
-	![](tadpinstaller_7.png)(w:450)
+1.	Click the **Next** button to continue.
 
-1.	Click **Next**. You will need an internet connection and this step may take a few minutes.
-	
-	![](tadpinstaller_8.png)(w:450)
+	![](AndroidWorks_01.png)
 
-1.	Click **Next**.
-	
-	![](tadpinstaller_9.png)(w:450)
+1.	Set a location to download and install Android Works and then click the **Next** button when done.
 
-1.	The installer should run, and this will take some time.
-	
-	![](tadpinstaller_10.png)(w:450)
+	![](AndroidWorks_02.png)
 
-1.	If you are prompted with a warning about the software publisher, select **Install this driver software anyway**.
+	[region:note]
+	It is perfectly fine to leave this set to the default location.
+	[/region]
+
+1.	Make sure that your are installing **AndroidWorks1R1** and then click the **Next** button.
 	
-	![](tadpinstaller_11.png)
+	![](AndroidWorks_03.png)
 
-1.	When prompted about saving files after the install, you can select either **Yes** or **No**.
+1.	Click on the checkmark box, **Accept All** if you agree with the license terms and then press the **Accept** button.
 	
-	![](tadpinstaller_12.png)
+	![](AndroidWorks_04.png)
 
-1.	Click **Next**.
+1.	A new window will be displayed that will show you the download and install progress of the required Android Works tools. 
+	[region:note]
+	You will need an internet connection to complete this portion. Depending on the speed of that internet connection this could take a few minute to a few hours to complete.
+	[/region]
 	
-	![](tadpinstaller_13.png)(w:450)
+	![](AndroidWorks_06.png)
 
-1.	Verify that all the components have been installed correctly.
+1.	When all required tools have been downloaded and installed the following window will be displayed. Press the **OK** button continue the install process.
 	
-	![](tadpinstaller_14.png)(w:450)
-1.	Uncheck View Readme File, then click Finish.
+	![](AndroidWorks_08.png)
+
+1.	Now press the **Finish** button to complete the Android Works install.
 	
-	![](tadpinstaller_15.png)(w:450)
+	![](AndroidWorks_09.png)
 
+1.	Finally press the **Yes** button to re-start the computer.
+	
+	![](AndroidWorks_10.png)
+	[region:note]
+	If you have more programs to install you can press the **No** button but make sure to re-start your development PC before trying to use Android Work tools with UE4.
+	[/region]
 [/EXCERPT:WindowsSetup]
 
 [EXCERPT:MacOSSetup]
 
+[REGION:note]
+If you already have Android Works installed, but your environment variables have been lost, you can skip all of this and simply point the editor to the install locations using the **Edit Menu** > **Project Settings** > **Android SDK** properties in the Editor.
+[/REGION]
+
 [region:tip]
-It is best to do this step while the editor and launcher are both closed. If you are installing the SDK with them open, you will need to restart both the editor and the launcher before you can work with an Android device.
+It is highly recommended that you install Android Works while the editor and launcher are both closed to ensure that there are not problems with the install.
 [/region]
-
-TADP installer and Android SDK tools require that both Java 6 and Java 7 are installed on your Mac. 
+Before you can install Android Works to your Mac you will need to make sure that both Java 6 and Java 7 are installed. 
+Use the links provide below to download and install the correct versions of Java.
 
 * **Java 6** can be downloaded from [here](http://support.apple.com/kb/DL1572).
 * **Java 7** can be downloaded from [here](http://www.oracle.com/technetwork/java/javase/downloads/jre7-downloads-1880261.html).
 
-1. Navigate to the TADP installation folder:
+1. To install Android Works navigate to the **Engine** > **Extras** > **Android** folder. 
 | Install Type  | Location  |
 | --- | --- |
 | Binary | [EngineInstallLocation]\Unreal Engine[EngineVersion]\Engine\Extras\Android |
 | GitHub | [EngineSyncLocation\UE4\Engine\Extras\Android |
 
-1. **Double-click** on tadp-3.0r4-osx.dmg. This should mount TADP installer disk image and open its window.
+	[region:note]
+	If you downloaded UE4 from **Github** you will need to perform the following steps before you can run the Android Works installer. 
+	Please note that the follow applies to both Windows and Mac based PCs that download the UE4 source code from Github.
+	If you are using the binary build downloaded from the launcher then you can skip this part and go right into installing Android Works. 
 
-1. **Double-click** TADP installer icon.
+	1. Once UE4 has been downloaded from Github and extracted, run the **Setup.bat** file to download the required files as well as the Android Works installer.
 
-	![](Tadp_Mac_Setup-00.png)
+	1. After  **Setup.bat** has finished running locate and run the Android Works installer following the instructions below.
 
-1. Click **Next**.
+	1. When Android Works has finished installing, run the **Setup.bat** file on more time to download the required Android Development files.
+	[/region]
 
-	![](Tadp_Mac_Setup-01.png)
 
-1. Select the **I accept the agreement** option and then click **Next**.
+1. **Double-click** on the **AndroidWorks-1R1-osx.dmg** to load the Android Works installer.
+	
+	![](AndroidWorks_Mac_00.png)
 
-	![](Tadp_Mac_Setup-02.png)
+1. **Double-click** the Android Works installer icon to start the install process.
+
+	![](AndroidWorks_Mac_01.png)
 
 1. Click **Next**.
 
-	![](Tadp_Mac_Setup-03.png)
+	![](AndroidWorks_Mac_02.png)
+
+1. Click **Next**.
 
-1. Select the **Express** install, then click **Next**.
+	![](AndroidWorks_Mac_03.png)
 
-	![](Tadp_Mac_Setup-04.png)
+1. Specify a location where Android Works should be installed then press the **Next** button.
 
-1. Select the install location. 
 	[region:note]
 	It is perfectly fine to leave this set to the default location.
 	[/region]
 
-	![](Tadp_Mac_Setup-05.png)
+	![](AndroidWorks_Mac_04.png)
 
-1. Verify that the list of components to install matches what is in the image below.
+1. When the the Android Works Component Manager is displayed press the **Next** button to begin the download and install process.
 
-	![](Tadp_Mac_Setup-06.png)
+	![](AndroidWorks_Mac_05.png)
 
-1. Enter a proxy if needed.
+1. Click on the checkmark box, **Accept All** if you agree with the license terms and then press the **Accept** button.
 
-	![](Tadp_Mac_Setup-07.png)
+	![](AndroidWorks_Mac_06.png)
 
-1. Click **Next**. 
+1. The Android Works Component Manager will now begin to download and install the required tools. 
 	[region:note]
 	For this part of the setup will take a few minutes to completed and you will need an internet connection to complete it.
 	[/region]
+	![](AndroidWorks_Mac_07.png)
 
-	![](Tadp_Mac_Setup-08.png)
+1. Once everything has been installed the following window will be displayed. Click on the **OK** button to continue.
 
-1. Click **Next**.
-
-	![](Tadp_Mac_Setup-09.png)
+	![](AndroidWorks_Mac_08.png)
 
-1. The installer will now run and it will take some time for it complete it's install.
+1. Click on the **Finish** button to complete the Android Works install.
 
-	![](Tadp_Mac_Setup-10.png)
-
-1. Click **Next**.
-
-	![](Tadp_Mac_Setup-11.png)
-
-1. Click **Next**.
+	![](AndroidWorks_Mac_09.png)
 
-	![](Tadp_Mac_Setup-12.png)
-
-1. Verify that all the components have been installed correctly.
-
-	![](Tadp_Mac_Setup-13.png)
-
-1. Uncheck both **View Readme File** and **Browse Documentation**, then click **Finish**.
-
-	![](Tadp_Mac_Setup-14.png)
 
 [/EXCERPT:MacOSSetup]
-
 -->
+
 %Steps%
 
diff --git a/Engine/Documentation/Source/Platforms/Android/GettingStarted/2/AndroidQS_2.INT.udn b/Engine/Documentation/Source/Platforms/Android/GettingStarted/2/AndroidQS_2.INT.udn
--- a/Engine/Documentation/Source/Platforms/Android/GettingStarted/2/AndroidQS_2.INT.udn
+++ b/Engine/Documentation/Source/Platforms/Android/GettingStarted/2/AndroidQS_2.INT.udn
@@ -2,6 +2,8 @@ Availability:Public
 Title:2. Set Up Your Android Device
 Crumbs:%ROOT%, Platforms, Platforms/Android, Platforms/Android/GettingStarted
 Description:Set up your Android device for development.
+Version:4.9
+Skilllevel:Intermediate
 
 
 
diff --git a/Engine/Documentation/Source/Platforms/Android/GettingStarted/3/AndroidQS_3.INT.udn b/Engine/Documentation/Source/Platforms/Android/GettingStarted/3/AndroidQS_3.INT.udn
--- a/Engine/Documentation/Source/Platforms/Android/GettingStarted/3/AndroidQS_3.INT.udn
+++ b/Engine/Documentation/Source/Platforms/Android/GettingStarted/3/AndroidQS_3.INT.udn
@@ -2,6 +2,8 @@ Availability:Public
 Title:3. Create Your Project
 Crumbs:%ROOT%, Platforms, Platforms/Android, Platforms/Android/GettingStarted
 Description:Start a template project
+Version:4.9
+Skilllevel:Intermediate
 
 [VAR:Steps]
 [OBJECT:Navigation]
diff --git a/Engine/Documentation/Source/Platforms/Android/GettingStarted/4/AndroidQS_4.INT.udn b/Engine/Documentation/Source/Platforms/Android/GettingStarted/4/AndroidQS_4.INT.udn
--- a/Engine/Documentation/Source/Platforms/Android/GettingStarted/4/AndroidQS_4.INT.udn
+++ b/Engine/Documentation/Source/Platforms/Android/GettingStarted/4/AndroidQS_4.INT.udn
@@ -2,6 +2,8 @@ Availability:Public
 Title:4. Deploy Your Level
 Crumbs:%ROOT%, Platforms, Platforms/Android, Platforms/Android/GettingStarted
 Description:Launch your level with one-click deploying.
+Version:4.9
+Skilllevel:Intermediate
 
 [VAR:Steps]
 [OBJECT:Navigation]
diff --git a/Engine/Documentation/Source/Platforms/Android/GettingStarted/5/AndroidQS_5.INT.udn b/Engine/Documentation/Source/Platforms/Android/GettingStarted/5/AndroidQS_5.INT.udn
--- a/Engine/Documentation/Source/Platforms/Android/GettingStarted/5/AndroidQS_5.INT.udn
+++ b/Engine/Documentation/Source/Platforms/Android/GettingStarted/5/AndroidQS_5.INT.udn
@@ -2,6 +2,8 @@ Availability:Public
 Title:5. Package Your Game
 Crumbs:%ROOT%, Platforms, Platforms/Android, Platforms/Android/GettingStarted
 Description:Package your game for installation.
+Version:4.9
+Skilllevel:Intermediate
 
 [VAR:Steps]
 [OBJECT:Navigation]
diff --git a/Engine/Documentation/Source/Platforms/Android/GettingStarted/6/AndroidQS_6.INT.udn b/Engine/Documentation/Source/Platforms/Android/GettingStarted/6/AndroidQS_6.INT.udn
--- a/Engine/Documentation/Source/Platforms/Android/GettingStarted/6/AndroidQS_6.INT.udn
+++ b/Engine/Documentation/Source/Platforms/Android/GettingStarted/6/AndroidQS_6.INT.udn
@@ -2,6 +2,8 @@ Availability:Public
 Title:6. Additional Resources
 Crumbs:%ROOT%, Platforms, Platforms/Android, Platforms/Android/GettingStarted
 Description:Read more about Android development with Unreal Engine.
+Version:4.9
+Skilllevel:Intermediate
 
 [VAR:Steps]
 [OBJECT:Navigation]
diff --git a/Engine/Documentation/Source/Platforms/Android/GettingStarted/AndroidQuickStart.INT.udn b/Engine/Documentation/Source/Platforms/Android/GettingStarted/AndroidQuickStart.INT.udn
--- a/Engine/Documentation/Source/Platforms/Android/GettingStarted/AndroidQuickStart.INT.udn
+++ b/Engine/Documentation/Source/Platforms/Android/GettingStarted/AndroidQuickStart.INT.udn
@@ -2,6 +2,8 @@ Availability: Public
 Title: Android Quick Start
 Crumbs:%ROOT%, Platforms, Platforms/Android
 Description:Setting up for development for the Android platform.
+Version:4.9
+Skilllevel:Intermediate
 
 [VAR:Topic]
 [OBJECT:Topic]
@@ -38,7 +40,7 @@ Description:Setting up for development for the Android platform.
 [/VAR]
 
 The Android Quick Start will walk you through setting up your computer and your Android device for deployment, and putting a template game on your device using both the launching and packaging features.
-This guide uses TADP as the Android SDK, but if you would like to know how to set up another SDK for use with Unreal Engine, see the [](Platforms\Android\Reference) page.
+This guide uses **Android Works** as the Android SDK, but if you would like to know how to set up another SDK for use with Unreal Engine, see the [](Platforms\Android\Reference) page.
 
 
 [OBJECT:TopicButtonList]
diff --git a/Engine/Documentation/Source/Platforms/Android/LegacyAndroidDevelopment/LegacyAndroidDevelopment.INT.udn b/Engine/Documentation/Source/Platforms/Android/LegacyAndroidDevelopment/LegacyAndroidDevelopment.INT.udn
new file mode 100644
--- /dev/null
+++ b/Engine/Documentation/Source/Platforms/Android/LegacyAndroidDevelopment/LegacyAndroidDevelopment.INT.udn
@@ -0,0 +1,221 @@
+Availability: Public
+Title: Legacy Android Development
+Crumbs:%ROOT%, Platforms
+Description:Setting up TADP for Legacy Android Development
+Platform:Android
+Navigation:platform
+Version:4.8
+Skilllevel:Intermediate
+
+If you are using Unreal Engine version **4.8** and want to develop for Android you will need to install and use the **Tegra Android Development Pack** or **TADP** for short. 
+In the following document you will be shown how to install TADP on your development PC.
+
+[region:note]
+Android Works **will not** work with Unreal Engine 4.8. If you have installed **Android Works** you will need to uninstall it and re-start your development PC before installing TADP. 
+[/region]
+
+%Globals:OSSelect%
+
+[OBJECT:ToggleButtonContent]
+	[PARAMLITERAL:category]
+	OS
+	[/PARAMLITERAL]
+	[PARAMLITERAL:id]
+	windows
+	[/PARAMLITERAL]
+	[PARAMLITERAL:active]
+	active_button_content
+	[/PARAMLITERAL]
+	[PARAM:content]
+	[INCLUDE:#WindowsSetup]
+	[/PARAM]
+[/OBJECT]
+
+[OBJECT:ToggleButtonContent]
+	[PARAMLITERAL:category]
+	OS
+	[/PARAMLITERAL]
+	[PARAMLITERAL:id]
+	mac
+	[/PARAMLITERAL]
+	[PARAMLITERAL:active]
+	[/PARAMLITERAL]
+	[PARAM:content]
+	[INCLUDE:#MacOSSetup]
+	[/PARAM]
+[/OBJECT]
+ 
+<!--
+[EXCERPT:WindowsSetup]
+
+[REGION:tip]
+If you already have the Android SDK installed, but your environment variables have been lost, you can skip all of this and simply point the editor to the install locations using the **Edit Menu** -> **Project Settings** -> **Android SDK** properties in the Editor.
+[/REGION]
+
+1. Navigate to the TADP installation folder:
+| Install Type | Location |
+| ------------ | ---------|
+| Binary | [EngineInstallLocation]\Unreal Engine\[EngineVersion]\Engine\Extras\Android |
+| GitHub | [EngineSyncLocation\UE4\Engine\Extras\Android |
+
+1. **Double-click** on `tadp-2.0r8-windows.exe`.
+
+1. Click **Next**.
+    
+	![](Platforms\Android\GettingStarted\1\tadpinstaller_1.png)(w:450)
+
+1.	Select **I accept the agreement**, then click **Next**.
+
+	![](Platforms\Android\GettingStarted\1\tadpinstaller_2.png)(w:450)
+
+1.	Click **Next**.
+
+	![](Platforms\Android\GettingStarted\1\tadpinstaller_3.png)(w:450)
+
+1.	Select the **Express** install, then click **Next**.
+	
+	![](Platforms\Android\GettingStarted\1\tadpinstaller_4.png)(w:450)
+
+1.	Make sure that the install location is `C:\NVPACK`.
+	
+	![](Platforms\Android\GettingStarted\1\tadpinstaller_5.png)(w:450) 
+
+1.	Verify that the list of components to install matches this list:
+	
+	![](Platforms\Android\GettingStarted\1\tadpinstaller_6.png)(w:450)
+
+1.	Enter a proxy if needed:
+	
+	![](Platforms\Android\GettingStarted\1\tadpinstaller_7.png)(w:450)
+
+1.	Click **Next**. You will need an internet connection and this step may take a few minutes.
+	
+	![](Platforms\Android\GettingStarted\1\tadpinstaller_8.png)(w:450)
+
+1.	Click **Next**.
+	
+	![](Platforms\Android\GettingStarted\1\tadpinstaller_9.png)(w:450)
+
+1.	The installer should run, and this will take some time.
+	
+	![](Platforms\Android\GettingStarted\1\tadpinstaller_10.png)(w:450)
+
+1.	If you are prompted with a warning about the software publisher, select **Install this driver software anyway**.
+	
+	![](Platforms\Android\GettingStarted\1\tadpinstaller_11.png)
+
+1.	When prompted about saving files after the install, you can select either **Yes** or **No**.
+	
+	![](Platforms\Android\GettingStarted\1\tadpinstaller_12.png)
+
+1.	Click **Next**.
+	
+	![](Platforms\Android\GettingStarted\1\tadpinstaller_13.png)(w:450)
+
+1.	Verify that all the components have been installed correctly.
+	
+	![](Platforms\Android\GettingStarted\1\tadpinstaller_14.png)(w:450)
+
+1.	Uncheck View Readme File, then click Finish.
+	
+	![](Platforms\Android\GettingStarted\1\tadpinstaller_15.png)(w:450)
+
+[/EXCERPT:WindowsSetup]
+
+[EXCERPT:MacOSSetup]
+
+[region:tip]
+It is best to do this step while the editor and launcher are both closed. If you are installing the SDK with them open, you will need to restart both the editor and the launcher before you can work with an Android device.
+[/region]
+
+TADP installer and Android SDK tools require that both Java 6 and Java 7 are installed on your Mac. 
+
+* **Java 6** can be downloaded from [here](http://support.apple.com/kb/DL1572).
+* **Java 7** can be downloaded from [here](http://www.oracle.com/technetwork/java/javase/downloads/jre7-downloads-1880261.html).
+
+1. Navigate to the TADP installation folder:
+| Install Type  | Location  |
+| --- | --- |
+| Binary | [EngineInstallLocation]\Unreal Engine[EngineVersion]\Engine\Extras\Android |
+| GitHub | [EngineSyncLocation\UE4\Engine\Extras\Android |
+
+1. **Double-click** on tadp-3.0r4-osx.dmg. This should mount TADP installer disk image and open its window.
+
+1. **Double-click** TADP installer icon.
+
+	![](Platforms\Android\GettingStarted\1\Tadp_Mac_Setup-00.png)
+
+1. Click **Next**.
+
+	![](Platforms\Android\GettingStarted\1\Tadp_Mac_Setup-01.png)
+
+1. Select the **I accept the agreement** option and then click **Next**.
+
+	![](Platforms\Android\GettingStarted\1\Tadp_Mac_Setup-02.png)
+
+1. Click **Next**.
+
+	![](Platforms\Android\GettingStarted\1\Tadp_Mac_Setup-03.png)
+
+1. Select the **Express** install, then click **Next**.
+
+	![](Platforms\Android\GettingStarted\1\Tadp_Mac_Setup-04.png)
+
+1. Select the install location. 
+	[region:note]
+	It is perfectly fine to leave this set to the default location.
+	[/region]
+
+	![](Platforms\Android\GettingStarted\1\Tadp_Mac_Setup-05.png)
+
+1. Verify that the list of components to install matches what is in the image below.
+
+	![](Platforms\Android\GettingStarted\1\Tadp_Mac_Setup-06.png)
+
+1. Enter a proxy if needed.
+
+	![](Platforms\Android\GettingStarted\1\Tadp_Mac_Setup-07.png)
+
+1. Click **Next**. 
+	[region:note]
+	For this part of the setup will take a few minutes to completed and you will need an internet connection to complete it.
+	[/region]
+
+	![](Platforms\Android\GettingStarted\1\Tadp_Mac_Setup-08.png)
+
+1. Click **Next**.
+
+	![](Platforms\Android\GettingStarted\1\Tadp_Mac_Setup-09.png)
+
+1. The installer will now run and it will take some time for it complete it's install.
+
+	![](Platforms\Android\GettingStarted\1\Tadp_Mac_Setup-10.png)
+
+1. Click **Next**.
+
+	![](Platforms\Android\GettingStarted\1\Tadp_Mac_Setup-11.png)
+
+1. Click **Next**.
+
+	![](Platforms\Android\GettingStarted\1\Tadp_Mac_Setup-12.png)
+
+1. Verify that all the components have been installed correctly.
+
+	![](Platforms\Android\GettingStarted\1\Tadp_Mac_Setup-13.png)
+
+1. Uncheck both **View Readme File** and **Browse Documentation**, then click **Finish**.
+
+	![](Platforms\Android\GettingStarted\1\Tadp_Mac_Setup-14.png)
+
+[/EXCERPT:MacOSSetup]
+
+
+
+
+
+
+
+
+
+
+
