[UE-19981] Taking pull request with documentation fix.
#github [REDACTED_PROFILE_URL]/UnrealEngine/commit/40c6901db7163a04b48a42495f3928f693ce756d rwils

diff --git a/Engine/Documentation/Source/Engine/Actors/SkeletalMeshActors/SkeletalMeshActors.INT.udn b/Engine/Documentation/Source/Engine/Actors/SkeletalMeshActors/SkeletalMeshActors.INT.udn
--- a/Engine/Documentation/Source/Engine/Actors/SkeletalMeshActors/SkeletalMeshActors.INT.udn
+++ b/Engine/Documentation/Source/Engine/Actors/SkeletalMeshActors/SkeletalMeshActors.INT.udn
@@ -96,7 +96,7 @@ To create a Physics Asset for your mesh, all you have to do is simply **Right-cl
 ![](SkeletalMeshActor_Create_Physics_Asset.png)
 
 ### Assigning a Physics Asset
-Once your Physics Asset has been created, it is not time to assign it to your mesh. Physics Assets can be assigned to Skeletal Meshes by simply opening up the Skeletal Mesh and then assigning the Physics Asset to the **Physics Asset** slot that is under the **Physics** section. 
+Once your Physics Asset has been created, it is now time to assign it to your mesh. Physics Assets can be assigned to Skeletal Meshes by simply opening up the Skeletal Mesh and then assigning the Physics Asset to the **Physics Asset** slot that is under the **Physics** section. 
 
 ![](SkeletalMeshActor_Add_Physics_Assets.png)
 
