#UE4doc #docs:
Added engine version field
Changed font settings for multiple terms

diff --git a/Engine/Documentation/Source/Engine/Animation/NodeReference/AnimationNodes.INT.udn b/Engine/Documentation/Source/Engine/Animation/NodeReference/AnimationNodes.INT.udn
--- a/Engine/Documentation/Source/Engine/Animation/NodeReference/AnimationNodes.INT.udn
+++ b/Engine/Documentation/Source/Engine/Animation/NodeReference/AnimationNodes.INT.udn
@@ -2,9 +2,10 @@ Availability:Public
 Title: Animation Node Reference
 Crumbs: %ROOT%, Engine, Engine/Animation
 Description:Descriptions of the various animation nodes available for use in Animation Blueprints.
+version: 4.9
 
 [EXCERPT:Intro]
-AnimationBlueprints make use of various nodes within their graphs to perform operations on input 
+Animation Blueprints make use of various nodes within their graphs to perform operations on input 
 poses, such as blending, direct bone manipulation, and more. There are a few distinct types 
 of animation nodes provided in the engine including events, blend nodes, skeletal controllers, 
 space nodes, and conversion nodes.
@@ -41,10 +42,10 @@ space nodes, and conversion nodes.
 
 ## Pins and Properties
 
-Animation nodes can have pins - both input and output - which values can be passed to in the **EventGraph** 
-or **AnimGraph** and properties which can be edited in the **Details** panel. Pins can be either _data 
-pins_ identical to those found on nodes in regular _Blueprints_ which take in data from variables or _pure_ 
-functions, or they can be _pose pins_ which are similar to _exec pins_ in regular _Blueprints_ in that they 
+Animation nodes can have pins (both input and output) which values can be passed to in the **EventGraph** 
+or **AnimGraph** and properties which can be edited in the **Details** panel. Pins can be either data 
+pins identical to those found on nodes in regular Blueprints which take in data from variables or pure 
+functions, or they can be pose pins which are similar to exec pins in regular Blueprints in that they 
 determine the flow of execution.
 
 Data pins are actually properties like any other and are exposed in the **Details** panel. These can be toggled 
