#UE4doc #docs:
Added engine version field
Changed font settings for multiple terms and fixed some capitalization.

diff --git a/Engine/Documentation/Source/Engine/Animation/NodeReference/Blend/BlendNodes.INT.udn b/Engine/Documentation/Source/Engine/Animation/NodeReference/Blend/BlendNodes.INT.udn
--- a/Engine/Documentation/Source/Engine/Animation/NodeReference/Blend/BlendNodes.INT.udn
+++ b/Engine/Documentation/Source/Engine/Animation/NodeReference/Blend/BlendNodes.INT.udn
@@ -2,6 +2,7 @@ Availability:Public
 Title:Blend Nodes
 Crumbs: %ROOT%, Engine, Engine/Animation, Engine/Animation/NodeReference
 Description:Animation nodes that blend multiple animations together based on a set of criteria.
+version: 4.9
 
 [REGION:banner]
 ![Blend Nodes](blend_banner.png)(convert:false)
@@ -10,7 +11,7 @@ Description:Animation nodes that blend multiple animations together based on a s
 [TOC(start:2)]
 
 [EXCERPT:Intro]
-Blend nodes are used to blend multiple animations together. These nodes are only available to be placed in the **Anim Graph** of a _AnimationBlueprint_. Each blend node takes in multiple poses and an alpha or weight value that is used to calculate the weighting for each pose in the final output. Some blend nodes may be more complex than others and require additional data to be input as well.
+**Blend Nodes** are used to blend multiple animations together. These nodes are only available to be placed in the **Anim Graph** of a Animation Blueprint. Each blend node takes in multiple poses and an alpha or weight value that is used to calculate the weighting for each pose in the final output. Some blend nodes may be more complex than others and require additional data to be input as well.
 [/EXCERPT:Intro]
 
 ## Apply Additive
@@ -22,11 +23,11 @@ The **Apply Additive** node adds an additive pose to a base normal animation pos
 | Pin | Description |
 | --- | ----------- |
 | [REGION:tablesection]**Input Pins**[/REGION] ||
-| Base | The normal pose to use as the base. |
-| Additive | The additive pose to be blended on top. |
-| Alpha | A `float` value in the range **[0.0, 1.0]** to use as the alpha value to determine how much influence the additive pose imparts. A value of **0.0** means the **Additive** pose is not added to the **Base** input pose at all, while a value of **1.0** means the **Additive** pose is added fully to the **Base** input pose. |
+| **Base** | The normal pose to use as the base. |
+| **Additive** | The additive pose to be blended on top. |
+| **Alpha** | A float value in the range [0.0, 1.0] to use as the alpha value to determine how much influence the additive pose imparts. A value of 0.0 means the **Additive** pose is not added to the **Base** input pose at all, while a value of **1.0** means the **Additive** pose is added fully to the **Base** input pose. |
 | [REGION:tablesection]**Output Pins**[/REGION] ||
-| Pose | The final pose after the blending has been applied. |
+| **Pose** | The final pose after the blending has been applied. |
 
 ## Blend
 
@@ -37,69 +38,69 @@ The standard **Blend** node is a straight mixing of the two input poses based on
 | Pin | Description |
 | --- | ----------- |
 | [REGION:tablesection]**Input Pins**[/REGION] ||
-| A | The first pose to be blended. |
-| B | The second pose to be blended. |
-| Alpha | A `float` value in the range **[0.0, 1.0]** to use as the alpha value to determine the weighting of the two poses. A value of **0.0** gives full weighting to the **A** input pose, while a value of **1.0** gives full weighting to the **B** input pose. |
+| **A** | The first pose to be blended. |
+| **B** | The second pose to be blended. |
+| **Alpha** | A float value in the range [0.0, 1.0] to use as the alpha value to determine the weighting of the two poses. A value of 0.0 gives full weighting to the A input pose, while a value of 1.0 gives full weighting to the B input pose. |
 | [REGION:tablesection]**Output Pins**[/REGION] ||
-| Pose | The final pose after the blending has been applied. |
+| **Pose** | The final pose after the blending has been applied. |
 
 
-## Blend Poses by Bool
+## Blend Poses by bool
 
-The **Blend Poses by Bool** node performs a time-based blend between two poses using a Boolean value as the key. When the Boolean reads _true_, the pose connected to the _true_ input pin is used; when the Boolean is _false_, the _false_ pose is used. Each pose has a float value Blend Time that is used to control how long it will take to blend into the pose.
+The **Blend Poses by bool** node performs a time-based blend between two poses using a Boolean value as the key. When the Boolean reads true, the pose connected to the true input pin is used; when the Boolean is false, the false pose is used. Each pose has a float value Blend Time that is used to control how long it will take to blend into the pose.
 
 ![](BlendPosesByBool.png)
 
 | Pin | Description |
 | --- | ----------- |
 | [REGION:tablesection]**Input Pins**[/REGION] ||
-| Active Value | This is the input Boolean that decides which of the two poses to use. |
-| True Pose | Takes in a pose that will blend when the Boolean is set to _true_. |
-| False Pose | Takes in a pose that will blend when the Boolean is set to _false_. |
-| True Blend Time | Defines the amount of time needed to blend into the _true_ pose. |
-| False Blend Time | Defines the amount of time needed to blend into the _false_ pose. |
+| **Active Value** | This is the input Boolean that decides which of the two poses to use. |
+| **True Pose** | Takes in a pose that will blend when the Boolean is set to true. |
+| **False Pose** | Takes in a pose that will blend when the Boolean is set to false. |
+| **True Blend Time** | Defines the amount of time needed to blend into the true pose. |
+| **False Blend Time** | Defines the amount of time needed to blend into the false pose. |
 | [REGION:tablesection]**Output Pins**[/REGION] ||
-| Pose | The final pose after the blending has been applied. |
+| **Pose** | The final pose after the blending has been applied. |
 
-## Blend Poses by Int
+## Blend Poses by int
 
-The **Blend Poses by Int** node performs a time-based blend between any number of poses using an integer value as the key. For each input integer value, the pose associated with that value's input pin is used. For instance, when the integer is set to 0, the pose connected to _Blend Pose 0_ is used. Each pose has a float value Blend Time that is used to control how long it will take to blend into the pose.
+The **Blend Poses by int** node performs a time-based blend between any number of poses using an integer value as the key. For each input integer value, the pose associated with that value's input pin is used. For instance, when the integer is set to 0, the pose connected to Blend Pose 0 is used. Each pose has a float value Blend Time that is used to control how long it will take to blend into the pose.
 
 ![](BlendPosesByInt.png)
 
-To create additional Blend Pose inputs, simply **right-click** on the node and choose **Add Blend Pin** from the context menu.
+To create additional Blend Pose inputs, simply **Right-click** on the node and choose **Add Blend Pin** from the context menu.
 
 ![](AddBlendInput.png)
 
-Inputs can be removed by **right-clicking** on the desired input and choosing **Remove Blend Pin**.
+Inputs can be removed by **Right-clicking** on the desired input and choosing **Remove Blend Pin**.
 
 
 | Pin | Description |
 | --- | ----------- |
 | [REGION:tablesection]**Input Pins**[/REGION] ||
-| Active Child Index | This is the input int that decides which of the poses to use. |
-| Blend Pose X | Takes in a pose that will blend when the Active Child Index is set to X. |
-| Blend Time X | Defines the amount of time needed to blend into the X pose. |
+| **Active Child Index** | This is the input int that decides which of the poses to use. |
+| **Blend Pose X** | Takes in a pose that will blend when the Active Child Index is set to X. |
+| **Blend Time X** | Defines the amount of time needed to blend into the X pose. |
 | [REGION:tablesection]**Output Pins**[/REGION] ||
-| Pose | The final pose after the blending has been applied. |
+| **Pose** | The final pose after the blending has been applied. |
 
 
 
-## Layered Blend Per Bone
+## Layered blend per bone
 
-The **Layered Blend Per Bone** node performs a masked blend that only affects a specific subset of the bones in the Skeleton. For instance, if you wanted to apply an animation to a character only from the waist up, this could be done using **Blend Per Bone**.
+The **Layered blend per bone** node performs a masked blend that only affects a specific subset of the bones in the Skeleton. For instance, if you wanted to apply an animation to a character only from the waist up, this could be done using **Blend Per Bone**.
 
 ![Blend Nodes - Blend Per Bone](CE_nodes_blend_per_bone.png)
 
 | Pin | Description |
 | --- | ----------- |
 | [REGION:tablesection]**Input Pins**[/REGION] ||
-| BasePose | The first pose to be blended. |
-| BlendPose | The second pose to be blended. |
-| BlendWeight | A `float` value in the range **[0.0, 1.0]** to use as the alpha value to determine how much influence the additive pose imparts. A value of **0.0** means the **Additive** pose is not added to the **Base** input pose at all, while a value of **1.0** means the **Additive** pose is added fully to the **Base** input pose. |
-| Add Pin | Creates another Blend Pose and Blend Weight input. |
+| **BasePose** | The first pose to be blended. |
+| **BlendPose** | The second pose to be blended. |
+| **BlendWeight** | A float value in the range [0.0, 1.0] to use as the alpha value to determine how much influence the additive pose imparts. A value of 0.0 means the **Additive** pose is not added to the **Base** input pose at all, while a value of 1.0 means the **Additive** pose is added fully to the **Base** input pose. |
+| **Add Pin** | Creates another Blend Pose and Blend Weight input. |
 | [REGION:tablesection]**Output Pins**[/REGION] ||
-| OutPose | The final pose after the blending has been applied. |
+| **OutPose** | The final pose after the blending has been applied. |
 
 
 
