#UE4doc #docs:
Added engine version field
Changed font settings for multiple terms

diff --git a/Engine/Documentation/Source/Engine/Animation/NodeReference/Events/AnimationEvents.INT.udn b/Engine/Documentation/Source/Engine/Animation/NodeReference/Events/AnimationEvents.INT.udn
--- a/Engine/Documentation/Source/Engine/Animation/NodeReference/Events/AnimationEvents.INT.udn
+++ b/Engine/Documentation/Source/Engine/Animation/NodeReference/Events/AnimationEvents.INT.udn
@@ -2,6 +2,7 @@ Availability:Public
 Title:Animation Events
 Crumbs: %ROOT%, Engine, Engine/Animation, Engine/Animation/NodeReference
 Description:Events called by the animation system to initialize and update the Animation Blueprint.
+version: 4.9
 
 [REGION:banner]
 ![Animation Events](animationevent_banner.png)(convert:false)
@@ -13,7 +14,7 @@ Description:Events called by the animation system to initialize and update the A
 Various events are called by the animation system to initialize and update the Animation Blueprint. These can be placed in the **Event Graph** to implement or override functionality.
 [/EXCERPT:Intro]
 
-## BlueprintInitializeAnimation
+## Blueprint Initialize Animation
 
 [EXCERPT:Initialize]
 The **Blueprint Initialize Animation** event is executed once when the Animation Blueprint instance is created to perform initialization operations.
@@ -27,7 +28,7 @@ The **Blueprint Initialize Animation** event is executed once when the Animation
 
 [/EXCERPT:Initialize]
 
-## BlueprintUpdateAnimation
+## Blueprint Update Animation
 
 [EXCERPT:Update]
 The **Blueprint Update Animation** event is executed every frame allowing the Animation Blueprint to perform calculations and updates to any values it needs. This event is the entry point into the update loop of the EventGraph. The amount of time elapsed since the last update is available from the **DeltaTimeX** pin so time-dependent interpolations or incremental updates can be performed.
@@ -38,6 +39,6 @@ The **Blueprint Update Animation** event is executed every frame allowing the An
 | --- | ----------- |
 | [REGION:tablesection]**Output Pins**[/REGION] ||
 | Exec (Blank) | Fired once each frame when the **BlueprintUpdateAnimation()** event is executed in code. |
-| Delta Time X | Outputs the amount of time since the last execution cycle. |
+| **Delta Time X** | Outputs the amount of time since the last execution cycle. |
 
 [/EXCERPT:Update]
