#UE4doc #docs:
Added engine version field
Changed font settings for multiple terms
Replaced image due to updated UI

diff --git a/Engine/Documentation/Source/Engine/Animation/NodeReference/Fabrik/AnimNode_Fabrik.INT.udn b/Engine/Documentation/Source/Engine/Animation/NodeReference/Fabrik/AnimNode_Fabrik.INT.udn
--- a/Engine/Documentation/Source/Engine/Animation/NodeReference/Fabrik/AnimNode_Fabrik.INT.udn
+++ b/Engine/Documentation/Source/Engine/Animation/NodeReference/Fabrik/AnimNode_Fabrik.INT.udn
@@ -2,6 +2,7 @@ Availability:Public
 Title:FABRIK
 Crumbs: %ROOT%, Engine, Engine/Animation, Engine/Animation/NodeReference
 Description:Animation node for Forward And Backward Reaching Inverse Kinematics.
+version: 4.9
 
 [REGION:banner]
 ![FABRIK](FABRIK.png)(convert:false)
@@ -11,9 +12,9 @@ Description:Animation node for Forward And Backward Reaching Inverse Kinematics.
 **Forward And Backward Reaching Inverse Kinematics** (or **FABRIK**) is an IK solver that works on a chain of bones of arbitrary length (a minimum of 2 links). 
 [/EXCERPT:Intro]
 
-You can use **FABRIK**, by first adding the **FABRIK** node to the **AnimGraph** of your **Animation Blueprint** (pictured above). 
+You can use FABRIK, by first adding the **FABRIK** node to the **AnimGraph** of your Animation Blueprint (pictured above). 
 
-Once added, in addition to the **Component Pose**, the node takes in an **Effector Transform** which can be an absolute **Transform** or a relative one (based on another bone from the same Skeleton); the settings for the **End Effector** are also very similar to the **TwoBone_IK** node. 
+Once added, in addition to the **Component Pose**, the node takes in an **Effector Transform** which can be an absolute Transform or a relative one (based on another bone from the same Skeleton); the settings for the **End Effector** are also very similar to the **TwoBone_IK** node. 
 
 In order to determine the current strength of the Skeletal Control applied, the **FABRIK** node takes in a float **Alpha** which can be set on the node itself or set through a float variable. 
 
@@ -27,10 +28,10 @@ In the **End Effector** section, you determine the target location and rotation.
 | --- | ----------- |
 | **Location, Rotation, Scale**| Coordinates for target location of tip bone - if EffectorLocationSpace is bone, this is the offset from Target Bone to use as target location (can also be set as a Pin on the node). |
 | **Effector Transform Space**| Set the position of the bone in SkeletalMeshComponent's reference frame. |
-| **Effector Transform Bone**| If EffectorTransformSpace is bone, this is the bone to use. |
+| **Effector Transform Bone**| If Effector Transform Space is bone, this is the bone to use. |
 | **Effector Rotation Source**| Control the Rotation (maintain component space, local space or match end effector target rotation). |
 
-In the **Solver** section, you define the chain of bones to use, from the **Root** to the **Tip**. The **Tip** will try to reach the end effector location.
+In the **Solver** section, you define the chain of bones to use, from the **Root** to the **Tip**. The Tip will try to reach the end effector location.
 
 | Property | Description |
 | --- | ----------- |
