#UE4doc #docs
Updated page links.

diff --git a/Engine/Documentation/Source/Programming/Tutorials/FirstPersonShooter/1/FirstPersonShooter_1.INT.udn b/Engine/Documentation/Source/Programming/Tutorials/FirstPersonShooter/1/FirstPersonShooter_1.INT.udn
--- a/Engine/Documentation/Source/Programming/Tutorials/FirstPersonShooter/1/FirstPersonShooter_1.INT.udn
+++ b/Engine/Documentation/Source/Programming/Tutorials/FirstPersonShooter/1/FirstPersonShooter_1.INT.udn
@@ -1,7 +1,8 @@
 Availability: Docs
 Title: 1 - Setting up your Project
-Crumbs: %ROOT%, GettingStarted, Programming/Tutorials
+Crumbs: %ROOT%, Programming, Programming/Tutorials, Programming/Tutorials/FirstPersonShooter
 Description: How to set up your First Person Shooter project.
+prereq: Programming/Tutorials/FirstPersonShooter
 SkillLevel: Beginner
 Version: 4.9
 
diff --git a/Engine/Documentation/Source/Programming/Tutorials/FirstPersonShooter/2/1/FirstPersonShooter_2_1.INT.udn b/Engine/Documentation/Source/Programming/Tutorials/FirstPersonShooter/2/1/FirstPersonShooter_2_1.INT.udn
--- a/Engine/Documentation/Source/Programming/Tutorials/FirstPersonShooter/2/1/FirstPersonShooter_2_1.INT.udn
+++ b/Engine/Documentation/Source/Programming/Tutorials/FirstPersonShooter/2/1/FirstPersonShooter_2_1.INT.udn
@@ -1,6 +1,6 @@
 Availability: Docs
 Title: 2.1 - Making a New Character
-Crumbs: %ROOT%, GettingStarted, Engine/Content/QuickStart
+Crumbs: %ROOT%, Programming, Programming/Tutorials/FirstPersonShooter, Programming/Tutorials/FirstPersonShooter/2
 Description: Learn how to make a new character for your First Person Shooter project.
 SkillLevel: Beginner
 Version: 4.9
diff --git a/Engine/Documentation/Source/Programming/Tutorials/FirstPersonShooter/2/8/FirstPersonShooter_2_8.INT.udn b/Engine/Documentation/Source/Programming/Tutorials/FirstPersonShooter/2/8/FirstPersonShooter_2_8.INT.udn
--- a/Engine/Documentation/Source/Programming/Tutorials/FirstPersonShooter/2/8/FirstPersonShooter_2_8.INT.udn
+++ b/Engine/Documentation/Source/Programming/Tutorials/FirstPersonShooter/2/8/FirstPersonShooter_2_8.INT.udn
@@ -1,6 +1,6 @@
 Availability: Docs
 Title: 2.8 - Add a First Person Mesh to your Character
-Crumbs: %ROOT%, GettingStarted, Engine/Content/QuickStart
+Crumbs: %ROOT%, Programming, Programming/Tutorials/FirstPersonShooter, Programming/Tutorials/FirstPersonShooter/2
 Description: Learn how to add a first person mesh to your First Person Shooter character.
 SkillLevel: Beginner
 Version: 4.9
diff --git a/Engine/Documentation/Source/Programming/Tutorials/FirstPersonShooter/2/FirstPersonShooter_2.INT.udn b/Engine/Documentation/Source/Programming/Tutorials/FirstPersonShooter/2/FirstPersonShooter_2.INT.udn
--- a/Engine/Documentation/Source/Programming/Tutorials/FirstPersonShooter/2/FirstPersonShooter_2.INT.udn
+++ b/Engine/Documentation/Source/Programming/Tutorials/FirstPersonShooter/2/FirstPersonShooter_2.INT.udn
@@ -1,7 +1,9 @@
 Availability: Docs
 Title: 2 - Implementing your Character
-Crumbs: %ROOT%, GettingStarted, Programming/Tutorials
+Crumbs: %ROOT%, Programming, Programming/Tutorials, Programming/Tutorials/FirstPersonShooter
 Description: How to implement your First Person Shooter character.
+prereq: Programming/Tutorials/FirstPersonShooter/1
+Related: Programming/Tutorials/FirstPersonShooter
 SkillLevel: Beginner
 Version: 4.9
 
diff --git a/Engine/Documentation/Source/Programming/Tutorials/FirstPersonShooter/3/1/FirstPersonShooter_3_1.INT.udn b/Engine/Documentation/Source/Programming/Tutorials/FirstPersonShooter/3/1/FirstPersonShooter_3_1.INT.udn
--- a/Engine/Documentation/Source/Programming/Tutorials/FirstPersonShooter/3/1/FirstPersonShooter_3_1.INT.udn
+++ b/Engine/Documentation/Source/Programming/Tutorials/FirstPersonShooter/3/1/FirstPersonShooter_3_1.INT.udn
@@ -1,6 +1,6 @@
 Availability: Docs
 Title: 3.1 - Adding Projectiles to your Game
-Crumbs: %ROOT%, GettingStarted, Engine/Content/QuickStart
+Crumbs: %ROOT%, Programming, Programming/Tutorials/FirstPersonShooter, Programming/Tutorials/FirstPersonShooter/3
 Description: Learn how to add projectiles to your First Person Shooter project.
 SkillLevel: Beginner
 Version: 4.9
diff --git a/Engine/Documentation/Source/Programming/Tutorials/FirstPersonShooter/3/4/FirstPersonShooter_3_4.INT.udn b/Engine/Documentation/Source/Programming/Tutorials/FirstPersonShooter/3/4/FirstPersonShooter_3_4.INT.udn
--- a/Engine/Documentation/Source/Programming/Tutorials/FirstPersonShooter/3/4/FirstPersonShooter_3_4.INT.udn
+++ b/Engine/Documentation/Source/Programming/Tutorials/FirstPersonShooter/3/4/FirstPersonShooter_3_4.INT.udn
@@ -1,6 +1,6 @@
 Availability: Docs
-Title: 3.4 - Getting your Projectiles to Interact with the World
-Crumbs: %ROOT%, GettingStarted, Engine/Content/QuickStart
+Title: 3.4 - Get Projectiles to Interact with the World
+Crumbs: %ROOT%, Programming, Programming/Tutorials/FirstPersonShooter, Programming/Tutorials/FirstPersonShooter/3
 Description: Learn how to get projectiles to interact with the world for your First Person Shooter project.
 SkillLevel: Beginner
 Version: 4.9
diff --git a/Engine/Documentation/Source/Programming/Tutorials/FirstPersonShooter/3/FirstPersonShooter_3.INT.udn b/Engine/Documentation/Source/Programming/Tutorials/FirstPersonShooter/3/FirstPersonShooter_3.INT.udn
--- a/Engine/Documentation/Source/Programming/Tutorials/FirstPersonShooter/3/FirstPersonShooter_3.INT.udn
+++ b/Engine/Documentation/Source/Programming/Tutorials/FirstPersonShooter/3/FirstPersonShooter_3.INT.udn
@@ -1,7 +1,9 @@
 Availability: Docs
 Title: 3 - Implementing Projectiles
-Crumbs: %ROOT%, GettingStarted, Programming/Tutorials
+Crumbs: %ROOT%, Programming, Programming/Tutorials, Programming/Tutorials/FirstPersonShooter
 Description: How to implement projectiles for your First Person Shooter game.
+prereq: Programming/Tutorials/FirstPersonShooter/2
+Related: Programming/Tutorials/FirstPersonShooter
 SkillLevel: Beginner
 Version: 4.9
 
diff --git a/Engine/Documentation/Source/Programming/Tutorials/FirstPersonShooter/4/1/FirstPersonShooter_4_1.INT.udn b/Engine/Documentation/Source/Programming/Tutorials/FirstPersonShooter/4/1/FirstPersonShooter_4_1.INT.udn
--- a/Engine/Documentation/Source/Programming/Tutorials/FirstPersonShooter/4/1/FirstPersonShooter_4_1.INT.udn
+++ b/Engine/Documentation/Source/Programming/Tutorials/FirstPersonShooter/4/1/FirstPersonShooter_4_1.INT.udn
@@ -1,6 +1,6 @@
 Availability: Docs
 Title: 4.1 - Animating your Character
-Crumbs: %ROOT%, GettingStarted, Engine/Content/QuickStart
+Crumbs: %ROOT%, Programming, Programming/Tutorials, Programming/Tutorials/FirstPersonShooter, Programming/Tutorials/FirstPersonShooter/4
 Description: Learn how to animate your First Person Shooter character.
 SkillLevel: Beginner
 Version: 4.9
diff --git a/Engine/Documentation/Source/Programming/Tutorials/FirstPersonShooter/4/4/FirstPersonShooter_4_4.INT.udn b/Engine/Documentation/Source/Programming/Tutorials/FirstPersonShooter/4/4/FirstPersonShooter_4_4.INT.udn
--- a/Engine/Documentation/Source/Programming/Tutorials/FirstPersonShooter/4/4/FirstPersonShooter_4_4.INT.udn
+++ b/Engine/Documentation/Source/Programming/Tutorials/FirstPersonShooter/4/4/FirstPersonShooter_4_4.INT.udn
@@ -1,6 +1,6 @@
 Availability: Docs
 Title: 4.4 - Adding Animation Transition States
-Crumbs: %ROOT%, GettingStarted, Programming/Tutorials
+Crumbs: %ROOT%, Programming, Programming/Tutorials, Programming/Tutorials/FirstPersonShooter, Programming/Tutorials/FirstPersonShooter/4
 Description: How to add animation transition states to your First Person Shooter project.
 SkillLevel: Beginner
 Version: 4.9
@@ -16,7 +16,7 @@ Version: 4.9
 		[Section Home](Programming/Tutorials/FirstPersonShooter/4)
 	[/PARAM]
 	[PARAM:next]
-		[Next Step](Programming/Tutorials/FirstPersonShooter/4/5)
+		[Next Step](Programming/Tutorials/FirstPersonShooter/4/4/1)
 	[/PARAM]
 [/OBJECT]
 [/VAR]
diff --git a/Engine/Documentation/Source/Programming/Tutorials/FirstPersonShooter/4/5/FirstPersonShooter_4_5.INT.udn b/Engine/Documentation/Source/Programming/Tutorials/FirstPersonShooter/4/5/FirstPersonShooter_4_5.INT.udn
--- a/Engine/Documentation/Source/Programming/Tutorials/FirstPersonShooter/4/5/FirstPersonShooter_4_5.INT.udn
+++ b/Engine/Documentation/Source/Programming/Tutorials/FirstPersonShooter/4/5/FirstPersonShooter_4_5.INT.udn
@@ -1,6 +1,6 @@
 Availability: Docs
 Title: 4.5 - Associating Animation & Character Blueprints
-Crumbs: %ROOT%, GettingStarted, Engine/Content/QuickStart
+Crumbs: %ROOT%, Programming, Programming/Tutorials, Programming/Tutorials/FirstPersonShooter, Programming/Tutorials/FirstPersonShooter/4
 Description: Learn how to associate the Animation Blueprint with the Character Blueprint for your First Person Shooter project.
 SkillLevel: Beginner
 Version: 4.9
@@ -8,7 +8,7 @@ Version: 4.9
 [VAR:Steps]
 [OBJECT:Navigation]
 	[PARAM:previous]
-		[Previous Step](Programming/Tutorials/FirstPersonShooter/4/4)
+		[Previous Step](Programming/Tutorials/FirstPersonShooter/4/4/6)
 	[/PARAM]
 	[PARAM:current]
 	[/PARAM]
diff --git a/Engine/Documentation/Source/Programming/Tutorials/FirstPersonShooter/4/FirstPersonShooter_4.INT.udn b/Engine/Documentation/Source/Programming/Tutorials/FirstPersonShooter/4/FirstPersonShooter_4.INT.udn
--- a/Engine/Documentation/Source/Programming/Tutorials/FirstPersonShooter/4/FirstPersonShooter_4.INT.udn
+++ b/Engine/Documentation/Source/Programming/Tutorials/FirstPersonShooter/4/FirstPersonShooter_4.INT.udn
@@ -1,7 +1,9 @@
 Availability: Docs
 Title: 4 - Adding Character Animation
-Crumbs: %ROOT%, GettingStarted, Programming/Tutorials
+Crumbs: %ROOT%, Programming, Programming/Tutorials, Programming/Tutorials/FirstPersonShooter
 Description: How to add animations to your First Person Shooter character.
+prereq: Programming/Tutorials/FirstPersonShooter/3
+Related: Programming/Tutorials/FirstPersonShooter
 SkillLevel: Beginner
 Version: 4.9
 
