#UE4doc #docs:
Added engine version field
Changed font settings for many terms

diff --git a/Engine/Documentation/Source/Engine/Animation/NodeReference/SkeletalControls/SkeletalControls.INT.udn b/Engine/Documentation/Source/Engine/Animation/NodeReference/SkeletalControls/SkeletalControls.INT.udn
--- a/Engine/Documentation/Source/Engine/Animation/NodeReference/SkeletalControls/SkeletalControls.INT.udn
+++ b/Engine/Documentation/Source/Engine/Animation/NodeReference/SkeletalControls/SkeletalControls.INT.udn
@@ -2,6 +2,7 @@ Availability:Public
 Title:Skeletal Controls
 Crumbs: %ROOT%, Engine, Engine/Animation, Engine/Animation/NodeReference
 Description:Animation nodes that allow for direct manipulation and applying solvers to the bones of the target Skeleton.
+version: 4.9
 
 [REGION:banner]
 ![Skeletal Control](skelcontrol_banner.png)(convert:false)
@@ -24,65 +25,65 @@ See [](Engine/Animation/NodeReference/SpaceConversion) for more information on t
 
 ## Common Pins and Properties
 
-Some pins and properties are common to all _SkeletalControls_. These are described below.
+Some pins and properties are common to all SkeletalControls. These are described below.
 
 | Pin | Description |
 | --- | ----------- |
 | [REGION:tablesection]**Input Pins**[/REGION] ||
-| Component Pose | The input pose to be transformed. |
-| Alpha | A `float` value in the range **[0.0, 1.0]** to use as the alpha value to determine the weighting of the `Transform` applied by the _SkeletalControl_. A value of **0.0** gives full weighting to the input pose, while a value of **1.0** gives full weighting to the control's computed `Transform`. |
+| **Component Pose** | The input pose to be transformed. |
+| **Alpha** | A float value in the range [0.0, 1.0] to use as the alpha value to determine the weighting of the Transform applied by the SkeletalControl. A value of 0.0 gives full weighting to the input pose, while a value of 1.0 gives full weighting to the control's computed Transform. |
 | [REGION:tablesection]**Output Pins**[/REGION] ||
-| Pose (Blank) | The final pose after the `Transform` has been applied. |
+| Pose (Blank) | The final pose after the Transform has been applied. |
 
 ## Apply a Percentage of Rotation
 
-The **Apply a Percentage of Rotation** control drives the `Rotation` of a target bone at some specified percentage 
-of the `Rotation` of another bone within the _Skeleton_.
+The **Apply a Percentage of Rotation** control drives the Rotation of a target bone at some specified percentage 
+of the Rotation of another bone within the Skeleton.
 
 ![Apply Percentage of Rotation](node_applypercent.png)
 
 | Property | Description |
 | -------- | ----------- |
 | [REGION:tablesection]**Multiplier**[/REGION] ||
-| Target Bone | Name of target bone to control. This is the main bone chain to modify from. |
-| Source Bone | Name of the source bone to get `Transform` from. |
-| Rotation Axis To Refer | The axis (`X`, `Y`, or `Z`) to control the `Rotation` for. |
+| **Target Bone** | Name of target bone to control. This is the main bone chain to modify from. |
+| **Source Bone** | Name of the source bone to get Transform from. |
+| **Rotation Axis To Refer** | The axis (X, Y, or Z) to control the Rotation for. |
 
 | Pin | Description |
 | --- | ----------- |
 | [REGION:tablesection]**Input Pins**[/REGION] ||
-| Multiplier | A `float` value that represents the percentage of the source bone's `Rotation` to apply to the target bone. [INCLUDE:#Multiplier] |
+| **Multiplier** | A float value that represents the percentage of the source bone's Rotation to apply to the target bone. [INCLUDE:#Multiplier] |
 
 <!--
 [EXCERPT:Multiplier]
 [REGION:simpletable]
 | Value | Effect |
 | ----- | ------ |
-| **0.0** | None of the source bone's `Rotation` is applied to the target bone. |
-| **1.0** | All of the source bone's `Rotation` is applied to the target bone. |
+| 0.0 | None of the source bone's Rotation is applied to the target bone. |
+| 1.0 | All of the source bone's Rotation is applied to the target bone. |
 [/REGION] 
 [/EXCERPT:Multiplier]
 -->
 
 ## Copy Bone
 
-The **Copy Bone** control copies the `Transform` data or any component of it - i.e. `Translation`, `Rotation`, or 
-`Scale` - from one bone to another.
+The **Copy Bone** control copies the Transform data or any component of it - i.e. Translation, Rotation, or 
+Scale - from one bone to another.
 
 ![Apply Percentage of Rotation](node_copybone.png)
 
 | Property | Description |
 | -------- | ----------- |
 | [REGION:tablesection]**Copy**[/REGION] ||
-| Target Bone | Name of the target bone to copy `Transform` to. |
-| Source Bone | Name of the source bone to get `Transform` from. |
+| **Target Bone** | Name of the target bone to copy Transform to. |
+| **Source Bone** | Name of the source bone to get Transform from. |
 
 | Pin | Description |
 | --- | ----------- |
 | [REGION:tablesection]**Input Pins**[/REGION] ||
-| Copy Translation | A `bool` specifying whether to copy the `Translation` component of the source bone's `Transform`. |
-| Copy Rotation | A `bool` specifying whether to copy the `Rotation` component of the source bone's `Transform`. |
-| Copy Scale | A `bool` specifying whether to copy the `Scale` component of the source bone's `Transform`. |
+| **Copy Translation** | A bool specifying whether to copy the Translation component of the source bone's Transform. |
+| **Copy Rotation** | A bool specifying whether to copy the Rotation component of the source bone's Transform. |
+| **Copy Scale** | A bool specifying whether to copy the Scale component of the source bone's Transform. |
 
 ## Transform (Modify) Bone
 
@@ -91,14 +92,14 @@ The **Copy Bone** control copies the `Transform` data or any component of it - i
 | Property | Description |
 | -------- | ----------- |
 | [REGION:tablesection]**Skeletal Control**[/REGION] ||
-| Bone To Modify | Name of the target bone to modify the `Transform` of. |
+| **Bone To Modify** | Name of the target bone to modify the Transform of. |
 
 | Pin | Description |
 | --- | ----------- |
 | [REGION:tablesection]**Input Pins**[/REGION] ||
-| Translation | A `Vector` specifying the `Translation` to apply to the target bone. |
-| Rotation | A `Rotator` specifying the `Rotation` to apply to the target bone. |
-| Scale | A `Vector` specifying the `Scale` to apply to the target bone. |
+| **Translation** | A Vector specifying the Translation to apply to the target bone. |
+| **Rotation** | A Rotator specifying the Rotation to apply to the target bone. |
+| **Scale** | A Vector specifying the Scale to apply to the target bone. |
 
 ## Spring Controller
 
@@ -110,46 +111,46 @@ its reference pose position and apply a force in the opposite direction.
 | Property | Description |
 | -------- | ----------- |
 | [REGION:tablesection]**Spring**[/REGION] ||
-| Spring Bone | Name of the bone to apply the spring solver to. |
-| Limit Displacement | If enabled, the spring solver is applied. |
-| Max Displacement | The maximum distance the bone can stretch. |
-| Spring Stiffness | A multiplier for the amount of force applied by the spring. Larger values require more bone velocity to displace the bone and result in a larger applied force. |
-| Spring Damping | The amount of the bone's velocity removed by the spring. |
-| Error Reset Thresh | A threshold displacement whereby if the spring stretches more than this amount, it resets in order to avoid errors introduced by sudden, large displacements such as those caused by teleporting Actors. |
-| No ZSpring | If enabled, the displacement in the `Z` axis is ignored by the spring and is always correct. |
+| **Spring Bone** | Name of the bone to apply the spring solver to. |
+| **Limit Displacement** | If enabled, the spring solver is applied. |
+| **Max Displacement** | The maximum distance the bone can stretch. |
+| **Spring Stiffness** | A multiplier for the amount of force applied by the spring. Larger values require more bone velocity to displace the bone and result in a larger applied force. |
+| **Spring Damping** | The amount of the bone's velocity removed by the spring. |
+| **Error Reset Thresh** | A threshold displacement whereby if the spring stretches more than this amount, it resets in order to avoid errors introduced by sudden, large displacements such as those caused by teleporting Actors. |
+| **No ZSpring** | If enabled, the displacement in the Z axis is ignored by the spring and is always correct. |
 
 ## Two Bone IK
 
-The **Two Bone IK** control applies an inverse kinematic (IK) solver to a _3-joint_ chain, such as the limbs of a character.
+The **Two Bone IK** control applies an Inverse Kinematic (IK) solver to a 3-joint chain, such as the limbs of a character.
 
 ![Apply Percentage of Rotation](node_twobone.png)
 
 | Property | Description |
 | -------- | ----------- |
 | [REGION:tablesection]**IK**[/REGION] ||
-| IK Bone | Name of the bone to apply the IK solver to. This should be the bone at the end of the bone chain, i.e. the _ankle_ or _wrist_ bone. |
+| **IK Bone** | Name of the bone to apply the IK solver to. This should be the bone at the end of the bone chain, i.e. the ankle or wrist bone. |
 | [REGION:tablesection]**Stretch Limits**[/REGION] ||
-| X | The ratio of the Max Limb Length, which is used to determine when to scale the bones. |
-| Y | The maximum allowable amount of scaling of the translation component. |
-| Allow Stretching | If _true_, allows the bones to be stretched. |
+| **X** | The ratio of the Max Limb Length, which is used to determine when to scale the bones. |
+| **Y** | The maximum allowable amount of scaling of the translation component. |
+| **Allow Stretching** | If true, allows the bones to be stretched. |
 | [REGION:tablesection]**End Effector**[/REGION] ||
-| Effector Space Bone Name | The name of the bone to use as the effector location when **Effector Location Space** is set to _BCS_BoneSpace_. |
-| Take Rotation from Effector Space | If enabled, the `Rotation` of the effector (component, parent, or bone) is applied to the **IK Bone**.  |
-| Maintain Effector Rel Rot | If enabled, the **IK Bone** will maintain its `Rotation` relative to the effector. [REGION:note]This is overridden when **Take Rotation from Effector Space** is enabled.[/REGION] |
-| Effector Location Space | The reference frame for the **Effector Location**. [INCLUDE:#LocationSpace] |
+| **Effector Space Bone Name** | The name of the bone to use as the effector location when Effector Location Space is set to BCS_BoneSpace. |
+| **Take Rotation from Effector Space** | If enabled, the Rotation of the effector (component, parent, or bone) is applied to the IK Bone.  |
+| **Maintain Effector Rel Rot** | If enabled, the IK Bone will maintain its Rotation relative to the effector. [REGION:note]This is overridden when Take Rotation from Effector Space is enabled.[/REGION] |
+| **Effector Location Space** | The reference frame for the Effector Location. [INCLUDE:#LocationSpace] |
 | [REGION:tablesection]**Join Target**[/REGION] ||
-| Joint Target Location Space | The reference frame for the **Joint Target Location**. [INCLUDE:#LocationSpace] |
-| Joint Target Space Bone Name | The name of the bone to use as the joint target location when **Joint Target Location Space** is set to _BCS_BoneSpace_. |
+| **Joint Target Location Space** | The reference frame for the Joint Target Location. [INCLUDE:#LocationSpace] |
+| **Joint Target Space Bone Name** | The name of the bone to use as the joint target location when Joint Target Location Space is set to BCS_BoneSpace. |
 
 <!--
 [EXCERPT:LocationSpace]
 [REGION:simpletable]
 | Value | Effect |
 | ----- | ------ |
-| BCS_WorldSpace | The location is assumed to be relative to world space. |
-| BCS_ComponentSpace | The location is assumed to be relative to the _SkeletalMeshComponent_. |
-| BCS_ParentSpace | The location is assumed to be relative to the parent bone. |
-| BCS_BoneSpace | The location is assumed to be relative to an arbitrary bone (specified by name as the value of another property). |
+| **BCS_WorldSpace** | The location is assumed to be relative to world space. |
+| **BCS_ComponentSpace** | The location is assumed to be relative to the SkeletalMeshComponent. |
+| **BCS_ParentSpace** | The location is assumed to be relative to the parent bone. |
+| **BCS_BoneSpace** | The location is assumed to be relative to an arbitrary bone (specified by name as the value of another property). |
 [/REGION] 
 [/EXCERPT:LocationSpace]
 -->
@@ -157,5 +158,5 @@ The **Two Bone IK** control applies an inverse kinematic (IK) solver to a _3-joi
 | Pin | Description |
 | --- | ----------- |
 | [REGION:tablesection]**Input Pins**[/REGION] ||
-| Effector Location | A `Vector` specifying the location effector in the **Effector Location Space**. |
-| Joint Target Location | A `Vector` specifying the location joint target in the **Joint Target Location Space**. |
+| **Effector Location** | A Vector specifying the location effector in the Effector Location Space. |
+| **Joint Target Location** | A Vector specifying the location joint target in the Joint Target Location Space. |
