#UE4doc #docs:
Removed Overview header
Added engine version field
Replaced 'Result node' with 'Result pin' as it was referring to the input node on the Final Animation Pose node.

diff --git a/Engine/Documentation/Source/Engine/Animation/NodeReference/SpaceConversion/SpaceConversion.INT.udn b/Engine/Documentation/Source/Engine/Animation/NodeReference/SpaceConversion/SpaceConversion.INT.udn
--- a/Engine/Documentation/Source/Engine/Animation/NodeReference/SpaceConversion/SpaceConversion.INT.udn
+++ b/Engine/Documentation/Source/Engine/Animation/NodeReference/SpaceConversion/SpaceConversion.INT.udn
@@ -2,6 +2,7 @@ Availability:Public
 Title:Convert Spaces Nodes
 Crumbs: %ROOT%, Engine, Engine/Animation, Engine/Animation/NodeReference
 Description:Animation nodes that convert poses between local and component space.
+version: 4.9
 
 [REGION:banner]
 ![Space Conversion](spaceconversion_banner.png)(convert:false)
@@ -9,7 +10,7 @@ Description:Animation nodes that convert poses between local and component space
 
 [TOC(start:2 end:3)]
 
-## Overview
+
 
 [EXCERPT:Intro]
 Poses can either be in local-space or component-space. The **Convert Spaces** nodes available in the **AnimGraph** 
@@ -22,7 +23,7 @@ nodes and all SkeletalControls operate in component-space. This means that the i
 prior to being passed in to one of these types of nodes. If the input pose is coming from a node that outputs a local 
 space pose, the pose must be converted to the correct space before a SkeletalControl can perform operations on it. 
 The resulting pose, after the SkeletalControl have performed the operations, then must be converted back to 
-local space for input to additional blends or the **Result** node.
+local space for input to additional blends or the Result pin.
 [/EXCERPT:Intro]
 
 [REGION:fullwidth]
@@ -30,7 +31,7 @@ local space for input to additional blends or the **Result** node.
 [/REGION]
 
 Because there is a cost associated with each conversion to or from component space, it is best to group any nodes that 
-operate in component-space, such as _SkeletalControls_, so that they are all performed in a row, limiting the number 
+operate in component-space, such as SkeletalControls, so that they are all performed in a row, limiting the number 
 of conversions needed.
 
 
