#UE4doc #docs:
Removed Overview header
Added engine version field
Changed font settings for a sentence

diff --git a/Engine/Documentation/Source/Engine/Animation/NonUniformScale/NonUniformScale.INT.udn b/Engine/Documentation/Source/Engine/Animation/NonUniformScale/NonUniformScale.INT.udn
--- a/Engine/Documentation/Source/Engine/Animation/NonUniformScale/NonUniformScale.INT.udn
+++ b/Engine/Documentation/Source/Engine/Animation/NonUniformScale/NonUniformScale.INT.udn
@@ -2,8 +2,9 @@ Availability:Public
 Title:Non-Uniform Scale Animation
 Crumbs: %ROOT%, Engine, Engine/Animation
 Description:Overview of how non-uniform scale animations can be used.
+version: 4.9
+
 
-## Overview
 The animation system in Unreal Engine 4 does support non-uniform scale animations. When importing an animation, if scale exists, it is automatically imported without the need for setting additional options. For memory conservation, the engine does not save scale for all animations and only saves it if the animation has a scale that is not 1. 
 
 For examples, see the videos below. 
@@ -13,7 +14,7 @@ For examples, see the videos below.
 [/REGION]
 
 
-**Click to watch playback of each video**
+_Click to watch playback of each video_
 
 Refer to the [FBX Animation Pipeline](Engine\Content\FBX\Animations) page for more information on importing FBX Animations. 
 
