#UE4doc #docs:
Added engine version field
Changed font settings for multiple terms

diff --git a/Engine/Documentation/Source/Engine/Animation/Overview/AnimationOverview.INT.udn b/Engine/Documentation/Source/Engine/Animation/Overview/AnimationOverview.INT.udn
--- a/Engine/Documentation/Source/Engine/Animation/Overview/AnimationOverview.INT.udn
+++ b/Engine/Documentation/Source/Engine/Animation/Overview/AnimationOverview.INT.udn
@@ -7,7 +7,7 @@ Description:Overview of the system for deforming SkeletalMeshes based on animati
 
 
 [EXCERPT:Intro]
-Unreal Engine 4's animation system allows for a deep level of control for your characters and Skeletal Meshes. In this document, we will provide a high level overview  of how animation works in UE4, and the base concepts that drive it. <!---In the documents that follow, we'll cover the---> 
+**Unreal Engine 4's** (**UE4**) animation system allows for a deep level of control for your characters and Skeletal Meshes. In this document, we will provide a high level overview  of how animation works in UE4, and the base concepts that drive it. <!---In the documents that follow, we'll cover the---> 
 
 
 
@@ -19,7 +19,7 @@ directly control the transformations of bones using SkeletalControls.
 [/EXCERPT:Intro]
 
 
-The purpose of this document is to offer a high-level overview of Unreal Engine 4's animation system, geared primarily toward users who are new to animating in Unreal. However, this should not be considered the comprehensive manual for how to animate skeletal assets in the Unreal Engine. Rather, think of this as a primer to familiarize you with the various pieces of the animation system, providing an illustration of how they all fit together and showing how animation in UE4 is more powerful flexible than ever before. 
+The purpose of this document is to offer a high-level overview of UE4's animation system, geared primarily toward users who are new to animating in UE4. However, this should not be considered the comprehensive manual for how to animate skeletal assets in the UE4. Rather, think of this as a primer to familiarize you with the various pieces of the animation system, providing an illustration of how they all fit together and showing how animation in UE4 is more powerful flexible than ever before. 
 
 We will start by identifying the primary terms and concepts of UE4's animation system:
 
@@ -82,24 +82,24 @@ We will start by identifying the primary terms and concepts of UE4's animation s
 
 		**Animation Poses**
 
-		An **Animation Pose** is essentially a snapshot of the Skeleton, including the position and rotation of all of its 
+		An Animation Pose is essentially a snapshot of the Skeleton, including the position and rotation of all of its 
 		bones. Think of a pose as what you see if you pause the playback of an animation. For example, if you paused a run 
 		animation mid-stride, you might see something like this:
 
 		![](AnimationPose_Running.png)(w:600)
 
-		It is as if the mesh, via the underlying Skeleton, has been **posed**. The data that represents the transforms of 
-		all of the bones in the _Skeleton_ is what makes up the actual Animation Pose.
+		It is as if the mesh, via the underlying Skeleton, has been posed. The data that represents the transforms of 
+		all of the bones in the Skeleton is what makes up the actual Animation Pose.
 
-		Some _AnimationSequences_ are designed purely to be used as poses. For instance, you might have a series of short 
+		Some AnimationSequences are designed purely to be used as poses. For instance, you might have a series of short 
 		one- or two-frame animations that simply define a pose. The pose below is from this type of Animation Sequence, 
 		which is designed to be used to cause the mesh to aim its weapon where the player is looking.
 
 		![](AnimationPose_Aiming.png)(w:600)
 
-		The Animation system uses this concept of Animation Poses within the _AnimGraph_. The data output from an _AnimationSequence_, 
-		_Blendspace_, or simple blend node is an Animation Pose, or snapshot of the animation up to that point, to apply 
-		to the _SkeletalMesh_. 
+		The Animation system uses this concept of Animation Poses within the AnimGraph. The data output from an AnimationSequence, 
+		Blendspace, or simple blend node is an Animation Pose, or snapshot of the animation up to that point, to apply 
+		to the SkeletalMesh. 
 
 		**Component vs Local Space**
 
@@ -215,15 +215,15 @@ For more information on Content Examples and how to use them, refer to the [Cont
 
 ## Animation Debug Features
 
-Debug features can be accessed by pressing the ` key to bring up the console while in play mode. 
+Debug features can be accessed by pressing the **`** key to bring up the console while in play mode. 
 
 ### Rendering Skeletal Mesh Bones
 
 ![](ShowDebugBones.png)(h:320)
 
-To display the **Skeletal Mesh Bones**, enter the `ShowDebug Bones` command inside the console. As seen in the image above the bones are represented by individual white lines. 
+To display the Skeletal Mesh Bones, enter the `ShowDebug Bones` command inside the console. As seen in the image above the bones are represented by individual white lines. 
 
-An alternative look, matching the bones displayed in **Persona**, can be enabled by entering the `ShowDebugToggleSubCategory 3DBones` command inside the console. 
+An alternative look, matching the bones displayed in Persona, can be enabled by entering the `ShowDebugToggleSubCategory 3DBones` command inside the console. 
 
 ### Animation Debug Output
 
