#UE4doc #docs:
Removed Overview header
Added engine version field
Changed font settings for multiple terms

diff --git a/Engine/Documentation/Source/Engine/Animation/Persona/PersonaEditor.INT.udn b/Engine/Documentation/Source/Engine/Animation/Persona/PersonaEditor.INT.udn
--- a/Engine/Documentation/Source/Engine/Animation/Persona/PersonaEditor.INT.udn
+++ b/Engine/Documentation/Source/Engine/Animation/Persona/PersonaEditor.INT.udn
@@ -2,6 +2,7 @@ Availability:Public
 Title:Persona Animation Editor Reference
 Crumbs: %ROOT%, Engine, Engine/Animation
 Description:Editor for previewing and editing AnimBlueprints, AnimationSequences and other animation assets.
+version: 4.9
 
 
 [VAR:ModeSize]w:450 convert:false[/VAR]
@@ -32,7 +33,7 @@ Description:Editor for previewing and editing AnimBlueprints, AnimationSequences
 
 [TOC (start:2 end:3)]
 
-## Overview
+
 
 [EXCERPT:Intro]
 
@@ -47,32 +48,32 @@ Description:Editor for previewing and editing AnimBlueprints, AnimationSequences
 [/REGION]
 
 
-Persona can be opened from the **Content Browser** by **double-clicking** on an animation asset, such as a:
+Persona can be opened from the **Content Browser** by **Double-clicking** on an animation asset, such as a:
 
-* _Animation Blueprint_
-* _Skeleton_
-* _Skeletal Mesh_
-* _Animation Sequence_
-* _BlendSpace_
-* _Animation Montage_
+* Animation Blueprint
+* Skeleton
+* Skeletal Mesh
+* Animation Sequence
+* BlendSpace
+* Animation Montage
 
 [REGION:tip]
-You can also open Persona by **right-clicking** on any of the above assets in the **Content Browser** and choosing **Edit** from the context menu!
+You can also open Persona by **Right-clicking** on any of the above assets in the **Content Browser** and choosing **Edit** from the context menu!
 [/REGION]
 
 ## Persona Modes and User Interface
 
 ![](ModeSelector.gif)(convert:false)
 [REGION:caption]
-The [**Persona Editing Modes**](Engine\Animation\Persona\Modes) are represented as buttons on the **Modes** toolbar in the top right corner of the interface. 
+The [Persona Editing Modes](Engine\Animation\Persona\Modes) are represented as buttons on the **Modes** toolbar in the top right corner of the interface. 
 [/REGION]
 
-As a vast collection of individual tools, Persona's interface could easily overwhelm if all of its components were visible at all times. To solve this, Persona has been divided into four key _modes_:
+As a vast collection of individual tools, Persona's interface could easily overwhelm if all of its components were visible at all times. To solve this, Persona has been divided into four key modes:
 
-1. **Skeleton**
-1. **Mesh**
-1. **Animation**
-1. **Graph**
+1. Skeleton
+1. Mesh
+1. Animation
+1. Graph
 
 Each mode corresponds to a specific type of task you will want to perform with the tools. Within each mode, only a certain selection of tools and interface elements are visible and available. Keep in mind, however, there is still tremendous room for flexibility, along with some cross-over for performing various tasks without having to switch modes. 
 
@@ -82,7 +83,7 @@ The locations below are the default locations for each window. As you move panel
 
 
 [REGION:tip]
-Whether visible or not, all aspects of the interface that can be accessed for each mode are available via Persona's _Window_ menu.
+Whether visible or not, all aspects of the interface that can be accessed for each mode are available via Persona's Window menu.
 [/REGION]
 
 [OBJECT:IconSection]
@@ -130,7 +131,7 @@ Whether visible or not, all aspects of the interface that can be accessed for ea
 		Interface options include:
 		[INCLUDE:Engine\Animation\Persona\Modes#animationmodeui]
 
-		Animation mode has one unique UI feature, which is the **Asset Editor**. This is a special panel that only appears when you **double-click** on an Animation Asset in the **Asset Browser**. By				default, the Asset Editor will be located at the center-bottom of Persona, directly below the Viewport.
+		Animation mode has one unique UI feature, which is the **Asset Editor**. This is a special panel that only appears when you **Double-click** on an Animation Asset in the **Asset Browser**. By default, the Asset Editor will be located at the center-bottom of Persona, directly below the Viewport.
 	[/PARAM]
 [/OBJECT]
 
@@ -337,7 +338,7 @@ See the [](Engine/Blueprints\Editor/UIComponents\MyBlueprint) page for details.
 
 ![](Variables.png)
 
-The **Variables** panel allows custom variables to be added to the _Animation Blueprint_ and lists any existing 
+The **Variables** panel allows custom variables to be added to the Animation Blueprint and lists any existing 
 variables. Variables will typically either have their values updated within the EventGraph of the Animation 
 Blueprint or externally through code. These can then be used in other calculations or as inputs to 
 [Blendspaces](Engine/Animation/Blendspaces) to control animation blending.
@@ -358,7 +359,7 @@ See the [](Engine/UI/LevelEditor/Details) page for general usage details.
 ![Persona: Debug Panel](CE_debug_panel.png)
 
 The **Debug** panel shows information designated as important by the designer in the form of breakpoints 
-and watch values as well as a trace stack of all nodes belonging to the _Animation Blueprint_ that have been executed. 
+and watch values as well as a trace stack of all nodes belonging to the Animation Blueprint that have been executed. 
 This window also shows controls for controlling execution of the game when using breakpoints.
 
 See [Debugging](Engine/Blueprints/UserGuide/Debugging) for more information on working with the **Debug** panel to debug Animation Blueprints.
@@ -373,7 +374,7 @@ See the [](Engine/Animation/Persona/AssetBrowser) page for details.
 
 ![](SkeletonTree.png)
 
-The **Skeleton Tree** panel displays the bone hierarchy of the _Skeleton_ the _Animation Blueprint_ being edited was created from.
+The **Skeleton Tree** panel displays the bone hierarchy of the Skeleton the Animation Blueprint being edited was created from.
 
 See the [Skeleton Tree](Engine/Animation/Persona\SkeletonTree) page for details.
 
@@ -428,7 +429,7 @@ See the [Blueprint Graph Panel](Engine/Blueprints/Editor/Modes/GraphPanel) secti
 
 ![Animation Sequence Editor](Engine/Animation/Sequences/Editor/persona_sequence.png)(w:720)
 
-See [](Engine/Animation/Sequences/Editor) for more information on setting up _AnimationSequences_.
+See [](Engine/Animation/Sequences/Editor) for more information on setting up AnimationSequences.
 
 ### BlendSpace Editor
 
@@ -436,7 +437,7 @@ See [](Engine/Animation/Sequences/Editor) for more information on setting up _An
 
 ![Animation Sequence Editor](Engine/Animation/Blendspaces/Editor/bs_editor.png)(w:720)
 
-See [](Engine/Animation/Blendspaces/Editor) for more information on setting up _Blendspaces_.
+See [](Engine/Animation/Blendspaces/Editor) for more information on setting up Blendspaces.
 
 
 --->
