#UE4doc #docs:
Added engine version field
Updated image to match updated UI

diff --git a/Engine/Documentation/Source/Engine/Animation/Persona/AnimAssetDetails/AnimAssetDetails.INT.udn b/Engine/Documentation/Source/Engine/Animation/Persona/AnimAssetDetails/AnimAssetDetails.INT.udn
--- a/Engine/Documentation/Source/Engine/Animation/Persona/AnimAssetDetails/AnimAssetDetails.INT.udn
+++ b/Engine/Documentation/Source/Engine/Animation/Persona/AnimAssetDetails/AnimAssetDetails.INT.udn
@@ -2,6 +2,7 @@ Availability:Public
 Title:Persona Anim Asset Details
 Crumbs: %ROOT%, Engine, Engine/Animation, Engine/Animation/Persona
 Description:User guide for the Anim Asset Details in the Persona animation editor.
+version: 4.9
 
 [VAR:Topic]
 [OBJECT:Topic]
