#UE4doc #docs:
Added engine version field
Changed some capitalization

diff --git a/Engine/Documentation/Source/Engine/Animation/Persona/AnimationNotifies/AnimationNotifies.INT.udn b/Engine/Documentation/Source/Engine/Animation/Persona/AnimationNotifies/AnimationNotifies.INT.udn
--- a/Engine/Documentation/Source/Engine/Animation/Persona/AnimationNotifies/AnimationNotifies.INT.udn
+++ b/Engine/Documentation/Source/Engine/Animation/Persona/AnimationNotifies/AnimationNotifies.INT.udn
@@ -2,6 +2,7 @@ Availability:Public
 Title:Persona Animation Notifies
 Crumbs: %ROOT%, Engine, Engine/Animation, Engine/Animation/Persona
 Description:User guide for the Animation Notifies window in the Persona animation editor.
+version: 4.9
 
 [VAR:Topic]
 [OBJECT:Topic]
@@ -33,8 +34,8 @@ The **Animation Notifies** window is not visible by default. To view it, on Pers
 Selecting a custom Animation Notify in the **Animation Notifies** window exposes its properties in Persona's **Details** panel.
 [/REGION]
 
-To delete a custom Animation Notify, **right-click** it, and then click **Delete**.
+To delete a custom Animation Notify, **Right-click** it, and then click **Delete**.
 
-To rename a custom Animation Notify, **right-click** it, and then click **Rename**.
+To rename a custom Animation Notify, **Right-click** it, and then click **Rename**.
 
 For more information about Persona's **Asset Editor**, see [](Engine/Animation/Persona\AssetEditor).
