#UE4doc #docs:
Added engine version field
Changed font settings for multiple terms
Changed a few words to be more clear.

diff --git a/Engine/Documentation/Source/Engine/Animation/Persona/AssetBrowser/PersonaAssetBrowser.INT.udn b/Engine/Documentation/Source/Engine/Animation/Persona/AssetBrowser/PersonaAssetBrowser.INT.udn
--- a/Engine/Documentation/Source/Engine/Animation/Persona/AssetBrowser/PersonaAssetBrowser.INT.udn
+++ b/Engine/Documentation/Source/Engine/Animation/Persona/AssetBrowser/PersonaAssetBrowser.INT.udn
@@ -2,6 +2,7 @@ Availability:Public
 Title:Persona Asset Browser
 Crumbs: %ROOT%, Engine, Engine/Animation, Engine/Animation/Persona
 Description:User guide for the Asset Browser in the Persona animation editor.
+version: 4.9
 
 [VAR:Topic]
 [OBJECT:Topic]
@@ -26,15 +27,15 @@ Description:User guide for the Asset Browser in the Persona animation editor.
 [EXCERPT:Overview]
 ![](AssetBrowser.png)
 
-The **Asset Browser** panel displays the _AnimSequences_, _BlendSpaces_, _AnimMontages_, and _AimOffsetBlendSpaces_ 
-that are valid for the target _Skeleton_ in a mini **Content Browser**-style interface. 
+The **Asset Browser** panel displays the AnimSequences, BlendSpaces, AnimMontages, and AimOffsetBlendSpaces 
+that are valid for the target Skeleton in a mini **Content Browser**-style interface. 
 [/EXCERPT:Overview]
 
-The **Asset Browser** panel is visible only when **Persona** is in **Animation** mode.
+The **Asset Browser** panel is visible only when Persona is in **Animation** mode.
 
 | Action | Description |
 | ------ | ----------- |
-| **Click Asset** | Selects the asset and enables it in the preview Viewport. |
+| **Left-click Asset** | Selects the asset and enables it in the preview Viewport. |
 | **Double-click Asset** | Selects the asset and displays its properties in the **Anim Asset Details** panel and the **Asset Editor**. For more information, see [](Engine/Animation/Persona\AnimAssetDetails) and [](Engine/Animation/Persona\AssetEditor).|
 | **Drag Asset into Level Editor viewport** | Creates an instance of the asset as a Skeletal Mesh Actor in the Level Editor Viewport. |
 
