#UE4doc #docs:
Added engine version field
Changed font settings/capitalization for multiple terms
Replaced both pictures to match updated UI

diff --git a/Engine/Documentation/Source/Engine/Animation/Persona/AssetEditor/AssetEditor.INT.udn b/Engine/Documentation/Source/Engine/Animation/Persona/AssetEditor/AssetEditor.INT.udn
--- a/Engine/Documentation/Source/Engine/Animation/Persona/AssetEditor/AssetEditor.INT.udn
+++ b/Engine/Documentation/Source/Engine/Animation/Persona/AssetEditor/AssetEditor.INT.udn
@@ -2,6 +2,7 @@ Availability:Public
 Title:Asset Editor
 Crumbs: %ROOT%, Engine, Engine/Animation, Engine/Animation/Persona
 Description:User guide for the Asset Editor in the Persona animation editor.
+version: 4.9
 
 [VAR:Topic]
 [OBJECT:Topic]
@@ -23,13 +24,13 @@ Description:User guide for the Asset Editor in the Persona animation editor.
 [/OBJECT]
 [/VAR]
 
-Persona's **Asset Editor** opens when you **double-click** an animation asset in the **Asset Browser**. For more information about the **Asset Browser**, see [](Engine/Animation/Persona\AssetBrowser).
+Persona's **Asset Editor** opens when you **Double-click** an animation asset in the **Asset Browser**. For more information about the **Asset Browser**, see [](Engine/Animation/Persona\AssetBrowser).
 
 You can use the **Asset Editor** to add Notifies, including Particle Systems, Sound Cues, and Animation Notifies, to specific points in time on the tracks of animation assets.
 
 ![](AssetEditor.png)
 
-To add a Notify to an animation track, **right-click** the track, click **Add Notify**, and do one of the following:
+To add a Notify to an animation track, **Right-click** the track, click **Add Notify**, and do one of the following:
 * To add a new Particle System, click **PlayParticleEffect**.
 * To add a new Sound Cue, click **PlaySound**.
 * To add a new custom Animation Notify, click **New Notify**, and give your new Animation Notify a name.
