#UE4doc #docs:
Added engine version field
Replaced image to match updated UI
Changed steps to match updated functionality.

diff --git a/Engine/Documentation/Source/Engine/Animation/Persona/BasePoseManager/BasePoseManager.INT.udn b/Engine/Documentation/Source/Engine/Animation/Persona/BasePoseManager/BasePoseManager.INT.udn
--- a/Engine/Documentation/Source/Engine/Animation/Persona/BasePoseManager/BasePoseManager.INT.udn
+++ b/Engine/Documentation/Source/Engine/Animation/Persona/BasePoseManager/BasePoseManager.INT.udn
@@ -2,6 +2,7 @@ Availability:Public
 Title:Persona Retarget Source Manager
 Crumbs: %ROOT%, Engine, Engine/Animation, Engine/Animation/Persona
 Description:User guide for the Retarget Source Manager in the Persona animation editor.
+version: 4.9
 
 [VAR:Topic]
 [OBJECT:Topic]
@@ -29,7 +30,7 @@ Persona's **Retarget Source Manager** enables you to create a list of potential
 
 **To add a new retarget source for reference:**
 
-1. **Right-click** in the **Retarget Source Manager**, and click **Add**.
+1. **Left-click** the ** Add New Retarget Source** button.
 1. Select the Skeletal Mesh whose base pose you want to use from the list.
 
 [region:note]
