#UE4doc #docs:
Added engine version field
Changed font settings for multiple terms
Replaced multiple images to match update UI
Changed all mentions of 'Anim Blueprint Editor' to 'Anim Preview Editor' to match the updated name of the panel

diff --git a/Engine/Documentation/Source/Engine/Animation/Persona/BlueprintDefaults/BlueprintDefaults.INT.udn b/Engine/Documentation/Source/Engine/Animation/Persona/BlueprintDefaults/BlueprintDefaults.INT.udn
--- a/Engine/Documentation/Source/Engine/Animation/Persona/BlueprintDefaults/BlueprintDefaults.INT.udn
+++ b/Engine/Documentation/Source/Engine/Animation/Persona/BlueprintDefaults/BlueprintDefaults.INT.udn
@@ -1,7 +1,8 @@
 Availability:Public
-Title:Persona Anim Blueprint Editor
+Title:Persona Anim Preview Editor
 Crumbs: %ROOT%, Engine, Engine/Animation, Engine/Animation/Persona
-Description:User guide for the Anim Blueprint Editor panel in the Persona animation editor.
+Description:User guide for the Anim Preview Editor panel in the Persona animation editor.
+version: 4.9
 
 [VAR:Topic]
 [OBJECT:Topic]
@@ -26,9 +27,9 @@ Description:User guide for the Anim Blueprint Editor panel in the Persona animat
 [TOC (start:2)]
 
 
-## Overview
 
-The **Anim Blueprint Editor** panel in Persona lists the variables and variable categories of the **Animation Blueprint**.
+
+The **Anim Preview Editor** panel in Persona lists the variables and variable categories of the **Animation Blueprint**.
 
 ![](BlueprintDefaults.png)
 
@@ -36,25 +37,25 @@ When you create or modify a Blueprint variable, it is assigned to the **Default*
 
 ![](AB_Details.png)
 
-For the most part, an Animation Blueprint's **Anim Blueprint Editor** panel can be modified in the same way as the Blueprint Editor Defaults mode.  For more information about the Blueprint Editor Defaults mode, see [](Engine/Blueprints\Editor\Modes\Defaults).
+For the most part, an Animation Blueprint's **Anim Preview Editor** panel can be modified in the same way as the Blueprint Editor Defaults mode.  For more information about the Blueprint Editor Defaults mode, see [](Engine/Blueprints\Editor\Modes\Defaults).
 
 ## Previewing
 
-You can preview your various animation assets by manually modifying the values of the variables that control them in the **Anim Blueprint Editor** panel. As you modify the values of the variables, the Viewport displays the appropriate animation assets based on the new values. 
+You can preview your various animation assets by manually modifying the values of the variables that control them in the **Anim Preview Editor** panel. As you modify the values of the variables, the Viewport displays the appropriate animation assets based on the new values. 
 
 ![](BlueprintDefaults_Preview.png)
 
 ## Edit Modes
 
-Along the top of the **Anim Blueprint Editor** panel are two modes that can be switched between and affect the permanency of changes made to variable values in the panel.
+Along the top of the **Anim Preview Editor** panel are two modes that can be switched between and affect the permanency of changes made to variable values in the panel.
 
 | Edit Mode | Description | |
 | --- | ----------- | --- |
-| **Edit Preview** | Changes made to variables in the **Anim Blueprint Editor** will only reflect in the Viewport preview and will not be saved. Changes will revert to default values when the Blueprint is recompiled. | ![](EditPreview.png) |
-| **Edit Defaults** | Changes made to variables in the **Anim Blueprint Editor** affect both the preview and the default values of the variables and will be saved.<br>[INCLUDE:#publicview]| ![](EditDefaults.png) |
+| **Edit Preview** | Changes made to variables in the **Anim Preview Editor** will only reflect in the Viewport preview and will not be saved. Changes will revert to default values when the Blueprint is recompiled. | ![](EditPreview.png) |
+| **Edit Defaults** | Changes made to variables in the **Anim Preview Editor** affect both the preview and the default values of the variables and will be saved.<br>[INCLUDE:#publicview]| ![](EditDefaults.png) |
 
 <!--
 [EXCERPT:publicview]
-* **Public View** - An option to toggle whether or not to display only variables marked as Public. After toggling **Public View** on or off, the Blueprint must be recompiled before changes can be seen. Only available in **Edit Defaults** mode.
+* **Public View** - An option to toggle whether or not to display only variables marked as Public. After toggling Public View on or off, the Blueprint must be recompiled before changes can be seen. Only available in **Edit Defaults** mode.
 [/EXCERPT:publicview]
 -->
