#UE4doc #docs
Removed duplicated steps and fixed reference to "Development Editor" mode.

diff --git a/Engine/Documentation/Source/Programming/Development/CompilingProjects/CompilingProjects.INT.udn b/Engine/Documentation/Source/Programming/Development/CompilingProjects/CompilingProjects.INT.udn
--- a/Engine/Documentation/Source/Programming/Development/CompilingProjects/CompilingProjects.INT.udn
+++ b/Engine/Documentation/Source/Programming/Development/CompilingProjects/CompilingProjects.INT.udn
@@ -92,7 +92,7 @@ When building in Visual Studio, you are only compiling the game project, not the
 When building in Visual Studio, you are compiling the engine in addition to the game project.  
 [/REGION]
 
--->
+
 
 1.	Set the **Solution Configuration** to the configuration you want to build. In this example, it is set to **Development Editor**. Refer to the [Build Configuration](#BuildConfiguration) section for descriptions of each available configuration.
 
@@ -100,15 +100,15 @@ When building in Visual Studio, you are compiling the engine in addition to the
 
 1. **Right-click** your game project and choose **Rebuild** to recompile.
 
-	![](rebuild_project_outside.png)
-
-1.	Set the **Solution Configuration** to the configuration you want to build. In this example, it is set to **Development**. Refer to the [Build Configuration](#BuildConfiguration) section for descriptions of each available configuration.
+	![](rebuild_project_rocket.png)
+-->
+1.	Set the **Solution Configuration** to the configuration you want to build. In this example, it is set to **Development Editor**. Refer to the [Build Configuration](#BuildConfiguration) section for descriptions of each available configuration.
 
 	![](rocket_config_development_menu.png)
 
 1. **Right-click** your game project and choose **Rebuild** to recompile.
 
-	![](rebuild_project_rocket.png)
+	![](rebuild_project_outside.png)
 
 You can now [run the engine](GettingStarted/RunningUnrealEngine) with your compiled project.  
 
