#UE4doc #docs:
Added engine version field

diff --git a/Engine/Documentation/Source/Engine/Animation/Persona/BoneManipulation/BoneManipulationMode.INT.udn b/Engine/Documentation/Source/Engine/Animation/Persona/BoneManipulation/BoneManipulationMode.INT.udn
--- a/Engine/Documentation/Source/Engine/Animation/Persona/BoneManipulation/BoneManipulationMode.INT.udn
+++ b/Engine/Documentation/Source/Engine/Animation/Persona/BoneManipulation/BoneManipulationMode.INT.udn
@@ -6,6 +6,7 @@ Related: Engine/Content/FBX/Animations
 Related: Engine/Content/Types/SkeletalMeshes
 Related: Engine/Animation/PhysicallyDrivenAnimation
 Related: Engine/Content/Tools/MayaRiggingTool
+version: 4.9
 
 **Bone Manipulation Mode** is a mode in Persona in which technical artists can temporarily translate or rotate bones in order to look for skinning issues in a Skeletal Mesh while previewing animations.
 
