#UE4doc #docs:
Added engine version field

diff --git a/Engine/Documentation/Source/Engine/Animation/Persona/Debug/Debug.INT.udn b/Engine/Documentation/Source/Engine/Animation/Persona/Debug/Debug.INT.udn
--- a/Engine/Documentation/Source/Engine/Animation/Persona/Debug/Debug.INT.udn
+++ b/Engine/Documentation/Source/Engine/Animation/Persona/Debug/Debug.INT.udn
@@ -2,6 +2,7 @@ Availability:Public
 Title:Persona Debug
 Crumbs: %ROOT%, Engine, Engine/Animation, Engine/Animation/Persona
 Description:User guide for the Asset Browser in the Persona animation editor.
+version: 4.9
 
 [VAR:Topic]
 [OBJECT:Topic]
