#UE4doc #docs:
Added engine version field
Removed Overview header
Changed font settings for multiple terms
Added terms/explinations to the Help section as it was missing options due to being out of date.

diff --git a/Engine/Documentation/Source/Engine/Animation/Persona/MenuBar/PersonaMenuBar.INT.udn b/Engine/Documentation/Source/Engine/Animation/Persona/MenuBar/PersonaMenuBar.INT.udn
--- a/Engine/Documentation/Source/Engine/Animation/Persona/MenuBar/PersonaMenuBar.INT.udn
+++ b/Engine/Documentation/Source/Engine/Animation/Persona/MenuBar/PersonaMenuBar.INT.udn
@@ -2,17 +2,16 @@ Availability:NoPublish
 Title:Persona Menu Bar
 Crumbs: %ROOT%, Engine, Engine/Animation, Engine/Animation/Persona
 Description:User guide for the Menu Bar in the Persona animation editor.
+version: 4.9
 
 
-## Overview
-
 ![](PersonaMenuBar.png)
 
 
-The menu is displayed at the top of the Window. It is important to note that not all menu options are available at all times. Availability will depend on the type of asset being edited. In this document, we are listing everything as if you have opened an **Animation Blueprint**. 
+The menu is displayed at the top of the Window. It is important to note that not all menu options are available at all times. Availability will depend on the type of asset being edited. In this document, we are listing everything as if you have opened an Animation Blueprint. 
 
 [REGION:note]
-If you are editing animation assets such as _Skeletal Meshes_ or _Animation Sequences_, you will see fewer menus.
+If you are editing animation assets such as Skeletal Meshes or Animation Sequences, you will see fewer menus.
 [/REGION]
 
 ## File
@@ -24,8 +23,8 @@ If you are editing animation assets such as _Skeletal Meshes_ or _Animation Sequ
 | **Save All** | Saves all assets. |
 | **Choose Files to Save** | Opens a dialog with save options for assets. |
 | **Connect to Source Control** | Connects to a source control server, such as Perforce. See [](Engine/UI/SourceControl) documentation for details. |
-| [REGION:tablesection]Blueprint - This section is only available when Persona is opened via an _Animation Blueprint_[/REGION] ||
-| **Compile** | Compiles the animation Blueprint. |
+| [REGION:tablesection]Blueprint - This section is only available when Persona is opened via an Animation Blueprint[/REGION] ||
+| **Compile** | Compiles the Animation Blueprint. |
 | **Refresh All Nodes** | Refreshes all nodes in the Blueprint to account for external changes. |
 | **Reparent Blueprint** | Allows you to change the parent of the Blueprint. |
 | **Diff** | Runs a source control diff operation to find any discrepancies between the current version and any previous version. Nonfunctional if not connected to source control.|
@@ -54,7 +53,7 @@ If you are editing animation assets such as _Skeletal Meshes_ or _Animation Sequ
 | **Set Preview Mesh as Default** | Changes the skeletons default preview mesh to the current open preview mesh. The skeleton will require saving after this action. |
 | **Remove Unused Bones from Skeleton** | Removes any bones from the skeleton that are not used by any of its meshes. The skeleton and associated animations will require saving after this action.  |
 | **Record to new Animation** | Create a new animation from the currently playing animation. |
-| [REGION:tablesection]**The following section applies only if Persona was opened via an _Animation_** asset[/REGION] ||
+| [REGION:tablesection]**The following section applies only if Persona was opened via an Animation** asset[/REGION] ||
 | **Reimport** (Animation) | Reimports the AnimationSequence. |
 | **Apply Compression** | Apply compression to the current animation. |
 | **Export to FBX** | Export the current animation to an FBX. |
@@ -63,7 +62,7 @@ If you are editing animation assets such as _Skeletal Meshes_ or _Animation Sequ
 
 ### View
 
-The View menu is only available when editing an Animation Blueprint.
+The **View** menu is only available when editing an Animation Blueprint.
 
 | Command | Description |
 | ------- | ----------- |
@@ -79,7 +78,7 @@ The View menu is only available when editing an Animation Blueprint.
 
 ### Debug
 
-The Debug menu is only available when editing an Animation Blueprint.
+The **Debug** menu is only available when editing an Animation Blueprint.
 
 | Command | Description |
 | ------- | ----------- |
@@ -100,7 +99,7 @@ The panels available in the **Window** menu, as well as their order of appearanc
 
 | Command | Description |
 | ------- | ----------- |
-| **Toolbar** |Opens the Toolbar, which provides a horizontal array of shortcut buttons for common operations.|
+| **Toolbar** |Opens the **Toolbar**, which provides a horizontal array of shortcut buttons for common operations.|
 | **Skeleton Tree** | Opens the **Skeleton Tree** window, showing the character's skeleton in a hierarchical list. |
 | **Asset Browser** | Opens the **Asset Browser**, which lists all the Animations, Blendspaces, and other animation assets associated with this skeletal mesh. |
 | **Viewport** | Opens the **Viewport**, where you can see the skeletal mesh animating. |
@@ -108,10 +107,10 @@ The panels available in the **Window** menu, as well as their order of appearanc
 | **Animation Notifies** | Opens the **Animation Notifies**  panel, which lists off all Animation Notifies on the mesh and with what animations each notify is associated. |
 | **Retarget Source Manager** | Opens the **Retarget Source Manager**, which is useful for correctly [retargeting](Engine/Animation/AnimationRetargeting) animation between meshes with different builds, such as from a fat to a skinny character.|
 | **Details** | Shows the **Details** window, where properties of various nodes can be edited. |
-| **Mesh Details**  _(Mesh mode)_ | Opens the **Mesh Details** panel, where you can change aspects of the mesh including Material applications, LODs, Physics Assets, and more. |
-| **Anim Asset Details** _(Animation mode)_ | Opens the **Anim Asset Details** panel, where you perform most of the necessary edits to Animation Sequences, Blend Spaces, Animation Montages, and Aim Offsets. |
-| **Anim Blueprint Editor** _(Animation Blueprint)_| Opens the **Anim Blueprint Editor** panel if it is not already open, where you can set the default values for all variables in your Blueprint or change the values of the preview instance, depending on mode. |
-| **Asset Override Editor** _(Animation Blueprint)_| Opens the **Asset Override Editor** panel if it is not already open. This editor is for overriding the animation assets referenced by the parent animation graph. |
+| **Mesh Details**  (Mesh mode) | Opens the **Mesh Details** panel, where you can change aspects of the mesh including Material applications, LODs, Physics Assets, and more. |
+| **Anim Asset Details** (Animation mode) | Opens the **Anim Asset Details** panel, where you perform most of the necessary edits to Animation Sequences, Blend Spaces, Animation Montages, and Aim Offsets. |
+| **Anim Blueprint Editor** (Animation Blueprint)| Opens the **Anim Blueprint Editor** panel if it is not already open, where you can set the default values for all variables in your Blueprint or change the values of the preview instance, depending on mode. |
+| **Asset Override Editor** (Animation Blueprint)| Opens the **Asset Override Editor** panel if it is not already open. This editor is for overriding the animation assets referenced by the parent animation graph. |
 | **Animation SlotName** | Opens the **Animation SlotName** panel and allows you to modify custom animation SlotNames. |
 | **Animation Slot Group Name** | Opens the **Animation Slot Group Name** panel which lets you modify custom animation Slot Group Names. |
 | **Debug** | Opens the **Debug** window, which shows all breakpoints throughout the Blueprint. |
@@ -127,8 +126,12 @@ The panels available in the **Window** menu, as well as their order of appearanc
 | Command | Description |
 | ------- | ----------- |
 | [REGION:tablesection]**Browse**[/REGION] ||
-| **Documentation** | Opens the Unreal Engine 4 documentation. |
+| **Browse Persona Documentation** | Opens the Unreal Engine 4 documentation. |
 | [REGION:tablesection]**Online**[/REGION] ||
-| Contains links to the Epic Games and Unreal Engine support sites. ||
+| **Forums...** | Contains a link to the Unreal Engine forums. ||
+| **Answerhub...** | Contains a link to the Unreal Engine Answerhub. ||
+| **Wiki...** | Contains a link to the Unreal Engine wiki. ||
+| **Visit UnrealEngine.com...** | Contains a link to UnrealEngine.com. ||
+| **Credits** | Displays the credits window. ||
 | [REGION:tablesection]**Application**[/REGION] ||
 | **About Unreal Editor** | Opens the About dialog displaying information about the engine version. |
