#UE4doc #docs:
Added engine version field
Replaced picture due to out of date UI

diff --git a/Engine/Documentation/Source/Engine/Animation/Persona/MeshDetails/MeshDetails.INT.udn b/Engine/Documentation/Source/Engine/Animation/Persona/MeshDetails/MeshDetails.INT.udn
--- a/Engine/Documentation/Source/Engine/Animation/Persona/MeshDetails/MeshDetails.INT.udn
+++ b/Engine/Documentation/Source/Engine/Animation/Persona/MeshDetails/MeshDetails.INT.udn
@@ -2,6 +2,7 @@ Availability:Public
 Title:Persona Mesh Details
 Crumbs: %ROOT%, Engine, Engine/Animation, Engine/Animation/Persona
 Description:User guide for the Mesh Details panel in the Persona animation editor.
+version: 4.9
 
 [VAR:Topic]
 [OBJECT:Topic]
