#UE4doc #docs:
Added engine version field
Removed Overview header
Changed font settings for multiple terms
Fixed some capitalization
Replaced all images to reflect UI changes

diff --git a/Engine/Documentation/Source/Engine/Animation/Persona/Modes/Modes.INT.udn b/Engine/Documentation/Source/Engine/Animation/Persona/Modes/Modes.INT.udn
--- a/Engine/Documentation/Source/Engine/Animation/Persona/Modes/Modes.INT.udn
+++ b/Engine/Documentation/Source/Engine/Animation/Persona/Modes/Modes.INT.udn
@@ -2,6 +2,7 @@ Availability: Public
 Title:Persona Editing Modes
 Crumbs: %ROOT%, Engine, Engine/Animation, Engine/Animation/Persona
 Description:User guide for the editing modes available in the Persona animation editor.
+version: 4.9
 
 [VAR:Topic]
 [OBJECT:Topic]
@@ -24,16 +25,15 @@ Description:User guide for the editing modes available in the Persona animation
 [/VAR]
 
 
-## Overview
 
 [EXCERPT:Overview]
-It is important to remember that Persona is actually a collection of various animation tools and editors, rather than just a single tool unto itself. It allows for editing of Skeletons, Meshes, Animation, and Animation Blueprints. As such, Persona contains a variety of different _modes_ in which it can operate in order to optimally handle the editing task at hand. 
+It is important to remember that Persona is actually a collection of various animation tools and editors, rather than just a single tool unto itself. It allows for editing of Skeletons, Meshes, Animation, and Animation Blueprints. As such, Persona contains a variety of different modes in which it can operate in order to optimally handle the editing task at hand. 
 
 In each mode, Persona's UI updates to show a different layout of tools and editors.
 [/EXCERPT:Overview]
 
 [REGION:note]
-	Not all of the modes are available at all times. For instance, Animation Blueprint mode is _only_ available when an Animation Blueprint is opened from the **Content Browser**. 
+	Not all of the modes are available at all times. For instance,**Graph** mode is only available when an Animation Blueprint is opened from the **Content Browser**. 
 
 	Also, be aware that changes made to editing layouts will be maintained. As an example, if you move a panel while in Skeleton editing mode, then switch away to Animation, when you switch back to Skeleton, the moved panel will keep the position you gave it.
 [/REGION]
@@ -44,7 +44,7 @@ In each mode, Persona's UI updates to show a different layout of tools and edito
 
 [![](Persona_Mode_Skeleton.png)(w:800)](Engine/Animation/Persona/Modes/Skeleton)
 
-(Click the image to visit the **Skeleton Editing Mode** page.)
+(Click the image to visit the **Skeleton Editing** mode page.)
 
 [EXCERPT:SkeletonModeUI]
 1. [Toolbar](Engine/Animation/Persona\Toolbar) 
@@ -64,7 +64,7 @@ The **Details** panel in Persona will only show properties pertaining to animati
 
 [![](Persona_Mode_Mesh.png)(w:800)](Engine/Animation/Persona/Modes/Mesh)
 
-(Click the image to visit the **Mesh Editing Mode** page.)
+(Click the image to visit the **Mesh Editing** mode page.)
 
 [EXCERPT:MeshModeUI]
 1. [Toolbar](Engine/Animation/Persona\Toolbar) 
@@ -82,7 +82,7 @@ The **Details** panel in Persona will only show properties pertaining to animati
 
 [![](Persona_Mode_Anmation.png)(w:800)](Engine/Animation/Persona/Modes/Animation)
 
-(Click the image to visit the **Animation Editing Mode** page.)
+(Click the image to visit the **Animation Editing** mode page.)
 
 [EXCERPT:AnimationModeUI]
 1. [Toolbar](Engine/Animation/Persona\Toolbar) 
@@ -101,7 +101,7 @@ The **Details** panel in Persona will only show properties pertaining to animati
 
 [![](Persona_Mode_Vim.png)(w:800)](Engine/Animation/Persona/Modes/Graph)
 
-(Click the image to visit the **Graph Editing Mode** page.)
+(Click the image to visit the **Graph Editing** mode page.)
 
 [EXCERPT:GraphModeUI]
 1. [Toolbar](Engine/Animation/Persona\Toolbar) 
