#UE4doc #docs:
Removed Overview header
Added engine version
Changed font settings for multiple terms
Fixed some capitalization
Replaced picture to reflect updated UI

diff --git a/Engine/Documentation/Source/Engine/Animation/Persona/Modes/Animation/Animation.INT.udn b/Engine/Documentation/Source/Engine/Animation/Persona/Modes/Animation/Animation.INT.udn
--- a/Engine/Documentation/Source/Engine/Animation/Persona/Modes/Animation/Animation.INT.udn
+++ b/Engine/Documentation/Source/Engine/Animation/Persona/Modes/Animation/Animation.INT.udn
@@ -2,6 +2,7 @@ Availability: Public
 Title:Animation Mode
 Crumbs: %ROOT%, Engine, Engine/Animation, Engine/Animation/Persona, Engine/Animation/Persona/Modes
 Description:User guide for Animation Mode in the Persona animation editor.
+version: 4.9
 
 [VAR:Topic]
 [OBJECT:Topic]
@@ -23,17 +24,17 @@ Description:User guide for Animation Mode in the Persona animation editor.
 [/OBJECT]
 [/VAR]
 
-##Overview
+
 
 [REGION:fullwidth]
 ![](AnimationModeWindow.png)
 [/REGION]
 
-The **Animation Mode** of **Persona** (accessed by selecting the **Animation** tab from the toolbar window) contains several windows which appear in other modes, however the three windows (denoted above) are the ones used most when working in **Animation Mode**. 
+The **Animation** mode of **Persona** (accessed by selecting the **Animation** tab from the toolbar window) contains several windows which appear in other modes, however the three windows (denoted above) are the ones used most when working in Animation mode. 
 
 1. [Anim Asset Details](Engine\Animation\Persona\AnimAssetDetails) - The **Anim Asset Details** panel is a context-sensitive property editor that allows you to change settings on various animation assets, such as Anim Sequences, Blend Spaces, Anim Montages, and so on. By default, the **Anim Asset Details** panel is only visible when Persona is in **Animation** mode.
-1. [Anim Sequence Editor](Engine\Animation\Sequences\Editor) - The **Animation Sequence Editor** provides functionality for previewing and editing individual _AnimationSequence_ assets. The properties, including compression scheme, of the _AnimationSequence_ can be set, and animation notification events - also called Notifies - can be added to allow camera effects, particle effects, sounds, and more to be triggered based on the animation.
-1. [Asset Browser](Engine\Animation\Persona\AssetBrowser) - The **Asset Browser** panel displays the _AnimSequences_, _BlendSpaces_, _AnimMontages_, and _AimOffsetBlendSpaces_ that are valid for the target _Skeleton_ in a mini **Content Browser**-style interface. 
+1. [Anim Sequence Editor](Engine\Animation\Sequences\Editor) - The **Animation Sequence Editor** provides functionality for previewing and editing individual AnimationSequence assets. The properties, including compression scheme, of the AnimationSequence can be set, and animation notification events - also called Notifies - can be added to allow camera effects, particle effects, sounds, and more to be triggered based on the animation.
+1. [Asset Browser](Engine\Animation\Persona\AssetBrowser) - The **Asset Browser** panel displays the AnimSequences, BlendSpaces, AnimMontages, and AimOffsetBlendSpaces that are valid for the target Skeleton in a mini **Content Browser**-style interface. 
 
 
 
