UEDOC-2249

Click to View videos have been replaced with Youtube Videos

diff --git a/Engine/Documentation/Source/Engine/Animation/NonUniformScale/NonUniformScale.INT.udn b/Engine/Documentation/Source/Engine/Animation/NonUniformScale/NonUniformScale.INT.udn
--- a/Engine/Documentation/Source/Engine/Animation/NonUniformScale/NonUniformScale.INT.udn
+++ b/Engine/Documentation/Source/Engine/Animation/NonUniformScale/NonUniformScale.INT.udn
@@ -9,12 +9,41 @@ The animation system in Unreal Engine 4 does support non-uniform scale animation
 
 For examples, see the videos below. 
 
-[REGION:lightbox]
-[![](AngryChest.png)(w:300 h:200)](AngryChest.wmv) [![](CastSpell.png)(w:300 h:200)](CastSpell.wmv) [![](Trooper.png)(w:300 h:200)](Trooper.wmv)
-[/REGION]
-
-
-_Click to watch playback of each video_
+[OBJECT:EmbeddedVideo]
+[PARAMLITERAL:width]
+640
+[/PARAMLITERAL]
+[PARAMLITERAL:height]
+360
+[/PARAMLITERAL]
+[PARAMLITERAL:videoid]
+kokp5ux-Jok
+[/PARAMLITERAL]
+[/OBJECT]
+
+[OBJECT:EmbeddedVideo]
+[PARAMLITERAL:width]
+640
+[/PARAMLITERAL]
+[PARAMLITERAL:height]
+360
+[/PARAMLITERAL]
+[PARAMLITERAL:videoid]
+Od5WGC9_bBo
+[/PARAMLITERAL]
+[/OBJECT]
+
+[OBJECT:EmbeddedVideo]
+[PARAMLITERAL:width]
+640
+[/PARAMLITERAL]
+[PARAMLITERAL:height]
+360
+[/PARAMLITERAL]
+[PARAMLITERAL:videoid]
+sh5DTeWRE0A
+[/PARAMLITERAL]
+[/OBJECT]
 
 Refer to the [FBX Animation Pipeline](Engine\Content\FBX\Animations) page for more information on importing FBX Animations. 
 
