#UE4doc #docs:
Added engine version field
Removed Overview header
Changed font settings for a few terms
Fixed capitalization of a term
Replaced image due to updated UI

diff --git a/Engine/Documentation/Source/Engine/Animation/Persona/Modes/Graph/Graph.INT.udn b/Engine/Documentation/Source/Engine/Animation/Persona/Modes/Graph/Graph.INT.udn
--- a/Engine/Documentation/Source/Engine/Animation/Persona/Modes/Graph/Graph.INT.udn
+++ b/Engine/Documentation/Source/Engine/Animation/Persona/Modes/Graph/Graph.INT.udn
@@ -2,6 +2,7 @@ Availability: Public
 Title:Graph Mode
 Crumbs: %ROOT%, Engine, Engine/Animation, Engine/Animation/Persona, Engine/Animation/Persona/Modes
 Description:User guide for Graph Mode in the Persona animation editor.
+version: 4.9
 
 [VAR:Topic]
 [OBJECT:Topic]
@@ -23,19 +24,19 @@ Description:User guide for Graph Mode in the Persona animation editor.
 [/OBJECT]
 [/VAR]
 
-##Overview
+
 
 [REGION:fullwidth]
 ![](AnimationGraph.png)
 [/REGION]
 
-The **Graph Mode** of **Persona** is where you specify the logic behind and hook up animations and animation states based off variables and conditions set within the **AnimGraph** and **EventGraph**. Working with the graphs, is similar to working with **Blueprint** graphs however there are some key details to **Animation Blueprints** as outlined below. 
+The **Graph** mode of **Persona** is where you specify the logic behind and hook up animations and animation states based off variables and conditions set within the **AnimGraph** and **EventGraph**. Working with the graphs, is similar to working with Blueprint graphs however there are some key details to Animation Blueprints as outlined below. 
 
 ##Animation Blueprints Overview
 
 [INCLUDE:Engine/Animation/AnimBlueprints#Intro]
 
-For more information on working with **Animation Blueprints**, refer to the [Animation Blueprints](Engine/Animation/AnimBlueprints) documentation. 
+For more information on working with Animation Blueprints, refer to the [Animation Blueprints](Engine/Animation/AnimBlueprints) documentation. 
 
 
 
