UEDOC-2250

Replaced gif on page with a YouTube video

diff --git a/Engine/Documentation/Source/Engine/Animation/Persona/PersonaEditor.INT.udn b/Engine/Documentation/Source/Engine/Animation/Persona/PersonaEditor.INT.udn
--- a/Engine/Documentation/Source/Engine/Animation/Persona/PersonaEditor.INT.udn
+++ b/Engine/Documentation/Source/Engine/Animation/Persona/PersonaEditor.INT.udn
@@ -63,7 +63,18 @@ You can also open Persona by **Right-clicking** on any of the above assets in th
 
 ## Persona Modes and User Interface
 
-![](ModeSelector.gif)(convert:false)
+[OBJECT:EmbeddedVideo]
+[PARAMLITERAL:width]
+640
+[/PARAMLITERAL]
+[PARAMLITERAL:height]
+360
+[/PARAMLITERAL]
+[PARAMLITERAL:videoid]
+3DKiBKb5a3E
+[/PARAMLITERAL]
+[/OBJECT]
+
 [REGION:caption]
 The [Persona Editing Modes](Engine\Animation\Persona\Modes) are represented as buttons on the **Modes** toolbar in the top right corner of the interface. 
 [/REGION]
