#UE4doc #docs:
Added engine version field
Removed Overview header
Changed font settings for multiple terms
Replaced an image due to new UI changes
Commented out the 'Mirroring' section as that section no longer exists in the editor.

diff --git a/Engine/Documentation/Source/Engine/Animation/Persona/Modes/Mesh/Mesh.INT.udn b/Engine/Documentation/Source/Engine/Animation/Persona/Modes/Mesh/Mesh.INT.udn
--- a/Engine/Documentation/Source/Engine/Animation/Persona/Modes/Mesh/Mesh.INT.udn
+++ b/Engine/Documentation/Source/Engine/Animation/Persona/Modes/Mesh/Mesh.INT.udn
@@ -2,6 +2,7 @@ Availability: Public
 Title:Mesh Mode
 Crumbs: %ROOT%, Engine, Engine/Animation, Engine/Animation/Persona, Engine/Animation/Persona/Modes
 Description:User guide for Mesh Mode in the Persona animation editor.
+version: 4.9
 
 [VAR:Topic]
 [OBJECT:Topic]
@@ -23,9 +24,9 @@ Description:User guide for Mesh Mode in the Persona animation editor.
 [/OBJECT]
 [/VAR]
 
-##Overview 
+
 [EXCERPT:Overview]
-**Mesh Mode** of the **Persona** animation editor (accessed by clicking the **Mesh** tab) shares some of the same windows that are available while in other Editing Modes, however there are two windows which are specific to **Mesh Mode** (highlighted below): [Mesh Details](Engine/Animation/Persona\MeshDetails) and [Morph Target Previewer](Engine/Animation/Persona\MorphTargetPreviewer). 
+**Mesh** mode of the **Persona** animation editor (accessed by clicking the **Mesh** tab) shares some of the same windows that are available while in other Editing Modes, however there are two windows which are specific to **Mesh** mode (highlighted below): [Mesh Details](Engine/Animation/Persona\MeshDetails) and [Morph Target Previewer](Engine/Animation/Persona\MorphTargetPreviewer). 
 [/EXCERPT:Overview]
 
 [REGION:fullwidth]
@@ -33,7 +34,7 @@ Description:User guide for Mesh Mode in the Persona animation editor.
 [/REGION]
 
 ##Mesh Details
-The **Mesh Details** window allow you to modify properties pertaining to the mesh of the Skeletal Mesh such as **Material** and **LOD** settings, assigning **Apex Clothing**, adding meshes to the existing mesh, adjusting **Mirroring** settings or assigning a **Physics Asset** to the mesh. You can also **Reimport** asset data, adjust **Thumbnail Orbit** or **Texture Streaming** settings. 
+The **Mesh Details** window allow you to modify properties pertaining to the mesh of the Skeletal Mesh such as **Material** and **Level of Detail (LOD)** settings, assigning **Apex Clothing**, adding meshes to the existing mesh, adjusting **Mirroring** settings or assigning a **Physics Asset** to the mesh. You can also **Reimport** asset data, adjust **Thumbnail Orbit** or **Texture Streaming** settings. 
 
 ###LOD Settings
 
@@ -41,9 +42,9 @@ The **Mesh Details** window allow you to modify properties pertaining to the mes
 ![](LOD_Settings.png)
 [/REGION]
 
-In this section, you can assign **Materials** to aspects of the Skeletal Mesh and toggle whether or not those **Materials** cast shadows or not. 
+In this section, you can assign Materials to aspects of the Skeletal Mesh and toggle whether or not those Materials cast shadows or not. 
 
-In the **LOD Info** section, there are settings that allow you to modify how the **LOD(s)** function: 
+In the **LOD Info** section, there are settings that allow you to modify how the LOD(s) function: 
 
 | Property| Description |
 |--- | --- |
@@ -78,14 +79,15 @@ This section allows you to add meshes to the existing mesh or remove meshes that
 
 ###Mesh
 
-This section lists the **Skeleton** that is used with the **Skeletal Mesh** and its location inside the **Content Browser**. 
+This section lists the Skeleton that is used with the Skeletal Mesh and its location inside the **Content Browser**. 
 
 ###Skeletal Mesh
 
 ![](SkeletalMeshSection.png)
 
-The **Skeletal Mesh** section allows you to define **Materials** for any **Material Interfaces** that exist and whether or not to enable **Shadow Casting** from the **Material**. The ability to set whether or not to **Use Full Precision UVs** is available which if _true_, will use 32 bit UVs and if _false_, will use 16 bit UVs to save memory. 
+The **Skeletal Mesh** section allows you to define Materials for any Material Interfaces that exist and whether or not to enable Shadow Casting from the Material. The ability to set whether or not to **Use Full Precision UVs** is available which if true, it will use 32 bit UVs and if false, it will use 16 bit UVs to save memory. 
 
+<!--Commenting this section out as it no longer exists as of 4.9 Preview 4 
 ###Mirroring
 
 ![](MirroringSection.png)
@@ -96,11 +98,11 @@ There are two options available inside the **Mirroring** section that can be set
 |--- | --- |
 | **Skel Mirror Table** | This lists all the bones that should be mirrored. |
 | **Skel Mirror Axis** | Set the Mirror Axis to X, Y, Z or None. |
-| **Skel Mirror Flip Axis** | Set the Mirror Flip Axis to X, Y, Z or None. |
+| **Skel Mirror Flip Axis** | Set the Mirror Flip Axis to X, Y, Z or None. |-->
 
 ###Physics
 
-In the **Physics** section, you can assign a **Physics Asset** to the mesh. The location of the asset in the **Content Browser** (if applicable) is also shown. 
+In the **Physics** section, you can assign a Physics Asset to the mesh. The location of the asset in the **Content Browser** (if applicable) is also shown. 
 
 ###Reimport
 
