#UE4doc #docs:
Added engine version field
Removed Overview header
Changed font settings for multiple terms
Fixed capitalization on a term
Replaced an image due to updated UI

diff --git a/Engine/Documentation/Source/Engine/Animation/Persona/Modes/Skeleton/Skeleton.INT.udn b/Engine/Documentation/Source/Engine/Animation/Persona/Modes/Skeleton/Skeleton.INT.udn
--- a/Engine/Documentation/Source/Engine/Animation/Persona/Modes/Skeleton/Skeleton.INT.udn
+++ b/Engine/Documentation/Source/Engine/Animation/Persona/Modes/Skeleton/Skeleton.INT.udn
@@ -2,11 +2,10 @@ Availability: Public
 Title:Skeleton Mode
 Crumbs: %ROOT%, Engine, Engine/Animation, Engine/Animation/Persona
 Description:
-
-## Overview
+version: 4.9
 
 [EXCERPT:Overview]
-**Skeleton mode** allows for visualization and control over those aspects of a character or Skeletal Mesh involving its bone or joint hierarchy. Example uses include:
+**Skeleton** mode allows for visualization and control over those aspects of a character or Skeletal Mesh involving its bone or joint hierarchy. Example uses include:
 
 * Creating [Sockets](Engine/Content\Types\SkeletalMeshes\Sockets) to attach items to a Skeletal Mesh
 * Editing the relative position and rotation for each [Sockets](Engine/Content\Types\SkeletalMeshes\Sockets)
