#UE4doc #docs:
Added engine version field
Removed Overview header
Changed font settings for multiple terms
Fixed capitalization on a few terms

diff --git a/Engine/Documentation/Source/Engine/Animation/Persona/MorphTargetPreviewer/MorphTargetPreviewer.INT.udn b/Engine/Documentation/Source/Engine/Animation/Persona/MorphTargetPreviewer/MorphTargetPreviewer.INT.udn
--- a/Engine/Documentation/Source/Engine/Animation/Persona/MorphTargetPreviewer/MorphTargetPreviewer.INT.udn
+++ b/Engine/Documentation/Source/Engine/Animation/Persona/MorphTargetPreviewer/MorphTargetPreviewer.INT.udn
@@ -8,7 +8,7 @@ Related: Engine/Content/Types/SkeletalMeshes
 Related: Engine/Content/Tools/MayaRiggingTool
 Related: Engine/Content/Types/StaticMeshes/MorphTargets
 Related: Engine/Animation/AnimBlueprints
-
+version: 4.9
 [VAR:Topic]
 [OBJECT:Topic]
 	[PARAM:image]
@@ -32,7 +32,6 @@ Related: Engine/Animation/AnimBlueprints
 [TOC (start:2 end:3)]
 
 
-## Overview
 
 [EXCERPT:Overview]
 
@@ -54,7 +53,7 @@ The **Morph Target Previewer** allows you to preview any morph targets (also kno
 The **Morph Target Previewer** is visible by default when Persona is in **Mesh** mode.
 
 [REGION:tip]
-To see an example of a **Morph Target** applied to a character, refer to the [Animation Content Examples](Resources\ContentExamples\Animation) page under section 1.6.
+To see an example of a Morph Target applied to a character, refer to the [Animation Content Examples](Resources\ContentExamples\Animation) page under section 1.6.
 [/REGION]
 
 
@@ -69,7 +68,7 @@ The **Morph Target Previewer** panel has a fairly straightforward interface:
 
 The Search Bar allows you to quickly filter the list of Morph Targets to quickly find the one you need. Simply type the first few letters of the desired target and the list will begin filtering.
 
-You can also **Right Click** on a **Morph Target** which will pop-up a dialog box with additional actions such as **Delete**. 
+You can also **Right-click** on a Morph Target which will pop-up a dialog box with additional actions such as **Delete**. 
 
 ![](RightClickDelete.png)
 
@@ -78,12 +77,12 @@ You can also **Right Click** on a **Morph Target** which will pop-up a dialog bo
 Morph Targets are importable as a part of a Skeletal Mesh, as well as independently of a given mesh. This is supported via the FBX file format. For more information on the setup procedure, as well as how to import Morph Targets into Unreal, please see [](Engine/Content\FBX\MorphTargets).
 
 [REGION:tip]
-To delete a Morph Target, **right-click** its name in the **Morph Target Previewer** window, and click **Delete** on the context menu that appears.
+To delete a Morph Target, **Right-click** its name in the **Morph Target Previewer** window, and click **Delete** on the context menu that appears.
 [/REGION]
 
 ## Using Morph Targets
 
-Once your Morph Targets are in place, you will need to set up your **Animation Blueprint** to utilize them. This is done within your Event Graph via the **Set Morph Target** node.
+Once your Morph Targets are in place, you will need to set up your Animation Blueprint to utilize them. This is done within your Event Graph via the **Set Morph Target** node.
 
 ![](SetMorphTarget.png)
 
@@ -91,9 +90,9 @@ Once your Morph Targets are in place, you will need to set up your **Animation B
 | --- | --- |
 | [REGION:tablesection] Input Pins [/REGION] ||
 | Execution | An execution wire to trigger the effect of the node. |
-| Target | The target Skeletal Mesh. For most circumstances, this will point to **self**. |
-| Morph Target Name | The name of the Morph Target being edited. |
-| Value | A float value (between 0 and 1) used to set the value of the Morph Target being edited. |
+| **Target** | The target Skeletal Mesh. For most circumstances, this will point to "self". |
+| **Morph Target Name** | The name of the Morph Target being edited. |
+| **Value** | A float value (between 0 and 1) used to set the value of the Morph Target being edited. |
 | [REGION:tablesection] Output Pins [/REGION] ||
 | Execution | Passes execution on to the next node. |
 
