#UE4doc #docs:
Added engine version field
Removed Overview header
Changed font settings for multiple terms
Fixed capitalization for multiple terms
Replaced multiple images to reflect updated UI
Added multiple terms/definitions to different sections along with editing existing ones to match changes to the editor.

diff --git a/Engine/Documentation/Source/Engine/Animation/Persona/SkeletonTree/SkeletonTree.INT.udn b/Engine/Documentation/Source/Engine/Animation/Persona/SkeletonTree/SkeletonTree.INT.udn
--- a/Engine/Documentation/Source/Engine/Animation/Persona/SkeletonTree/SkeletonTree.INT.udn
+++ b/Engine/Documentation/Source/Engine/Animation/Persona/SkeletonTree/SkeletonTree.INT.udn
@@ -2,6 +2,7 @@ Availability:Public
 Title:Persona Skeleton Tree
 Crumbs: %ROOT%, Engine, Engine/Animation, Engine/Animation/Persona
 Description:User guide for the Skeleton Tree in the Persona animation editor.
+version: 4.9
 
 [VAR:Topic]
 [OBJECT:Topic]
@@ -25,7 +26,7 @@ Description:User guide for the Skeleton Tree in the Persona animation editor.
 
 ![](SkeletonTree.png)  
 
-## Overview
+
 
 [EXCERPT:Overview]
 The **Skeleton Tree** panel shows the skeletal hierarchy of the current skeletal mesh. Within it, you can perform the following functions:
@@ -49,7 +50,7 @@ The **Skeleton Tree** panel shows the skeletal hierarchy of the current skeletal
 
 ![](SearchLine.png)
 
-The Search line allows for quick and easy searching of a skeletal hierarchy. Type in some part of the name of the Bone, Socket, or Preview  Object. and the list will filter to what you want to see.
+The **Search** line allows for quick and easy searching of a skeletal hierarchy. Type in some part of the name of the Bone, Socket, or Preview  Object. and the list will filter to what you want to see.
 
 ![](SeachFiltering.png)
 
@@ -59,7 +60,7 @@ The Search line allows for quick and easy searching of a skeletal hierarchy. Typ
 
 ![](button_BoneAndSocketButtons.png)
 
-The Bone and Socket buttons are used to control visibility over various bone and Socket types. Each button opens a menu allowing you to select the types of Bones or Sockets you want to see.
+The **Bone** and **Socket** buttons are used to control visibility over various bone and Socket types. Each button opens a menu allowing you to select the types of Bones or Sockets you want to see.
 
 [REGION:raw]
 ![](BonesMenu.png)
@@ -79,19 +80,19 @@ The Bone and Socket buttons are used to control visibility over various bone and
 
 | Option | Description |
 |---|---|
-|**Show Active Sockets**	| Shows only active Sockets. When you _customize_ a Socket (**right-click** and choose _Customize Socket_), the Socket is copied from the Skeleton onto the mesh, making a special socket that is exclusive to this Skeletal Mesh and not shared by any other Skeletal Meshes using the same Skeleton. When you do that, the new customized Socket - also known as a _Mesh Socket_ - becomes active and the original Socket, which remains a part of the Skeleton asset, is considered inactive. In effect, this will hide away any Sockets on the Skeleton that have a customized Mesh Socket counterpart. |
-|**Show Mesh Sockets**		| Shows only _customized_ Sockets, which are those Sockets that are exclusive to this Skeletal Mesh and are not used by any other Skeletal Mesh that shares a Skeleton asset with this one.  |
-|**Show Skeleton Sockets**	| Shows only those Sockets which are _non-customized_, meaning those which are part of the Skeleton asset and are not exclusive to this Skeletal Mesh. |
+|**Show Active Sockets**	| Shows only active Sockets. When you customize a Socket (**Right-click** and choose Customize Socket), the Socket is copied from the Skeleton onto the mesh, making a special socket that is exclusive to this Skeletal Mesh and not shared by any other Skeletal Meshes using the same Skeleton. When you do that, the new customized Socket - also known as a Mesh Socket - becomes active and the original Socket, which remains a part of the Skeleton asset, is considered inactive. In effect, this will hide away any Sockets on the Skeleton that have a customized Mesh Socket counterpart. |
+|**Show Mesh Sockets**		| Shows only customized Sockets, which are those Sockets that are exclusive to this Skeletal Mesh and are not used by any other Skeletal Mesh that shares a Skeleton asset with this one.  |
+|**Show Skeleton Sockets**	| Shows only those Sockets which are non-customized, meaning those which are part of the Skeleton asset and are not exclusive to this Skeletal Mesh. |
 |**Show All Sockets**		|Shows all Sockets available on the Skeletal Mesh. |
 |**Hide Sockets**			| Hides all Sockets. |
 
 ## Context Menus
 
-The context menu - accessible by **right-clicking** within the **Skeleton Tree** - has multiple forms depending on what you are clicking: one for bones, one for Sockets, and one for Preview Assets
+The context menu - accessible by **Right-clicking** within the **Skeleton Tree** - has multiple forms depending on what you are clicking: one for bones, one for Sockets, and one for Preview Assets
 
 ### Bone Context Menu
 
-When **right-clicking** on a bone, you will see the following context menu:
+When **Right-clicking** on a bone, you will see the following context menu:
 
 ![](SkeletonContextMenu_Bone.png)
 
@@ -99,46 +100,45 @@ When **right-clicking** on a bone, you will see the following context menu:
 |---|---|
 |[REGION:tablesection]Socket Actions[/REGION]||
 |**Add Socket** | Adds a new Socket to the currently selected bone. Note that if you have hidden Sockets using the button near the top of the **Skeleton Tree**, then this option will be grayed out. |
-|**Paste Selected Socket** | If you have copied a Socket to the clipboard (**right-click** and choose _Copy Selected Socket_), then this will paste that socket onto the selected bone.|
+|**Paste Selected Socket** | If you have copied a Socket to the clipboard (**Right-click** and choose Copy Selected Socket), then this will paste that socket onto the selected bone.|
 |[REGION:tablesection]Selected Bone Actions[/REGION]||
 |**Copy Selected Bone Names**|Copies the selected bone names to the clipboard. This is useful if you should need to enter a complex bone name into a property field.|
 | [REGION:tablesection]Bone Translation Retargeting[/REGION]||
 |**Recursively Set Translation Retargeting Skeleton**|This recursively sets the bone and all of its children to use Translation Retargeting based off the Skeleton. See the [Translation Retargeting Controls](#translationretargetingcontrols)
 |**Recursively Set Translation Retargeting Animation**|This recursively sets the bone and all of its children to use Translation Retargeting based off the Animation. See the [Translation Retargeting Controls](#translationretargetingcontrols)
 |**Recursively Set Translation Retargeting AnimationScaled**|This recursively sets the bone and all of its children to use Translation Retargeting based off the scaled Animation. See the [Translation Retargeting Controls](#translationretargetingcontrols)
+|**Recursively Set Translation Retargeting AnimationRelative**|This recursively sets the bone and all of its children to use Translation Retargeting relative to the Animation. See the [Translation Retargeting Controls](#translationretargetingcontrols)
 | [REGION:tablesection]LOD Bone Reduction[/REGION]||
-|**Remove Selected Bone from...**| Places the selected bone and its children on the _Bones to Remove_ list for a certain LOD level within the LOD settings. See LOD Settings for more information.|
-|**Add Remove Selected Bone to...**| Removes the selected bone from the _Bones to Remove_ list for the specified LOD level. See LOD Settings for more information.|
+|**Remove Selected Bone from...**| Places the selected bone and its children on the Bones to Remove list for a certain LOD level within the LOD settings. See LOD Settings for more information.|
+|**Add Remove Selected Bone to...**| Removes the selected bone from the Bones to Remove list for the specified LOD level. See LOD Settings for more information.|
 | [REGION:tablesection]Attached Assets Actions[/REGION]||
 |**Add Preview Asset**| Opens a menu of all eligible assets that can be temporarily attached to a bone for previewing purposes. For instance, if you want to see what a certain piece of armor would look like attached a bone, this is the option you would use. Note, however, that this is not a permanent attachment and will only be visible in Persona. |
 
 ### Socket Context Menu
 
-When **right-clicking** on a [Socket](Engine/Content\Types\SkeletalMeshes\Sockets), you  will see the following context menu:
+When **Right-clicking** on a [Socket](Engine/Content\Types\SkeletalMeshes\Sockets), you  will see the following context menu:
 
 ![](SkeletonContextMenu_Socket.png)
 
 | **Name** | **Description** |
 |---|---|
-|[REGION:tablesection]Selected Socket Actions[/REGION]||
-| **Rename Socket** | Allows you to rename a selected Socket.  |
+|[REGION:tablesection]Selected Item Actions[/REGION]||
 | **Delete** | Deletes a selected Socket.  |
-|**Customize Socket**| Copies the Socket from the Skeleton asset onto the specific Skeletal mesh being used. In this way, a Socket can be edited slighting for characters using the same Skeleton asset, but having different mesh proportions.|
-|**Remove Socket Customization** | This is only available when **right-clicking** on a Socket that has been _customized_. This removes the customized copy of the socket, leaving only the original Socket on the Skeleton asset.|
+|[REGION:tablesection]Selected Socket Actions[/REGION]||
 |**Copy Selected Sockets**| Copies the selected socket to the clipboard for pasting on to other bones.|
-| [REGION:tablesection]Bone Translation Retargeting[/REGION]||
-|**Recursively Set Translation Retargeting Skeleton**|This recursively sets the bone and all of its children to use Translation Retargeting based off the Skeleton. See the [Translation Retargeting Controls](#translationretargetingcontrols)
-|**Recursively Set Translation Retargeting Animation**|This recursively sets the bone and all of its children to use Translation Retargeting based off the Animation. See the [Translation Retargeting Controls](#translationretargetingcontrols)
-|**Recursively Set Translation Retargeting AnimationScaled**|This recursively sets the bone and all of its children to use Translation Retargeting based off the scaled Animation. See the [Translation Retargeting Controls](#translationretargetingcontrols)
+| **Rename Socket** | Allows you to rename a selected Socket.  |
+|**Create Mesh Socket**| Copies the Socket from the Skeleton asset onto the specific Skeletal mesh being used. In this way, a Socket can be edited slighting for characters using the same Skeleton asset, but having different mesh proportions.|
+|**Delete Mesh Socket** | This is only available when **Right-clicking** on a Socket that has been made into a Mesh Socket. This removes the Mesh copy of the socket, leaving only the original Socket on the Skeleton asset.|
 | [REGION:tablesection]Attached Assets Actions[/REGION]||
 |**Add Preview Asset**| Opens a menu of all eligible assets that can be temporarily attached to a bone for previewing purposes. For instance, if you want to see what a certain piece of armor would look like attached a bone, this is the option you would use. Note, however, that this is not a permanent attachment and will only be visible in Persona. |
+|**Remove All Attached Assets**| Removed all attached Preview Assets from this socket.|
 
 
 ### Preview Asset Context Menu
 
 A **Preview Asset** is an asset, such as a Static Mesh or Skeletal Mesh, that has been temporarily attached to a bone or Socket for Previewing purposes.
 
-When **right-clicking** on a Preview Asset, you will see the following context menu:
+When **Right-clicking** on a Preview Asset, you will see the following context menu:
 
 
 ![](SkeletonContextMenu_PreviewAsset.png)
@@ -166,10 +166,11 @@ Each dropdown allows you to control how that specific bone will handle retargeti
 
 ![](RetargetingControls.png)
 
-Each button contains the following three options. Note that you can change these settings recursively on each child of a selected bone by **right-clicking** and choosing the recursive commands in the context menu.
+Each button contains the following three options. Note that you can change these settings recursively on each child of a selected bone by **Right-clicking** and choosing the recursive commands in the context menu.
 
 |Name|Description|
 |---|---|
 |**Animation**|Use translation from animation.|
 |**Skeleton**|Use translation from skeleton.|
-|**AnimationScaled**|Use translation from animation, scale length by skeleton's proportions. |
+|**AnimationScaled**|Use translation from animation, scale length by the Skeleton's proportions. |
+|**AnimationRelative**|Use relative translation from animation similar to an additive animation. |
