UEDOC-2251

Linked the Blueprint Debugging page to the Persona Debugging page

diff --git a/Engine/Documentation/Source/Engine/Animation/Persona/Debug/Debug.INT.udn b/Engine/Documentation/Source/Engine/Animation/Persona/Debug/Debug.INT.udn
--- a/Engine/Documentation/Source/Engine/Animation/Persona/Debug/Debug.INT.udn
+++ b/Engine/Documentation/Source/Engine/Animation/Persona/Debug/Debug.INT.udn
@@ -2,7 +2,7 @@ Availability:Public
 Title:Persona Debug
 Crumbs: %ROOT%, Engine, Engine/Animation, Engine/Animation/Persona
 Description:User guide for the Asset Browser in the Persona animation editor.
-version: 4.9
+Version:4.9
 
 [VAR:Topic]
 [OBJECT:Topic]
@@ -24,4 +24,6 @@ version: 4.9
 [/OBJECT]
 [/VAR]
 
-[INCLUDE:Engine/Blueprints/UserGuide/Debugging#debugtab]
+The **Debug** panel shows information designated as important by the designer in the form of breakpoints and watch values as well as a trace stack of all nodes belonging to the Blueprint that have been executed. This window also shows controls for controlling execution of the game when using breakpoints.
+
+The Debug panel in persona is the same panel as found in the Blueprint Editor. To learn more about using this panel, see the [](Engine/Blueprints/UserGuide/Debugging) page.
\ No newline at end of file
