#UE4doc #docs - Added the engine version tag, and updated the steps for importing cloth assets to be clearer.

diff --git a/Engine/Documentation/Source/Engine/Physics/Apex/Apex.INT.udn b/Engine/Documentation/Source/Engine/Physics/Apex/Apex.INT.udn
--- a/Engine/Documentation/Source/Engine/Physics/Apex/Apex.INT.udn
+++ b/Engine/Documentation/Source/Engine/Physics/Apex/Apex.INT.udn
@@ -2,6 +2,7 @@ Availability: Public
 Title:APEX
 Crumbs: %ROOT%, Engine, Engine/Physics
 Description:Creation and Import of Destructibles using the Nvidia APEX Toolset.
+Version: 4.9
 
 [VAR:Topic]
 [OBJECT:Topic]
diff --git a/Engine/Documentation/Source/Engine/Physics/Apex/ApexTypes/ApexTypes.INT.udn b/Engine/Documentation/Source/Engine/Physics/Apex/ApexTypes/ApexTypes.INT.udn
--- a/Engine/Documentation/Source/Engine/Physics/Apex/ApexTypes/ApexTypes.INT.udn
+++ b/Engine/Documentation/Source/Engine/Physics/Apex/ApexTypes/ApexTypes.INT.udn
@@ -2,6 +2,7 @@ Availability: Public
 Title:APEX Types
 Crumbs: %ROOT%, Engine, Engine/Physics, Engine/Physics/Apex
 Description:Nvidia APEX Physics Lab APX and APB files.
+Version: 4.9
 
 [VAR:Topic]
 [OBJECT:Topic]
@@ -51,13 +52,14 @@ Description:Nvidia APEX Physics Lab APX and APB files.
 		![](APEXtopicIcon.png)(w:150 convert:false)
 	[/PARAM]
 	[PARAM:description]
-	Nvidia PhysX APEX has two file types that Unreal Editor 4 can import, .APX and .APB. Both cloth and destructible meshes can be exported as these file types, and imported into Unreal Engine 4 via the **Content Browser**.
+	Nvidia PhysX APEX has two file types that Unreal Editor 4 can import, .APX and .APB. 
+	Both cloth and destructible meshes can be exported as these file types, but only destructible meshes can be import via the **Content Browser.** Cloth assets will need to be imported via **Persona**.
 	[/PARAM]
 [/OBJECT]
 
 ## Creation
 
-See the [APEX Physics Lab](Engine/Physics/Apex) documentation on creating and exporting of APEX assets.
+See the [**APEX Physics Lab**](Engine/Physics/Apex) documentation on creating and exporting of APEX assets.
 
 ## Import
 
@@ -71,14 +73,14 @@ In the **Content Browser**:
 Or
 
 1. Select your **.APX** or **.APB** file(s) file in a Windows File Explorer window.
-1. Click and drag it (them) into the Content Browser in Unreal Editor 4.
+1. Click and drag them into the **Content Browser** in Unreal Editor 4.
 
 ###Cloth Assets
 
 1. Click the ![New Asset](CB_button_NewAsset.png) button.
 1. Select your Skeletal Mesh that uses the cloth asset.
-1. Open the **Skeletal Mesh**. 
-1. In the **Mesh Details** panel under the Clothing tab click the ![New Asset](CB_Button_AddAPEX.png) button.
+1. Double-click the **Skeletal Mesh** in the **Content Browser** to open it.
+1. In the **Mesh Details** panel under the **Clothing** tab click the ![New Asset](CB_Button_AddAPEX.png) button.
 1. Select the **.APX** or **.APB** file(s) you wish to import. 
 1. Use the ![New Asset](CB_Button_clothingElement.png) drop down below each **Material Element** to select the corresponding cloth asset to assign to the material.
 
