#UE4doc #docs - added engine version tag and bolded link.

diff --git a/Engine/Documentation/Source/Engine/Physics/Destructibles/DestructibleProperties/DestructibleProperties.INT.udn b/Engine/Documentation/Source/Engine/Physics/Destructibles/DestructibleProperties/DestructibleProperties.INT.udn
--- a/Engine/Documentation/Source/Engine/Physics/Destructibles/DestructibleProperties/DestructibleProperties.INT.udn
+++ b/Engine/Documentation/Source/Engine/Physics/Destructibles/DestructibleProperties/DestructibleProperties.INT.udn
@@ -2,6 +2,7 @@
 Title:Destructible Properties Matrix
 Crumbs: %ROOT%, Engine
 Description:Descriptions of all the properties on Destructible Actors in the DestructibleMesh Editor.
+Version: 4.9
 
 [VAR:Topic]
 [OBJECT:Topic]
@@ -74,7 +75,7 @@ Description:Descriptions of all the properties on Destructible Actors in the Des
 | **Sound** | The sound or sound cue to play when a chunk breaks.|
 
 
-For the Physics and Collision section of properties, see the [Collision](Engine/Physics/Collision) documentation.
+For the Physics and Collision section of properties, see the [**Collision**](Engine/Physics/Collision) documentation.
 	
 ### Skeletal Mesh Sub Category
 
