#UE4doc #docs - added engine version tag and removed redundancy in explanation.

diff --git a/Engine/Documentation/Source/Engine/Rendering/LightingAndShadows/AmbientOcclusion/AmbientOcclusion.INT.udn b/Engine/Documentation/Source/Engine/Rendering/LightingAndShadows/AmbientOcclusion/AmbientOcclusion.INT.udn
--- a/Engine/Documentation/Source/Engine/Rendering/LightingAndShadows/AmbientOcclusion/AmbientOcclusion.INT.udn
+++ b/Engine/Documentation/Source/Engine/Rendering/LightingAndShadows/AmbientOcclusion/AmbientOcclusion.INT.udn
@@ -2,6 +2,7 @@ Availability:Public
 Title:Ambient Occlusion
 Crumbs: %ROOT%, Engine, Engine/Rendering/LightingAndShadows
 Description:Approximation of attenuation of light due to occlusion.
+Version: 4.9
 
 [VAR:Topic]
 [OBJECT:Topic]
@@ -24,7 +25,7 @@ Description:Approximation of attenuation of light due to occlusion.
 [/VAR]
 
 [EXCERPT:AmbientOcclusion_Desc]
-**Screen Space Ambient Occlusion** is an effect that approximates the attenuation of light due to occlusion. This is generally best used as a subtle effect, in addition to standard global illumination, that darkens corners, crevices, or other features to create a more natural, realistic look.
+**Screen Space Ambient Occlusion** (SSAO) is an effect that approximates the attenuation of light due to occlusion. This is generally best used as a subtle effect, in addition to standard global illumination, that darkens corners, crevices, or other features to create a more natural, realistic look.
 [/EXCERPT:AmbientOcclusion_Desc]
 
 [REGION:imagetable]
@@ -33,7 +34,7 @@ Description:Approximation of attenuation of light due to occlusion.
 | Scene without Ambient Occlusion | Ambient Occlusion Only | Scene with Ambient Occlusion |
 [/REGION]
 
-This effect is achieved using a SSAO (Screen Space Ambient Occlusion) method, meaning the AO amount is computed in a deferred pass in screen space. 
+This effect is achieved using a SSAO method, meaning the AO amount is computed in a deferred pass in screen space. 
 Our method makes use of the depth buffer and the normal from the GBuffer (see Deferred Shading). This means normal map details will affect the results. In UE4, the effect is applied to the AmbientCubemap. It does not affect the direct lighting because here we use shadow mapping.
 
 You can look at the AO value directly by using the "Visualize GBuffer" view mode (see View Modes) or by using the show flag "Visualize Ambient Occlusion".
