#UE4doc #docs - Added engine version tag.

diff --git a/Engine/Documentation/Source/Engine/Rendering/LightingAndShadows/IndirectLightingCache/IndirectLightingCache.INT.udn b/Engine/Documentation/Source/Engine/Rendering/LightingAndShadows/IndirectLightingCache/IndirectLightingCache.INT.udn
--- a/Engine/Documentation/Source/Engine/Rendering/LightingAndShadows/IndirectLightingCache/IndirectLightingCache.INT.udn
+++ b/Engine/Documentation/Source/Engine/Rendering/LightingAndShadows/IndirectLightingCache/IndirectLightingCache.INT.udn
@@ -2,6 +2,7 @@ Availability:Public
 Title:Indirect Lighting Cache
 Crumbs: %ROOT%, Engine, Engine/Rendering/LightingAndShadows
 Description:Cached indirect lighting samples used for global illumination of dynamic objects and unbuilt scene preview.
+Versinon: 4.9
 
 [VAR:Topic]
 [OBJECT:Topic]
@@ -28,7 +29,7 @@ Description:Cached indirect lighting samples used for global illumination of dyn
 [/REGION]
 
 [EXCERPT:LightCache001]
-[Lightmass](Engine/Rendering/LightingAndShadows/Lightmass) generates lightmaps for indirect lighting on static objects, but dynamic objects like characters 
+[**Lightmass**](Engine/Rendering/LightingAndShadows/Lightmass) generates lightmaps for indirect lighting on static objects, but dynamic objects like characters 
 need a method of receiving indirect lighting as well. This is solved by the **Indirect Lighting Cache**, which 
 uses samples generated by Lightmass at build-time to calculate the indirect lighting for dynamic objects at runtime.
 [/EXCERPT:LightCache001]
