#UE4doc #docs - Updated images and added engine version tag.

diff --git a/Engine/Documentation/Source/Engine/Rendering/LightingAndShadows/Lightmass/Lightmass.INT.udn b/Engine/Documentation/Source/Engine/Rendering/LightingAndShadows/Lightmass/Lightmass.INT.udn
--- a/Engine/Documentation/Source/Engine/Rendering/LightingAndShadows/Lightmass/Lightmass.INT.udn
+++ b/Engine/Documentation/Source/Engine/Rendering/LightingAndShadows/Lightmass/Lightmass.INT.udn
@@ -2,6 +2,7 @@ Availability:Public
 Title: Lightmass Global Illumination
 Crumbs: %ROOT%, Engine, Engine/Rendering/LightingAndShadows
 Description:A run down of all the features and settings of Lightmass.
+Version: 4.9
 
 [VAR:Topic]
 [OBJECT:Topic]
@@ -36,7 +37,7 @@ Description:A run down of all the features and settings of Lightmass.
 **Lightmass** creates lightmaps with complex light interactions like area shadowing and diffuse interreflection.  It is used to precompute portions of the lighting contribution of lights with stationary and static mobility. 
 [/EXCERPT:LightMass_002]
 
-Communication between the editor and Lightmass is handled by the Swarm Agent, which manages the lighting build locally and can also distribute the lighting build to remote machines. The Swarm Agent, which opens minimized by default, also tracks lighting build progress and keeps you up to date with which machines are working for you, what they are working on, and how many threads each one is using.  
+Communication between the editor and Lightmass is handled by the **Swarm Agent**, which manages the lighting build locally and can also distribute the lighting build to remote machines. The Swarm Agent, which opens minimized by default, also tracks lighting build progress and keeps you up to date with which machines are working for you, what they are working on, and how many threads each one is using.  
 
 A sample image of the Swarm Agent is pictured below (the bar near the bottom shows how much of the build is complete).
 
@@ -65,7 +66,7 @@ And here is the scene with four diffuse bounces combined.  Simulating global ill
 ![3FourBounces.jpg](3FourBounces.jpg)(w:720)
 
 
-Bounced lighting picks up the diffuse color of the underlying material, as you can see below.  This is where the term Color Bleeding comes from.  Color Bleeding is most noticeable with highly saturated colors.  You can exaggerate the effect by raising DiffuseBoost on the Primitive, Material, or Level.
+Bounced lighting picks up the diffuse color of the underlying material, as you can see below.  This is where the term Color Bleeding comes from.  Color Bleeding is most noticeable with highly saturated colors.  You can exaggerate the effect by raising _DiffuseBoost_ on the Primitive, Material, or Level.
 
 [REGION:imagetable]
 | ![diffusebounce_T.JPG](diffusebounce_T.JPG)(h:445) | ![diffusebounce_L.JPG](diffusebounce_L.JPG)(h:445) |
@@ -76,7 +77,7 @@ Bounced lighting picks up the diffuse color of the underlying material, as you c
 
 ### Character lighting
 
-Lightmass places samples in a uniform 3d grid inside the [Lightmass Importance Volume](#LightmassImportanceVolume) at a low resolution, and also on top of upward facing surfaces that characters are likely to walk on, at a higher resolution.  Each lighting sample captures indirect lighting from all directions, but not direct lighting.
+Lightmass places samples in a uniform 3d grid inside the [**Lightmass Importance Volume**](#LightmassImportanceVolume) at a low resolution, and also on top of upward facing surfaces that characters are likely to walk on, at a higher resolution.  Each lighting sample captures indirect lighting from all directions, but not direct lighting.
 
 The first image is a debug visualization of lighting samples placed above the floor, the second image is the same scene in Lit mode.  Notice how the samples above the red tapestry picked up the red bounced light.  The samples are visualized as a single color, but they really capture lighting from all directions.
 
@@ -89,7 +90,7 @@ The first image is a debug visualization of lighting samples placed above the fl
 
 The Indirect Lighting Cache then interpolates the indirect lighting for a movable object using these lighting samples.  The indirect lighting affects the color of the light environment shadow, but not the direction.  Use Show->Visualize->Volume Lighting Samples to preview these in your viewport.
 
-For more information, see the [Indirect Lighting Cache](Engine/Rendering/LightingAndShadows/IndirectLightingCache) documentation.
+For more information, see the [**Indirect Lighting Cache**](Engine/Rendering/LightingAndShadows/IndirectLightingCache) documentation.
 
 #### Limitations
 
@@ -101,14 +102,14 @@ For more information, see the [Indirect Lighting Cache](Engine/Rendering/Lightin
 
 Lightmass calculates detailed indirect shadows automatically, but it can be useful to exaggerate indirect shadows for artistic purposes or to enhance the perception of proximity in a scene.  
 
-Ambient occlusion is the indirect shadowing you would get from a uniformly bright upper hemisphere, like an overcast sky.  Lightmass supports calculating ambient occlusion, applying it to direct and indirect illumination, and then baking it into lightmaps.  Ambient occlusion is enabled by default, and can be disabled by unchecking the checkbox Use Ambient Occlusion in the Lightmass Settings of Lightmass under the World Settings.
+**Ambient occlusion** is the indirect shadowing you would get from a uniformly bright upper hemisphere, like an overcast sky.  Lightmass supports calculating ambient occlusion, applying it to direct and indirect illumination, and then baking it into lightmaps.  Ambient occlusion is enabled by default, and can be disabled by unchecking the checkbox Use Ambient Occlusion in the Lightmass Settings of Lightmass under the World Settings.
 
 In the first image is a scene with indirect lighting but no ambient occlusion.  In the second image is the same scene with ambient occlusion applied to both the direct and indirect lighting, note the darkening where objects come together.
 
 [REGION:imagetable]
 ![9NoAO.jpg](9NoAO.jpg)(w:465) | ![9WithAO.jpg](9WithAO.jpg)(w:465) |
 | ------ | ------ |
-| | |
+|No Ambient Occlusion | Ambient Occlusion |
 [/REGION]
 
 #### Ambient Occlusion Settings
@@ -149,11 +150,11 @@ Lightmass takes the opacity mask of BLEND_Masked Materials into account when cal
 
 ### Bent Normal sky occlusion
 
-Lightmass generates directional occlusion in the form of a bent normal whenever a [Sky Light](Engine/Rendering/LightingAndShadows/LightTypes/SkyLight) with Stationary Mobility is enabled.  
+Lightmass generates directional occlusion in the form of a bent normal whenever a [**Sky Light**](Engine/Rendering/LightingAndShadows/LightTypes/SkyLight) with Stationary Mobility is enabled.  
 
 ### Distance field shadowmaps
 
-Lightmass computes distance field shadowmaps for [Stationary Lights](Engine/Rendering/LightingAndShadows/LightMobility/StationaryLights).  Distance field shadow maps maintain their curvy shape very well even under lower resolutions; however they do not support area shadowing or translucent shadows.
+Lightmass computes distance field shadowmaps for [**Stationary Lights**](Engine/Rendering/LightingAndShadows/LightMobility/StationaryLights).  Distance field shadow maps maintain their curvy shape very well even under lower resolutions; however they do not support area shadowing or translucent shadows.
 
 ## Features for Static Lights Only
 
