#UE4 Docs: added image comparison slider

diff --git a/Engine/Documentation/Source/Engine/Rendering/PostProcessEffects/Bloom/Bloom.INT.udn b/Engine/Documentation/Source/Engine/Rendering/PostProcessEffects/Bloom/Bloom.INT.udn
--- a/Engine/Documentation/Source/Engine/Rendering/PostProcessEffects/Bloom/Bloom.INT.udn
+++ b/Engine/Documentation/Source/Engine/Rendering/PostProcessEffects/Bloom/Bloom.INT.udn
@@ -32,11 +32,14 @@ We combine the blurs differently to get more control and higher quality. For bes
 [/EXCERPT:Intensity]
 
 [EXCERPT:Threshold]
-[REGION:imagetable]
-| ![Bloom Threshold - 0](bloom_threshold_0_small.png)(w:380) | ![Bloom Threshold - 10](bloom_threshold_10_small.png)(w:380) |
-| ------ | ------ |
-| 0.0 | 10.0 |
-[/REGION]
+[OBJECT:ComparisonSlider]
+	[PARAM:before]
+	![Threshold - 0.0](bloom_threshold_0_small.png)
+	[/PARAM]
+	[PARAM:after]
+	![Threshold - 10.0](bloom_threshold_10_small.png)
+	[/PARAM]
+[/OBJECT]
 [/EXCERPT:Threshold]
 
 [EXCERPT:Size]
