#UE4doc #docs
Completed first draft of 'Setting up your Project'.
Availability: Docs

diff --git a/Engine/Documentation/Source/Programming/Tutorials/FirstPersonShooter/1/1/FirstPersonShooter_1_1.INT.udn b/Engine/Documentation/Source/Programming/Tutorials/FirstPersonShooter/1/1/FirstPersonShooter_1_1.INT.udn
--- a/Engine/Documentation/Source/Programming/Tutorials/FirstPersonShooter/1/1/FirstPersonShooter_1_1.INT.udn
+++ b/Engine/Documentation/Source/Programming/Tutorials/FirstPersonShooter/1/1/FirstPersonShooter_1_1.INT.udn
@@ -22,6 +22,56 @@ Version: 4.9
 
 %Steps%
 
-<!--TODO: Add tutorial steps here-->
+[REGION:note]
+You should have a copy of **Visual Studio 2013** for Desktop (Preferably Community or Professional Editions) installed before starting this tutorial.
+[/REGION]
+
+During this step, you will create a starting point for your First Person Shooter (FPS) game by using Unreal Engine 4's Project Browser.
+
+1.	Open **Unreal Engine** from the Launcher. The [](Engine/Basics/Projects/Browser) will appear.
+
+1.	Click on the **New Project** tab and then select the **Blueprint** tab.
+
+1.	Select the **Blank** project.
+
+1.	Name the project "FPSProject". Make sure **No Starter Content** is set.
+	
+	![](UnrealPrjBrowser.png)
+
+	[REGION:warning]
+	Some of the code samples in this tutorial will need to be updated if you name your project differently.
+	[/REGION]
+
+1.	Click the **Create Project** button.
+
+1.	Your project is now open in Unreal Editor. Click the **Play** button in the Level Editor to move around the level. Use the WASD keys to move and use the mouse to aim your camera.
+
+	![](ClickPlayButton.png)
+
+1.	Click the **Stop** button in the Level Editor to exit Play in Editor (PIE) mode.
+
+	![](ClickStopButton.png)
+
+1.	Create a Maps folder within the Content folder.
+
+	![](MapsFolder.png)
+
+1.	In the **File** menu, select **Save as...** to save your level as "FPSMap" within the Maps folder.
+
+	![](SaveLevelAs.png)
+
+1.	In the **Edit** menu, click on **Project Settings**.
+
+1.	Under the **Project** heading on the left side of the **Project Settings** tab, click on **Maps & Modes**.
+
+1.	Select **FPSMap** as the **Editor Startup Map** using the dropdown menu.
+
+	![](EditorDefaultMap.png)
+
+	[REGION:note]
+	From now on, the editor will automatically load FPSMap as the default map.
+	[/REGION]
+
+1.	Close the **Project Settings** menu.
 
 %Steps%
\ No newline at end of file
