#UE4doc #docs
Wrote first draft of the 'Using the C++ Class Wizard' step.
Availability: Docs

diff --git a/Engine/Documentation/Source/Programming/Tutorials/FirstPersonShooter/1/2/FirstPersonShooter_1_2.INT.udn b/Engine/Documentation/Source/Programming/Tutorials/FirstPersonShooter/1/2/FirstPersonShooter_1_2.INT.udn
--- a/Engine/Documentation/Source/Programming/Tutorials/FirstPersonShooter/1/2/FirstPersonShooter_1_2.INT.udn
+++ b/Engine/Documentation/Source/Programming/Tutorials/FirstPersonShooter/1/2/FirstPersonShooter_1_2.INT.udn
@@ -23,6 +23,20 @@ Version: 4.9
 
 %Steps%
 
-<!--TODO: Add tutorial steps here-->
+During this step, you're going to create a Game Mode. Game Modes define your game's rules and win conditions. The Game Mode will also set default classes that will be used for some basic gameplay framework types, including Pawn, PlayerController, and HUD. As a final note, you need to create a Game Mode prior to setting up your FPS Character.
+
+1.	In the File menu, select **New C++ Class...** to choose your new parent class.
+
+1.	The **Choose Parent Class** menu will open. Scroll down, select **Game Mode** as the parent class and click **Next**.
+
+	![](AddGameModeParentClass.png)
+
+1.	Name the new class "FPSGameMode," then click **Create**.
+
+	![](CreateFPSGameMode.png)
+
+	[REGION:note]
+	Now that you've created our first C++ class, you can switch to **Visual Studio** to add code to your newly created class. `FPSGameMode.h` and `FPSGameMode.cpp` will be opened for you and **Unreal Engine** will automatically compile and reload the code for your new class.
+	[/REGION]
 
 %Steps%
\ No newline at end of file
