#UE4doc #docs
Completed first draft of the 'Adding Log Messaging' step.
Availability: Docs

diff --git a/Engine/Documentation/Source/Programming/Tutorials/FirstPersonShooter/1/3/FirstPersonShooter_1_3.INT.udn b/Engine/Documentation/Source/Programming/Tutorials/FirstPersonShooter/1/3/FirstPersonShooter_1_3.INT.udn
--- a/Engine/Documentation/Source/Programming/Tutorials/FirstPersonShooter/1/3/FirstPersonShooter_1_3.INT.udn
+++ b/Engine/Documentation/Source/Programming/Tutorials/FirstPersonShooter/1/3/FirstPersonShooter_1_3.INT.udn
@@ -23,6 +23,119 @@ Version: 4.9
 
 %Steps%
 
-<!--TODO: Add tutorial steps here-->
+You'll add a log message to your newly created FPSGameMode in this step. Log messages are useful for verifying and debugging your code. In this case, you'll use your log message to verify that you're actually using your new Game Mode.
+
+## FPSGameMode Interface File
+
+1.	In the **Solution Explorer**, expand **FPSProject > Source > FPSProject**.
+
+	![](InterfaceFile.png)
+
+1.	Double-click `FPSGameMode.h` to open the interface file for your newly created **FPSGameMode** class.
+
+1.	Your class declaration should like the following:
+
+		UCLASS()
+		class FPSPROJECT_API AFPSGameMode : public AGameMode
+		{
+			GENERATED_BODY()
+		};
+
+1.	Add the following constructor declaration under the `GENERATED_BODY()` macro:
+
+		AFPSGameMode(const FObjectInitializer& ObjectInitializer);
+
+1.	Add the following function declaration under the `AFPSGameMode` constructor declaration:
+
+		virtual void StartPlay() override;
+
+	[REGION:note]
+	This function declaration will allow you to override the StartPlay() function (inherited from the `AActor` class) so that you can print a log message to the screen when gameplay begins.
+	[/REGION]
+
+1.	Your class interface should look like the following:
+
+		// Fill out your copyright notice in the Description page of Project Settings.
+		
+		#pragma once
+		
+		#include "GameFramework/GameMode.h"
+		#include "FPSGameMode.generated.h"
+		
+		/**
+		 * 
+		 */
+		UCLASS()
+		class FPSPROJECT_API AFPSGameMode : public AGameMode
+		{
+			GENERATED_BODY()
+		
+			AFPSGameMode(const FObjectInitializer& ObjectInitializer);
+		
+			virtual void StartPlay() override;
+		};
+
+1.	Save your interface file in Visual Studio.
+
+## FPSGameMode Implementation File
+
+1.	In the **Solution Explorer**, locate `FPSGameMode.cpp` (also located in **FPSProject > Source > FPSProject**).
+
+	![](ImplementationFile.png)
+
+1.	Double-click `FPSGameMode.cpp` to open the implementation file for your **FPSGameMode** class.
+
+1.	Add the following constructor definition to `FPSGameMode.cpp`:
+
+		AFPSGameMode::AFPSGameMode(const FObjectInitializer& ObjectInitializer) : Super(ObjectInitializer)
+		{
+		}
+
+	[REGION:note]
+	The FPSGameMode is AFPSGameMode because it derives from a class that ultimately derives from the base Actor class.
+	[/REGION]
+
+1.	Add the following lines of code for the `StartPlay()` function below the `AFPSGameMode` constructor:
+
+		void AFPSGameMode::StartPlay()
+		{
+			Super::StartPlay();
+		
+			StartMatch();
+		
+			if (GEngine != nullptr)
+			{
+				GEngine->AddOnScreenDebugMessage(-1, 5.f, FColor::Yellow, TEXT("We are using FPSGameMode."));
+			}
+		}//StartPlay()
+
+	[REGION:note]
+	`StartPlay()` will print "Hello World, this is FPSGameMode!" to your screen in yellow text when gameplay begins.
+	[/REGION]
+
+1.	Your class implementation should look like the following:
+
+		// Fill out your copyright notice in the Description page of Project Settings.
+		
+		#include "FPSProject.h"
+		#include "FPSGameMode.h"
+		
+		AFPSGameMode::AFPSGameMode(const FObjectInitializer& ObjectInitializer) : Super(ObjectInitializer)
+		{
+		}//AFPSGameMode Constructor
+		
+		void AFPSGameMode::StartPlay()
+		{
+			Super::StartPlay();
+		
+			StartMatch();
+		
+			if (GEngine != nullptr)
+			{
+				GEngine->AddOnScreenDebugMessage(-1, 5.f, FColor::Yellow, TEXT("We are using FPSGameMode."));
+			}
+		}//StartPlay()
+
+1.	Save your implementation file in Visual Studio.
 
 %Steps%
\ No newline at end of file
