#UE4doc #docs
Added image goal for the section.

diff --git a/Engine/Documentation/Source/Programming/Tutorials/FirstPersonShooter/1/FirstPersonShooter_1.INT.udn b/Engine/Documentation/Source/Programming/Tutorials/FirstPersonShooter/1/FirstPersonShooter_1.INT.udn
--- a/Engine/Documentation/Source/Programming/Tutorials/FirstPersonShooter/1/FirstPersonShooter_1.INT.udn
+++ b/Engine/Documentation/Source/Programming/Tutorials/FirstPersonShooter/1/FirstPersonShooter_1.INT.udn
@@ -6,7 +6,8 @@ prereq: Programming/Tutorials/FirstPersonShooter
 SkillLevel: Beginner
 Version: 4.9
 
-<!--TODO: Insert Cover Image-->
+![](SectionResult.png)
+_This is what you'll see at the end of this section._
 
 ## Goals
 
