#UE4doc #docs
Completed first draft for 'Compile your First C++ Class' section.
Availability: Docs

diff --git a/Engine/Documentation/Source/Programming/Tutorials/FirstPersonShooter/1/4/FirstPersonShooter_1_4.INT.udn b/Engine/Documentation/Source/Programming/Tutorials/FirstPersonShooter/1/4/FirstPersonShooter_1_4.INT.udn
--- a/Engine/Documentation/Source/Programming/Tutorials/FirstPersonShooter/1/4/FirstPersonShooter_1_4.INT.udn
+++ b/Engine/Documentation/Source/Programming/Tutorials/FirstPersonShooter/1/4/FirstPersonShooter_1_4.INT.udn
@@ -23,6 +23,16 @@ Version: 4.9
 
 %Steps%
 
-<!--TODO: Add tutorial steps here-->
+You will now compile your project in this step so that you can see your code changes reflected in-game. 
+
+1.	Locate **FPSProject** in the **Solution Explorer**.
+
+1.	Right-click on **FPSProject** and select **Build** to compile your project.
+
+	![](CompileProject.png)
+
+1.	After the build finishes, open Unreal Editor and verify that your newly compiled **FPSGameMode** class is visible in the **Content Browser**.
+
+	![](CompiledFPSGameModeObj.png)
 
 %Steps%
\ No newline at end of file
