#UE4doc #docs
Wrote first draft for 'Set the Default Game Mode' section.
Availability: Docs

diff --git a/Engine/Documentation/Source/Programming/Tutorials/FirstPersonShooter/1/5/FirstPersonShooter_1_5.INT.udn b/Engine/Documentation/Source/Programming/Tutorials/FirstPersonShooter/1/5/FirstPersonShooter_1_5.INT.udn
--- a/Engine/Documentation/Source/Programming/Tutorials/FirstPersonShooter/1/5/FirstPersonShooter_1_5.INT.udn
+++ b/Engine/Documentation/Source/Programming/Tutorials/FirstPersonShooter/1/5/FirstPersonShooter_1_5.INT.udn
@@ -23,6 +23,32 @@ Version: 4.9
 
 %Steps%
 
-<!--TODO: Add tutorial steps here-->
+Now that we've successfully compiled our newly added Game Mode, we'll need to set **FPSProject** to use **FPSGameMode** as the default Game Mode in this step.
+
+1.	In the **Edit** menu, click on **Project Settings**.
+
+1.	Under the **Project** heading on the left side of the **Project Settings** tab, click on **Maps & Modes**.
+
+1.	Select **FPSGameMode** in the Default GameMode dropdown.
+
+	![](SettingFPSGameMode.png)
+
+1.	Close the **Project Settings** menu.
+
+1.	Click the **Play** button in the Level Editor Toolbar. "We are using FPSGameMode." should be displayed in the upper left corner of the viewport.
+
+	![](SectionResult.png)
+
+1.	Click the **Stop** button in the Level Editor to exit Play in Editor (PIE) mode.
+
+Congratulations! You've learned how to:
+
+&#x2713; Set up a New Project  
+&#x2713; Use the C++ Class Wizard  
+&#x2713; Add Log Messaging to your Project  
+&#x2713; Compile your First C++ Class  
+&#x2713; Set the Default Game Mode
+
+You're now ready to learn how to implement your character.
 
 %Steps%
\ No newline at end of file
