#UE4doc #docs
Completed first draft of 'Making a New Character' section.
Availability: Docs

diff --git a/Engine/Documentation/Source/Programming/Tutorials/FirstPersonShooter/2/1/FirstPersonShooter_2_1.INT.udn b/Engine/Documentation/Source/Programming/Tutorials/FirstPersonShooter/2/1/FirstPersonShooter_2_1.INT.udn
--- a/Engine/Documentation/Source/Programming/Tutorials/FirstPersonShooter/2/1/FirstPersonShooter_2_1.INT.udn
+++ b/Engine/Documentation/Source/Programming/Tutorials/FirstPersonShooter/2/1/FirstPersonShooter_2_1.INT.udn
@@ -22,6 +22,148 @@ Version: 4.9
 
 %Steps%
 
-<!--TODO: Add tutorial steps here-->
+During this step, you're going to make a new character using the engine's Character base class. The Character class (derived from the Pawn class) has built-in functionality for bipedal movement, such as walking, running, and jumping.
+
+[REGION:note]
+Although you can manually add \*.h and \*.cpp files to your Visual Studio solution, it's good practice to use the C++ Class Wizard to add new classes to your project. By using the C++ Class Wizard, the engine creates header and source templates that set-up Unreal specific macros for you.
+[/REGION]
+
+## Adding a Character Class
+
+1.	In the File menu, select **New C++ Class...** to choose your new parent class.
+
+1.	The **Choose Parent Class** menu will open. Scroll down, select **Character** as the parent class and click **Next**.
+
+	![](ChooseCharacterClass.png)
+
+1.	Name the new class "FPSCharacter," then click **Create**.
+
+	![](MakeFPSCharacterClass.png)
+
+	[REGION:note]
+	Now that you've created your FPSCharacter class, you can switch to **Visual Studio** to add code to your newly created class. `FPSCharacter.h` and `FPSCharacter.cpp` will be opened for you and **Unreal Engine** will automatically compile and reload the code for your new class.
+	[/REGION]
+
+## Editing the Game Mode Class
+
+First, you're going to edit your **GameMode** so that FPSCharacter is the defaul Pawn used when starting gameplay.
+
+1.	In the **Solution Explorer**, expand **FPSProject > Source > FPSProject**.
+
+	![](ExpandedSolutionExplorer.png)
+
+1.	Double-click `FPSGameMode.cpp` to open the implementation file for your **FPSGameMode** class.
+
+1.	First, include the `FPSCharacter` header at the top of `FPSGameMode.cpp`:
+
+		#include "FPSCharacter.h"
+
+1.	Now, add the following line in your `AFPSGameMode` constructor:
+
+		DefaultPawnClass = AFPSCharacter::StaticClass();
+
+1.	`FPSGameMode.cpp` should now look like the following:
+
+		// Fill out your copyright notice in the Description page of Project Settings.
+		
+		#include "FPSProject.h"
+		#include "FPSGameMode.h"
+		#include "FPSCharacter.h"
+		
+		AFPSGameMode::AFPSGameMode(const FObjectInitializer& ObjectInitializer) : Super(ObjectInitializer)
+		{
+			DefaultPawnClass = AFPSCharacter::StaticClass();
+		}//AFPSGameMode Constructor
+		
+		void AFPSGameMode::StartPlay()
+		{
+			Super::StartPlay();
+		
+			StartMatch();
+		
+			if (GEngine != nullptr)
+			{
+				GEngine->AddOnScreenDebugMessage(-1, 5.f, FColor::Yellow, TEXT("We are using FPSGameMode."));
+			}
+		}//StartPlay()
+
+1.	Save the `FPSGameMode.cpp` implementation file in Visual Studio.
+
+## Verifying the Character Class
+
+1.	In the **Solution Explorer**, expand **FPSProject > Source > FPSProject**.
+
+	![](ExpandedSolutionExplorer.png)
+
+1.	Double-click `FPSCharacter.cpp` to open the implementation file for your **FPSCharacter** class.
+
+1.	Add the following lines of code in the `BeginPlay()` function to verify that `FPSCharacter` class is being used:
+
+		if(GEngine != nullptr)
+		{
+			GEngine->AddOnScreenDebugMessage(-1, 5.f, FColor::White, TEXT("We are using FPSCharacter."));
+		}
+
+1.	`FPSCharacter.cpp` should now look like the following:
+
+		// Fill out your copyright notice in the Description page of Project Settings.
+		
+		#include "FPSProject.h"
+		#include "FPSCharacter.h"
+		
+		// Sets default values
+		AFPSCharacter::AFPSCharacter()
+		{
+		 	// Set this character to call Tick() every frame.  You can turn this off to improve performance if you don't need it.
+			PrimaryActorTick.bCanEverTick = true;
+		
+		}
+		
+		// Called when the game starts or when spawned
+		void AFPSCharacter::BeginPlay()
+		{
+			Super::BeginPlay();
+			
+			if (GEngine != nullptr)
+			{
+				GEngine->AddOnScreenDebugMessage(-1, 5.f, FColor::White, TEXT("We are using FPSCharacter."));
+			}
+		}
+		
+		// Called every frame
+		void AFPSCharacter::Tick( float DeltaTime )
+		{
+			Super::Tick( DeltaTime );
+		
+		}
+		
+		// Called to bind functionality to input
+		void AFPSCharacter::SetupPlayerInputComponent(class UInputComponent* InputComponent)
+		{
+			Super::SetupPlayerInputComponent(InputComponent);
+		
+		}
+
+1.	Save the `FPSCharacter` implementation file in Visual Studio.
+
+1.	Locate **FPSProject** in the **Solution Explorer**.
+
+1.	Right-click on **FPSProject** and select **Build** to compile your project.
+
+	![](BuildProject.png)
+
+1.	After the build finishes, open Unreal Editor to verify that your newly compiled **FPSCharacter** class is visible in the **Content Browser**.
+
+	![](FPSCharacterContentBrowser.png)
+
+1.	Click the **Play** button in the Level Editor Toolbar. "We are using FPSCharacter." should now be displayed below "We are using FPSGameMode." in the upper left corner of the viewport.
+
+	![](VerifyingFPSCharacterResult.png)
+
+	[REGION:note]
+	Your new Character doesn't have any movement controls yet, so you won't be able to move around in the level. If you're unable to move, you are using FPSCharacter as your Pawn correctly!
+	[/REGION]
+
+1.	Click the **Stop** button in the Level Editor to exit Play in Editor (PIE) mode.
 
 %Steps%
\ No newline at end of file
