#UE4doc #docs
Completed first draft of 'Setting up Axis Mapping' section.
Availability: Docs

diff --git a/Engine/Documentation/Source/Programming/Tutorials/FirstPersonShooter/2/2/FirstPersonShooter_2_2.INT.udn b/Engine/Documentation/Source/Programming/Tutorials/FirstPersonShooter/2/2/FirstPersonShooter_2_2.INT.udn
--- a/Engine/Documentation/Source/Programming/Tutorials/FirstPersonShooter/2/2/FirstPersonShooter_2_2.INT.udn
+++ b/Engine/Documentation/Source/Programming/Tutorials/FirstPersonShooter/2/2/FirstPersonShooter_2_2.INT.udn
@@ -23,6 +23,64 @@ Version: 4.9
 
 %Steps%
 
-<!--TODO: Add tutorial steps here-->
+During this step, you're going to get your new Character to move around the map by setting up the input axis mappings for the WASD keys.
+
+[REGION:note]
+In general, Axis Mappings allow you to map keyboard, mouse, and controller inputs to a "friendly name" that can later be bound to game behavior (such as movement). Axis Mappings are continuously polled, allowing for smooth movement transitions and smooth game behavior. Hardware axes (such as controller joysticks) provide degrees of input, rather than discrete input (1 - pressed vs. 0 - not pressed). While controller joystick input methods are effective at providing scalable movement input, Axis Mappings can also map common movement keys, like WASD to continuously-polled game behavior.
+[/REGION]
+
+You will now set-up axis mappings for the W, A, S, and D keys.
+
+## MoveForward Axis Mapping
+
+1.	In the **Edit** menu, click on **Project Settings**.
+
+1.	Under the **Engine** heading on the left side of the **Project Settings** tab, click on **Input**.
+
+1.	Under **Bindings**, click on the plus sign next to **Axis Mappings**.
+
+1.	Click on the arrow to the left of **Axis Mappings**.
+
+1.	Type "MoveForward" into the text field that appears, then click on the arrow to the left of the text box to expand the axis binding options.
+
+1.	In the dropdown menu, select **W** from the **Keyboard** dropdown list.
+
+1.	Your input settings should now look like the following:
+
+	![](MoveForwardAxisMap_W.png)
+
+1.	Click on the plus sign next to **MoveForward**.
+
+1.	In the second dropdown menu, select **S** from the **Keyboard** dropdown list.
+
+1.	Type "-1" in the **Scale** field next to **S**.
+
+1.	Your input settings should now look like the following:
+
+	![](MoveForwardAxisMap_S.png)
+
+## MoveRight Axis Mapping
+
+1.	Under **Bindings**, click on the plus sign next to **Axis Mappings**.
+
+1.	Type "MoveRight" into the text field that appears, then click on the arrow to the left of the text box to expand the axis binding options.
+
+1.	In the dropdown menu, select **D** from the **Keyboard** dropdown list.
+
+1.	Your input settings should now look like the following:
+
+	![](MoveRightAxisMap_D.png)
+
+1.	Click on the plus sign next to **MoveRight**.
+
+1.	In the second dropdown menu, select **A** from the **Keyboard** dropdown list.
+
+1.	Type "-1" in the **Scale** field next to **A**.
+
+1.	Your input settings should now look like the following:
+
+	![](MoveRightAxisMap_A.png)
+
+1.	Close the **Project Settings** menu.
 
 %Steps%
\ No newline at end of file
