#UE4doc #docs
Completed first draft of 'Implementing Character Movement Functions' section.
Availability: Docs

diff --git a/Engine/Documentation/Source/Programming/Tutorials/FirstPersonShooter/2/3/FirstPersonShooter_2_3.INT.udn b/Engine/Documentation/Source/Programming/Tutorials/FirstPersonShooter/2/3/FirstPersonShooter_2_3.INT.udn
--- a/Engine/Documentation/Source/Programming/Tutorials/FirstPersonShooter/2/3/FirstPersonShooter_2_3.INT.udn
+++ b/Engine/Documentation/Source/Programming/Tutorials/FirstPersonShooter/2/3/FirstPersonShooter_2_3.INT.udn
@@ -23,6 +23,210 @@ Version: 4.9
 
 %Steps%
 
-<!--TODO: Add tutorial steps here-->
+In this step, you're going to set-up your Player Input Component and implement the following functions in your `FPSCharacter` class:
+
+*	`MoveForward`
+*	`MoveRight`
+
+Now that you've set-up Axis mappings for your `FPSCharacter`, you can switch to your project in Visual Studio.
+
+## Movement Function Interfaces
+
+1.	In `FPSCharacter.h`, add the following function declarations under `SetupPlayerInputComponent`.
+
+		// handles moving forward and backward
+		UFUNCTION()
+		void MoveForward(float Value);
+		
+		// handles moving right and left
+		UFUNCTION()
+		void MoveRight(float Value);
+
+	[REGION:note]
+	The `UFUNCTION` macro is used above each of these functions. By itself, the `UFUNCTION` macro makes the engine aware of these functions so that they can be included in serialization, optimization, and other engine functionality.
+	[/REGION]
+
+1.	`FPSCharacter.h` should now look like the following:
+
+		// Fill out your copyright notice in the Description page of Project Settings.
+		
+		#pragma once
+		
+		#include "GameFramework/Character.h"
+		#include "FPSCharacter.generated.h"
+		
+		UCLASS()
+		class FPSPROJECT_API AFPSCharacter : public ACharacter
+		{
+			GENERATED_BODY()
+		
+		public:
+			// Sets default values for this character's properties
+			AFPSCharacter();
+		
+			// Called when the game starts or when spawned
+			virtual void BeginPlay() override;
+			
+			// Called every frame
+			virtual void Tick( float DeltaSeconds ) override;
+		
+			// Called to bind functionality to input
+			virtual void SetupPlayerInputComponent(class UInputComponent* InputComponent) override;
+		
+			// handles moving forward and backward
+			UFUNCTION()
+			void MoveForward(float Value);
+		
+			// handles moving right and left
+			UFUNCTION()
+			void MoveRight(float Value);
+			
+		};
+
+## Movement Function Implementations
+
+In a typical FPS control scheme, a character's movement axes are relative to the camera. "Forward" movement means "the direction that the camera is pointing" and "right" means "to the right of the direction that the camera is pointing." You are going to use the `PlayerController` to get the character's control rotation. Also, your `MoveForward` function will ignore the pitch component of the control rotation and restrict your input to the XY plane so as to guarantee that your character will move along the ground while you are looking up or down.
+
+1.	In `FPSCharacter.cpp`, add the following lines to the `SetupPlayerInputComponent` function under `Super::SetupPlayerInputComponent(InputComponent)`:
+
+		// Called to bind functionality to input
+		void AFPSCharacter::SetupPlayerInputComponent(class UInputComponent* InputComponent)
+		{
+			Super::SetupPlayerInputComponent(InputComponent);
+			
+			// set-up gameplay key bindings
+			InputComponent->BindAxis("MoveForward", this, &AFPSCharacter::MoveForward);
+			InputComponent->BindAxis("MoveRight", this, &AFPSCharacter::MoveRight);
+		}
+
+	[REGION:note]
+	An `InputComponent` is a component that defines how to handle input data. An `InputComponent` can be attached to an actor that wants to receive input.
+	[/REGION]
+
+1.	Add the following `MoveForward` function definition in `FPSCharacter.cpp`.
+
+		void AFPSCharacter::MoveForward(float Value)
+		{
+			if (Controller != nullptr && Value != 0.f)
+			{
+				//find out which way is forward
+				FRotator Rotation = Controller->GetControlRotation();
+				//limit pitch when walking for falling
+				if (GetCharacterMovement()->IsMovingOnGround() || GetCharacterMovement()->IsFalling())
+				{
+					Rotation.Pitch = 0.f;
+				}
+				//add movement in that direction
+				const FVector Direction = FRotationMatrix(Rotation).GetScaledAxis(EAxis::X);
+				AddMovementInput(Direction, Value);
+			}
+		}
+
+1.	Add the following `MoveRight` function definition in `FPSCharacter.cpp`.
+
+		//Right movement function definition
+		void AFPSCharacter::MoveRight(float Value)
+		{
+			if (Controller != nullptr && Value != 0.f)
+			{
+				//find our which way is right
+				const FRotator Rotation = Controller->GetControlRotation();
+				const FVector Direction = FRotationMatrix(Rotation).GetScaledAxis(EAxis::Y);
+				//add movement in that direction
+				AddMovementInput(Direction, Value);
+			}
+		}
+
+1.	`FPSCharacter.cpp` should now look like the following:
+
+		// Fill out your copyright notice in the Description page of Project Settings.
+		
+		#include "FPSProject.h"
+		#include "FPSCharacter.h"
+		
+		// Sets default values
+		AFPSCharacter::AFPSCharacter()
+		{
+		 	// Set this character to call Tick() every frame.  You can turn this off to improve performance if you don't need it.
+			PrimaryActorTick.bCanEverTick = true;
+		
+		}
+		
+		// Called when the game starts or when spawned
+		void AFPSCharacter::BeginPlay()
+		{
+			Super::BeginPlay();
+			
+			if (GEngine != nullptr)
+			{
+				GEngine->AddOnScreenDebugMessage(-1, 5.f, FColor::White, TEXT("We are using FPSCharacter."));
+			}
+		}
+		
+		// Called every frame
+		void AFPSCharacter::Tick( float DeltaTime )
+		{
+			Super::Tick( DeltaTime );
+		
+		}
+		
+		// Called to bind functionality to input
+		void AFPSCharacter::SetupPlayerInputComponent(class UInputComponent* InputComponent)
+		{
+			Super::SetupPlayerInputComponent(InputComponent);
+			
+			// set-up gameplay key bindings
+			InputComponent->BindAxis("MoveForward", this, &AFPSCharacter::MoveForward);
+			InputComponent->BindAxis("MoveRight", this, &AFPSCharacter::MoveRight);
+		}
+		
+		// Forward movement function definition
+		void AFPSCharacter::MoveForward(float Value)
+		{
+			if (Controller != nullptr && Value != 0.f)
+			{
+				//find out which way is forward
+				FRotator Rotation = Controller->GetControlRotation();
+				//limit pitch when walking for falling
+				if (GetCharacterMovement()->IsMovingOnGround() || GetCharacterMovement()->IsFalling())
+				{
+					Rotation.Pitch = 0.f;
+				}
+				//add movement in that direction
+				const FVector Direction = FRotationMatrix(Rotation).GetScaledAxis(EAxis::X);
+				AddMovementInput(Direction, Value);
+			}
+		}
+		
+		//Right movement function definition
+		void AFPSCharacter::MoveRight(float Value)
+		{
+			if (Controller != nullptr && Value != 0.f)
+			{
+				//find our which way is right
+				const FRotator Rotation = Controller->GetControlRotation();
+				const FVector Direction = FRotationMatrix(Rotation).GetScaledAxis(EAxis::Y);
+				//add movement in that direction
+				AddMovementInput(Direction, Value);
+			}
+		}
+
+## Testing Character Movement
+
+It's now time to compile and test our newly implemented character movement functions.
+
+1.	Save the `FPSCharacter` interface (\*.h) and implementation (\*.cpp) files in Visual Studio.
+
+1.	Locate **FPSProject** in the **Solution Explorer**.
+
+1.	Right-click on **FPSProject** and select **Build** to compile your project.
+
+	![](BuildProject.png)
+
+1.	After the build finishes, open your **FPSProject** in Unreal Editor.
+
+1.	Click the **Play** button in the Level Editor Toolbar. You should be able to move and strafe throughout the map (although your camera will be fixed in place).
+
+1.	Click the **Stop** button in the Level Editor to exit Play in Editor (PIE) mode.
 
 %Steps%
\ No newline at end of file
