#UE4doc #docs
Wrote first draft of 'Implementing Mouse Camera Control' section.
Availability: Docs

diff --git a/Engine/Documentation/Source/Programming/Tutorials/FirstPersonShooter/2/4/FirstPersonShooter_2_4.INT.udn b/Engine/Documentation/Source/Programming/Tutorials/FirstPersonShooter/2/4/FirstPersonShooter_2_4.INT.udn
--- a/Engine/Documentation/Source/Programming/Tutorials/FirstPersonShooter/2/4/FirstPersonShooter_2_4.INT.udn
+++ b/Engine/Documentation/Source/Programming/Tutorials/FirstPersonShooter/2/4/FirstPersonShooter_2_4.INT.udn
@@ -23,6 +23,89 @@ Version: 4.9
 
 %Steps%
 
-<!--TODO: Add tutorial steps here-->
+During this step, you're going to add the ability for your character to look around and steer with the mouse.
+
+## Turn Axis Mapping
+
+1.	In the **Edit** menu, click on **Project Settings**.
+
+1.	Under the **Engine** heading on the left side of the **Project Settings** tab, click on **Input**.
+
+1.	Under **Bindings**, click on the plus sign next to **Axis Mappings**.
+
+1.	Click on the arrow to the left of **Axis Mappings**.
+
+1.	Type "Turn" into the text field that appears, then click on the arrow to the left of the text box to expand the axis binding options.
+
+1.	In the dropdown menu, select **Mouse X** from the **Mouse** dropdown list.
+
+1.	Your input settings should now look like the following:
+
+	![](TurnAxisMapping_MouseX.png)
+
+## Look-up Axis Mapping
+
+1.	Under **Bindings**, click on the plus sign next to **Axis Mappings**.
+
+1.	Type **LookUp** into the text field that appears, then click on the arrow to the left of the text box to expand the axis binding options.
+
+1.	In the dropdown menu, select **Mouse Y** from the **Mouse** dropdown list.
+
+1.	Type "-1" in the **Scale** field next to **Mouse Y**.
+
+1.	Your input settings should now look like the following:
+
+	![](TurnAxisMapping_MouseY.png)
+
+1.	Close the **Project Settings** menu.
+
+## Implementing Input Handling
+
+Now it's time to add code to handle mouse input for turning and looking up. The `Character` base class defines the two necessary functions for us, namely:
+
+*	`AddYawInput`
+*	`AddPitchInput`
+
+[REGION:note]
+If you want to perform additional processing, such as adding support for sensitivity or axis inversion, you could provide your own functions to adjust the input values before passing them to function; however, in this case, you'll bind your inputs directly to the `AddYawInput` and `AddPitchInput` functions.
+[/REGION]
+
+1.	Add the following lines to the `SetupPlayerInputComponent` in `FPSCharacter.cpp`.
+		// set-up gameplay mouse bindings
+		InputComponent->BindAxis("Turn", this, &AFPSCharacter::AddControllerYawInput);
+		InputComponent->BindAxis("LookUp", this, &AFPSCharacter::AddControllerPitchInput);
+
+1.	The `SetupPlayerInputComponent` function should now look like the following:
+
+		// Called to bind functionality to input
+		void AFPSCharacter::SetupPlayerInputComponent(class UInputComponent* InputComponent)
+		{
+			Super::SetupPlayerInputComponent(InputComponent);
+			
+			// set-up gameplay key bindings
+			InputComponent->BindAxis("MoveForward", this, &AFPSCharacter::MoveForward);
+			InputComponent->BindAxis("MoveRight", this, &AFPSCharacter::MoveRight);
+			// set-up gameplay mouse bindings
+			InputComponent->BindAxis("Turn", this, &AFPSCharacter::AddControllerYawInput);
+			InputComponent->BindAxis("LookUp", this, &AFPSCharacter::AddControllerPitchInput);
+		}
+
+1.	Save the `FPSCharacter` implementation file in Visual Studio.
+
+1.	Locate **FPSProject** in the **Solution Explorer**.
+
+1.	Right-click on **FPSProject** and select **Build** to compile your project.
+
+	![](BuildProject.png)
+
+1.	After the build finishes, open your **FPSProject** in Unreal Editor.
+
+1.	Click the **Play** button in the Level Editor Toolbar. You can now control the camera with your mouse.
+
+1.	Click the **Stop** button in the Level Editor to exit Play in Editor (PIE) mode.
+
+## Testing Mouse Camera Control
+
+
 
 %Steps%
\ No newline at end of file
