#UE4doc #docs
Completed first draft of 'Implemented Character Jumping' section.
Availability: Docs

diff --git a/Engine/Documentation/Source/Programming/Tutorials/FirstPersonShooter/2/5/FirstPersonShooter_2_5.INT.udn b/Engine/Documentation/Source/Programming/Tutorials/FirstPersonShooter/2/5/FirstPersonShooter_2_5.INT.udn
--- a/Engine/Documentation/Source/Programming/Tutorials/FirstPersonShooter/2/5/FirstPersonShooter_2_5.INT.udn
+++ b/Engine/Documentation/Source/Programming/Tutorials/FirstPersonShooter/2/5/FirstPersonShooter_2_5.INT.udn
@@ -23,6 +23,224 @@ Version: 4.9
 
 %Steps%
 
-<!--TODO: Add tutorial steps here-->
+During this step, you're going to add the ability for your character to jump.
+
+[REGION:note]
+In general, Action Mappings deal with inputs for discrete events, they allow you to map inputs to a "friendly name" that can later be bound to event-driven behavior. The end effect is that pressing and/or releasing a key, mouse button, or keypad button directly triggers game behavior.
+[/REGION]
+
+## Jump Action Mapping
+
+1.	In the **Edit** menu, click on **Project Settings**.
+
+1.	Under the **Engine** heading on the left side of the **Project Settings** tab, click on **Input**.
+
+1.	Under **Bindings**, click on the plus sign next to **Action Mappings**.
+
+1.	Click on the arrow to the left of **Action Mappings**.
+
+1.	Type "Jump" into the text field that appears, then click on the arrow to the left of the text box to expand the action binding options.
+
+1.	In the dropdown menu, select **Space Bar** from the **Keyboard** dropdown list.
+
+1.	Your input settings should now look like the following:
+
+	![](JumpActionMapping_SpaceBar.png)
+
+1.	Close the **Project Settings** menu.
+
+## Implement Input Handling
+
+If you look inside the interface file (*.h) for the Character base class, you'll notice that there is built-in support for jumping characters. Character jumping is tied to the `bPressedJump` variable, so all we need to do is set that boolean to `true` when the jump action is pressed, and `false` when the jump action is released. We need the following two functions to accomplish this:
+
+*	`OnStartJump`
+*	`OnStopJump`
+
+Go back to Visual Studio to add code to your FPSCharacter class.
+
+1.	In `FPSCharacter.h`, add the following public function declarations:
+
+		// sets jump flag when key is pressed
+		UFUNCTION()
+		void OnStartJump();
+		
+		// clears jump flag when key is released
+		UFUNCTION()
+		void OnStopJump();
+
+1.	`FPSCharacter.h` should now look like the following:
+
+		// Fill out your copyright notice in the Description page of Project Settings.
+		
+		#pragma once
+		
+		#include "GameFramework/Character.h"
+		#include "FPSCharacter.generated.h"
+		
+		UCLASS()
+		class FPSPROJECT_API AFPSCharacter : public ACharacter
+		{
+			GENERATED_BODY()
+		
+		public:
+			// Sets default values for this character's properties
+			AFPSCharacter();
+		
+			// Called when the game starts or when spawned
+			virtual void BeginPlay() override;
+			
+			// Called every frame
+			virtual void Tick( float DeltaSeconds ) override;
+		
+			// Called to bind functionality to input
+			virtual void SetupPlayerInputComponent(class UInputComponent* InputComponent) override;
+		
+			// handles moving forward and backward
+			UFUNCTION()
+			void MoveForward(float Value);
+		
+			// handles moving right and left
+			UFUNCTION()
+			void MoveRight(float Value);
+		
+			// sets jump flag when key is pressed
+			UFUNCTION()
+			void OnStartJump();
+			
+			// clears jump flag when key is released
+			UFUNCTION()
+			void OnStopJump();
+		};
+
+1.	In `FPSCharacter.cpp`, add the following function definitions:
+
+		//Function setting the jump flag
+		void AFPSCharacter::OnStartJump()
+		{
+			bPressedJump = true;
+		}
+		
+		//Function clearning the jump flag
+		void AFPSCharacter::OnStopJump()
+		{
+			bPressedJump = false;
+		}
+
+1.	Add the following code to `SetupPlayerInputComponent` so as to bind the Jump action to your new functions:
+
+		// set-up jump action bindings
+		InputComponent->BindAction("Jump", IE_Pressed, this, &AFPSCharacter::OnStartJump);
+		InputComponent->BindAction("Jump", IE_Released, this, &AFPSCharacter::OnStopJump);
+
+1.	`FPSCharacter.cpp` should now look like the following:
+
+		// Fill out your copyright notice in the Description page of Project Settings.
+		
+		#include "FPSProject.h"
+		#include "FPSCharacter.h"
+		
+		
+		// Sets default values
+		AFPSCharacter::AFPSCharacter()
+		{
+		 	// Set this character to call Tick() every frame.  You can turn this off to improve performance if you don't need it.
+			PrimaryActorTick.bCanEverTick = true;
+		
+		}
+		
+		// Called when the game starts or when spawned
+		void AFPSCharacter::BeginPlay()
+		{
+			Super::BeginPlay();
+			
+			if (GEngine != nullptr)
+			{
+				GEngine->AddOnScreenDebugMessage(-1, 5.f, FColor::White, TEXT("We are using FPSCharacter."));
+			}
+		}
+		
+		// Called every frame
+		void AFPSCharacter::Tick( float DeltaTime )
+		{
+			Super::Tick( DeltaTime );
+		
+		}
+		
+		// Called to bind functionality to input
+		void AFPSCharacter::SetupPlayerInputComponent(class UInputComponent* InputComponent)
+		{
+			Super::SetupPlayerInputComponent(InputComponent);
+			
+			// set-up gameplay key bindings
+			InputComponent->BindAxis("MoveForward", this, &AFPSCharacter::MoveForward);
+			InputComponent->BindAxis("MoveRight", this, &AFPSCharacter::MoveRight);
+			// set-up gameplay mouse bindings
+			InputComponent->BindAxis("Turn", this, &AFPSCharacter::AddControllerYawInput);
+			InputComponent->BindAxis("LookUp", this, &AFPSCharacter::AddControllerPitchInput);
+			// set-up jump action bindings
+			InputComponent->BindAction("Jump", IE_Pressed, this, &AFPSCharacter::OnStartJump);
+			InputComponent->BindAction("Jump", IE_Released, this, &AFPSCharacter::OnStopJump);
+		}
+		
+		// Forward movement function definition
+		void AFPSCharacter::MoveForward(float Value)
+		{
+			if (Controller != nullptr && Value != 0.f)
+			{
+				//find out which way is forward
+				FRotator Rotation = Controller->GetControlRotation();
+				//limit pitch when walking for falling
+				if (GetCharacterMovement()->IsMovingOnGround() || GetCharacterMovement()->IsFalling())
+				{
+					Rotation.Pitch = 0.f;
+				}
+				//add movement in that direction
+				const FVector Direction = FRotationMatrix(Rotation).GetScaledAxis(EAxis::X);
+				AddMovementInput(Direction, Value);
+			}
+		}
+		
+		//Right movement function definition
+		void AFPSCharacter::MoveRight(float Value)
+		{
+			if (Controller != nullptr && Value != 0.f)
+			{
+				//find our which way is right
+				const FRotator Rotation = Controller->GetControlRotation();
+				const FVector Direction = FRotationMatrix(Rotation).GetScaledAxis(EAxis::Y);
+				//add movement in that direction
+				AddMovementInput(Direction, Value);
+			}
+		}
+		
+		//Function setting the jump flag
+		void AFPSCharacter::OnStartJump()
+		{
+			bPressedJump = true;
+		}
+		
+		//Function clearning the jump flag
+		void AFPSCharacter::OnStopJump()
+		{
+			bPressedJump = false;
+		}
+
+## Testing Character Jumping
+
+It's now time to compile and test our newly implemented character movement functions.
+
+1.	Save the `FPSCharacter` interface (\*.h) and implementation (\*.cpp) files in Visual Studio.
+
+1.	Locate **FPSProject** in the **Solution Explorer**.
+
+1.	Right-click on **FPSProject** and select **Build** to compile your project.
+
+	![](BuildProject.png)
+
+1.	After the build finishes, open your **FPSProject** in Unreal Editor.
+
+1.	Click the **Play** button in the Level Editor Toolbar. You should be able to jump throughout the map.
+
+1.	Click the **Stop** button in the Level Editor to exit Play in Editor (PIE) mode.
 
 %Steps%
\ No newline at end of file
