#UE4 Docs: Applying SME feedback and changing all the outlines from green to yellow.
sean.gribbin

diff --git a/Engine/Documentation/Source/Platforms/Android/GettingStarted/1/AndroidQS_1.INT.udn b/Engine/Documentation/Source/Platforms/Android/GettingStarted/1/AndroidQS_1.INT.udn
--- a/Engine/Documentation/Source/Platforms/Android/GettingStarted/1/AndroidQS_1.INT.udn
+++ b/Engine/Documentation/Source/Platforms/Android/GettingStarted/1/AndroidQS_1.INT.udn
@@ -21,7 +21,6 @@ Skilllevel:Intermediate
 [/VAR]
 
 
-
 %Steps%
 
 %Globals:OSSelect%
@@ -84,7 +83,7 @@ It is highly recommended that you install Android Works while the editor and lau
 
 	1. After  **Setup.bat** has finished running locate and run the Android Works installer following the instructions below.
 
-	1. When Android Works has finished installing, run the **Setup.bat** file on more time to download the required Android Development files.
+	1. When Android Works has finished installing, restart your PC so the environment variables are set and then run the **Setup.bat** file on more time to download the required Android Development files.
 	[/region]
 
 1. **Double-click** on the **AndroidWorks-1R1-windows.exe** file to initiate the install process.
@@ -168,7 +167,7 @@ Use the links provide below to download and install the correct versions of Java
 
 	1. After  **Setup.bat** has finished running locate and run the Android Works installer following the instructions below.
 
-	1. When Android Works has finished installing, run the **Setup.bat** file on more time to download the required Android Development files.
+	1. When Android Works has finished installing, restart your Mac so the environment variables are set and then run the **Setup.bat** file on more time to download the required Android Development files.
 	[/region]
 
 
