Adding the first pass on the SteamVR docs as well as the Motion Controller docs.

diff --git a/Engine/Documentation/Source/Platforms/SteamVR/BestPractices/SteamVRBestPractice.INT.udn b/Engine/Documentation/Source/Platforms/SteamVR/BestPractices/SteamVRBestPractice.INT.udn
new file mode 100644
--- /dev/null
+++ b/Engine/Documentation/Source/Platforms/SteamVR/BestPractices/SteamVRBestPractice.INT.udn
@@ -0,0 +1,68 @@
+Title: SteamVR Best Practices
+Crumbs:%ROOT%, Platforms
+Description:An introduction to getting started with SteamVR
+Availability:Public
+Navigation:platform
+Platform:SteamVR
+
+![](Platforms\SteamVR\T_Vive_HMD.png)
+
+
+Virtual Reality or VR for short is quickly becoming the next big thing in digital entertainment. 
+Because of this industry wide push to VR, Unreal Engine 4 (UE4) supports a wide variety of VR headsets and controllers without developers having to do anything else besides just opening the Unreal Editor and making sure their VR device is connected to their development PC. 
+One of the newer VR integrations UE4 now fully supports is SteamVR. 
+Unlike other VR integrations the SteamVR integration not only includes code for the new Vive VR headset, 
+it also includes code that helps to power the new Steam Controllers as well as the Steam Lighthouse Base stations. 
+In the following sections we will cover all you need to know about how to get UE4 working with the SteamVR integration.
+
+[region:note]
+Currently the SteamVR headset and controllers are not publicly available for purchase. 
+The only way to receive a development kit is to sign up as a SteamVR developer on the **[Valve](http://store.steampowered.com/universe/vr)** website. 
+Also keep in mind that because SteamVR is still currently prototype hardware, images or information contained within this document could abruptly change or be slightly different than what you might have received or are currently using.
+As things like this happen this document will be updated so make sure to check back from time to time to familiarize yourself with anything new or any changes that might have happened.
+[/region]
+
+
+## Prerequisites
+
+Before we can start developing content to use with SteamVR there are a few things that need to be setup first. 
+In the following section we will go over what needs to be setup before you can begin to develop content for SteamVR.
+
+
+* Make sure to turn off **[Windows Aero](http://windows.microsoft.com/en-us/windows-vista/what-is-windows-aero)** to ensure that all of your development PC's power is available to use with VR.
+* From the Steam client make sure to download and install the **SteamVR** tools. You can find and install the tools by doing the following.
+ 1. From the Steam client go to **Library** > **Tools** then in the list of tools look for **SteamVR**.
+ 1. **Double - Click** on the **SteamVR** entry to start the tools install process and then follow the directions provide by Steam.
+
+* Make sure to download and install Unreal Engine version 4.8 or later from the Unreal Engine launcher.
+* Have all the required cords and power supplies for both the Vive Head Mounted Display(HMD) and the Lighthouse Base stations unpacked and plugged in.
+* Make sure you have a 15' x 15' (4.5 M x 4.5 M) area cleared of all objects that you could trip over while you have the Vive HMD on. 
+* While not required, having a wireless mouse and keyboard will make entering commands easier as you will not have to take off the HMD and then move over to where the development PC is located then input the commands.
+
+
+## What is SteamVR
+
+SteamVR or Steam Virtual Reality is fully featured, 360°, room scale VR experience that allows users to experience, interact, and even move around in virtual worlds. 
+SteamVR can be utilized in either an active, moving around the room, or passive, sitting in a chair, setting allowing developers to choose which setting will give the best experience to users. 
+SteamVR also makes use of two wireless single handed controllers whose real world  movement, position, and rotation can be precisely tracked and then translated in to VR allowing users the ability to precisely interact with object in VR like never before. 
+If that was not enough, SteamVR makes use of the Lighthouse tracking system which utilizes room scanning lasers to precisely track the real world position, rotation, and movement of the user and then translates that data so that what the user is doing in the real world will match up to the VR world they are experiencing. 
+
+## How is SteamVR different than other VR devices
+
+While there are a quite a few similarities between SteamVR and other VR implementations, there are also quite a few differences.
+One of the main differences between SteamVR and other VR systems is that instead of using a camera to track the user's location in the real world and then translate that into the virtual one, Steam VR makes use of laser driven tracking stations to track the users real world movement. 
+These stations, called Lighthouse trackers or Base stations, can more accurately track the position and orientation of the headset and controllers and are optimized for accuracy with low latency. 
+In fact the Lighthouse trackers are so precise you can simultaneously track multiple input devices with sub millimeter precision in just milliseconds.
+SteamVR is also built around the SteamVR Controllers which are specifically designed to make it easier to interact with objects in VR. 
+Finally, and probably the biggest difference, is that with the SteamVR setup you can physically move around in the room you are using SteamVR in and have that movement translated into the VR world allowing the user to explore or view parts of the world that other VR systems would not be able to achieve due to you not being able to physically move around. 
+
+
+
+
+
+
+
+
+
+
+
diff --git a/Engine/Documentation/Source/Platforms/SteamVR/QuickStart/1/SteamVRQuickStart_1.INT.udn b/Engine/Documentation/Source/Platforms/SteamVR/QuickStart/1/SteamVRQuickStart_1.INT.udn
new file mode 100644
--- /dev/null
+++ b/Engine/Documentation/Source/Platforms/SteamVR/QuickStart/1/SteamVRQuickStart_1.INT.udn
@@ -0,0 +1,65 @@
+Title: 1. SteamVR Initial Setup
+Crumbs:%ROOT%, Platforms
+Description: Setting up the require Hardware and Software for SteamVR
+Availability:Public
+Navigation:platform
+Platform:SteamVR
+
+[VAR:Steps]
+[OBJECT:Navigation]
+	[PARAM:previous]
+	[/PARAM]
+	[PARAM:current]
+	[/PARAM]
+	[PARAM:home]
+		[SteamVR Quick Start Home](Platforms/SteamVR/QuickStart)
+	[/PARAM]
+	[PARAM:next]
+		[Next Step](Platforms/SteamVR/QuickStart/2 "%Platforms/SteamVR/QuickStart/2:title%")
+	[/PARAM]
+[/OBJECT]
+[/VAR]
+
+%Steps%
+
+Making sure that the SteamVR hardware is properly setup and that the correct software is installed on your development PC is of vital importance to ensuring that the SteamVR hardware will work correctly with Unreal Engine 4. 
+In the following section we will take a look at any special considerations that you might need to make when setting up SteamVR software and hardware to work with UE4.
+
+[region:note]
+With each SteamVR development kit Valve has provide [detailed instructions](http://media.steampowered.com/apps/steamvr/vr_setup.pdf ) that will show you how to correctly set everything up. 
+If you have not read this document please do so before going any further as the following information is not a substitute for the information contained in the document created by Valve.
+[/region]
+
+1. Make sure to clear a 12 x 9ft / 4 x 3M area of all objects that you could trip over or hit while you are using SteamVR.
+
+1. Inside of the Steam Client in the **Library** section under the **Tools** option, download and install both **SteamVR** and **VRMonitor**.
+
+	![](T_SteamTools_Software.png)
+
+1. Ensure that the Vive Head Mounted Display (HMD), Steam Controllers, Breakout box, and Lighthouse Base stations are all unpacked, powered on, connected and setup according the the [instruction](http://media.steampowered.com/apps/steamvr/vr_setup.pdf) provide by Valve.
+
+1. Once everything is connected and turned on, pick up the Vive HMD so that the lenses are facing you.
+On the **Left Hand Side** of the Vive HMD there is a small **Black Button** with a small light above it. 
+Currently the light should be **red** meaning that the Vive HMD is in standby mode. 
+Press and hold the **Black Button** in for **5 Seconds** and the red light should change to **green**, meaning the Vive HMD is now ready to be used.
+
+	![](T_Vive_Standby_Mode.png)
+
+	[region:note]
+	Keep in mind that in order for the Vive HMD to work correctly you will need to have a **[USB 3.0](http://en.wikipedia.org/wiki/USB_3.0)** port available to use on your development PC. 
+	[/region]
+
+1. From the Steam client launch the **Steam VR** application and the following window will be displayed.
+
+	![](T_Steam_VR_App.png)
+
+1. Make sure to read and complete the first 3 steps in the **Setting up** section and once everything has been completed click on the **Run your first VR app!** then put your HMD on. 
+   If everything has been setup correctly you should now see something that looks like the following image in the HMD.
+
+	![](T_FirstVRApp.png)
+
+
+
+
+
+%Steps%
\ No newline at end of file
diff --git a/Engine/Documentation/Source/Platforms/SteamVR/QuickStart/2/SteamVRQuickStart_2.INT.udn b/Engine/Documentation/Source/Platforms/SteamVR/QuickStart/2/SteamVRQuickStart_2.INT.udn
new file mode 100644
--- /dev/null
+++ b/Engine/Documentation/Source/Platforms/SteamVR/QuickStart/2/SteamVRQuickStart_2.INT.udn
@@ -0,0 +1,102 @@
+Title: 2. Setting up UE4 to work with SteamVR
+Crumbs:%ROOT%, Platforms
+Description: Setting up SteamVR and UE4 to work together
+Availability:Public
+Navigation:platform
+Platform:SteamVR
+
+[VAR:Steps]
+[OBJECT:Navigation]
+	[PARAM:previous]
+	[Previous Step](Platforms/SteamVR/QuickStart/1 "%Platforms/SteamVR/QuickStart/1:title%")
+	[/PARAM]
+	[PARAM:current]
+	[/PARAM]
+	[PARAM:home]
+		[SteamVR Quick Start Home](Platforms/SteamVR/QuickStart)
+	[/PARAM]
+	[PARAM:next]
+	[/PARAM]
+[/OBJECT]
+[/VAR]
+
+%Steps%
+
+Now that we have established the SteamVR is setup and working correctly, it is now time to get it working with UE4. 
+In the following section we will go over how to setup a project to work with Steam VR.
+
+1. First create a new Blueprint base project with the following settings.
+
+	![](T_CreateNewProject.png)
+
+1. Once the project loads check to see if your VR Headset is working by clicking on the small triangle next to the **Play** button and then from the menu that is displayed select the **VR Preview** option. 
+	If your VR Headset is properly set up you should now see the basic level being display and have the ability to rotate your head left or right.
+	
+	![](Check_VR_Working_00.png)
+
+	[region:tip]
+	If your can not get the Vive HMD to work check the Plugins section to make sure that the Steam VR plug-in is enabled and if not click on the checkmark box next to **Enable** to enable it.
+	 ![](T_Enable_Steam_VR.png)
+	 If that still does not get your VR Headset to work with the VR Preview make sure to check out the [VR Documents](Platforms/VR) to ensure that nothing was missed during setup.
+	[/region]
+
+1. Right - Click inside of the **Content Browser** and select the create a new Blueprint.
+	From the **Pick Parent Class** window that is displayed use the search box at the bottom and search for **PlayerCameraManager**. 
+	Once PlayerCameraManager is located click on it to create it and name it **VR_PlayerCameraManager**.
+
+	![](Creat_New_PCM.png)
+
+1. Open up the **VR_PlayerCameraManager** and then in the **Class Defaults** section under **Player Camera Manger** enable the  **Follow Hmd Orientation** option.
+
+	![](Setting_Up_PCM.png)
+
+1. Create a new Blueprint from the Content Browser and select the **Player Controller** option. 
+	Name the newly created Player Controller, **VR_PlayerController** and open it up.
+
+	![](Creat_New_PlayerController.png)
+
+1. Once opened in the **Class Defaults** section under **Player Controller** set the **Player Camera Manager Class** to the **VR_PlayerController** that was created previously by clicking on the white box that says **None** and selecting **VR_PlayerController** from the list.
+
+	![](Setting_UP_PC.png)
+
+1. Create another Blueprint in the Content Browser this time selecting **Pawn** from the **Pick Parent Class** window and name the pawn **VR_Pawn**.
+	
+	![](Create_New_Pawn.png)
+
+1. Open up the newly created **VR_Pawn** and navigate to the **Class Defaults** section.
+	From there in the **Details** panel under the **Camera** section set the **Base Eye Height** to **0.0**.
+	
+	[region:note]
+	This only applies to Steam VR other VR setups can skip this step.
+	[/region]
+
+	![](Setting_Up_PC_Defaults.png)
+
+1. Click on the **Event Graph** tab and then Right - Click in the graph to expose the search function. Search for the **Set Tracking Space** node add at it to the graph.
+	 Wire the output of the **Event Begin Play** to the input on the Set Tracking Space node.
+
+	[region:note]
+	This only applies to Steam VR other VR setups can skip this step.
+	[/region]
+
+	![](Set_Tracking_Space.png)
+
+	[region:tip]
+	You might notice that the **New Space** input on the Set Tracking Space was set to a variable. 
+	This was done purely to make adjusting this value easier as you can then adjust this value from the default properties of the Blueprint. 
+	[/region]
+
+1. Open up the **World Settings** of your project and make sure that under the **Game Mode** section the following data is filled in with what was created in the previous steps.
+
+	![](Setting_Up_Game_Mode.png)
+
+	[region:note]
+	Make sure that once you have input the Game Mode options you created you save your level so that when your project is re-opened what you input will be loaded.
+	[/region]
+
+1. Use the VR Preview mode to check to make sure that everything is working. 
+	If everything has been setup correctly you should now only be able to rotate your head left or right and up or down.
+
+
+%Steps%
+
diff --git a/Engine/Documentation/Source/Platforms/SteamVR/QuickStart/SteamVRQuickStart.INT.udn b/Engine/Documentation/Source/Platforms/SteamVR/QuickStart/SteamVRQuickStart.INT.udn
new file mode 100644
--- /dev/null
+++ b/Engine/Documentation/Source/Platforms/SteamVR/QuickStart/SteamVRQuickStart.INT.udn
@@ -0,0 +1,116 @@
+Title: SteamVR Quick Start
+Crumbs:%ROOT%, Platforms
+Description: A How To guide for setting up SteamVR to work with UE4
+Availability:Public
+Navigation:platform
+Platform:SteamVR
+
+![](Platforms\SteamVR\T_Vive_HMD.png)
+
+This "Quick Start" series will walk you through how to setup SteamVR so that it will work with Unreal Engine 4.
+By the end of this "Quick Start" series you will have a new project that will work SteamVR.
+
+[OBJECT:TopicButtonList]
+	[PARAM:icon]
+		![](%ROOT%/build_icon.png)(convert:false)
+	[/PARAM]
+	[PARAM:title]
+		Steps
+	[/PARAM]
+	[PARAM:description]
+	[/PARAM]
+	[PARAM:links]
+	* [](Platforms/SteamVR/QuickStart/1 "%Platforms/SteamVR/QuickStart/1:description%")
+	* [](Platforms/SteamVR/QuickStart/2 "%Platforms/SteamVR/QuickStart/2:description%")
+	[/PARAM]
+[/OBJECT]
+
+
+
+
+
+
+
+
+
+
+
+
+
+
+
+
+
+
+
+
+
+
+
+
+
+
+
+
+
+
+
+
+
+
+
+
+
+
+
+
+
+
+
+
+
+
+
+
+
+
+
+
+
+
+
+
+
+
+
+
+
+
+
+
+
+
+
+
+
+
+
+
+
+
+
+
+
+
+
+
+
+
+
+
+
+
+
+
+
+
diff --git a/Engine/Documentation/Source/Platforms/SteamVR/SteamVR.INT.udn b/Engine/Documentation/Source/Platforms/SteamVR/SteamVR.INT.udn
new file mode 100644
--- /dev/null
+++ b/Engine/Documentation/Source/Platforms/SteamVR/SteamVR.INT.udn
@@ -0,0 +1,88 @@
+Title: SteamVR Development
+Crumbs:%ROOT%, Platforms
+Description:An introduction to getting started with SteamVR
+Availability:Public
+Navigation:platform
+Platform:SteamVR
+
+![](T_Vive_HMD.png)
+
+[OBJECT:TopicButtonList]
+	[PARAM:icon]
+		![](%ROOT%/SteamVR_Icon.png)(convert:false)
+	[/PARAM]
+	[PARAM:title]
+		Oculus Rift Development
+	[/PARAM]
+	[PARAM:description]
+	[/PARAM]
+	[PARAM:links]
+		* [](Platforms/SteamVR/BestPractices "%Platforms/SteamVR/BestPractices:description%")
+		* [](Platforms/SteamVR/QuickStart "%Platforms/SteamVR/QuickStart:description%")
+		* [](Platforms/SteamVR/Tps&Tricks "%Platforms/SteamVR/Tps&Tricks:description%")
+	[/PARAM]
+[/OBJECT]
+
+
+
+
+
+
+
+
+
+
+
+
+
+
+
+
+
+
+
+
+
+
+
+
+
+
+
+
+
+
+
+
+
+
+
+
+
+
+
+
+
+
+
+
+
+
+
+
+
+
+
+
+
+
+
+
+
+
+
+
+
+
+
+
diff --git a/Engine/Documentation/Source/Platforms/SteamVR/Tps&Tricks/SteamVRTips&Tricks.INT.udn b/Engine/Documentation/Source/Platforms/SteamVR/Tps&Tricks/SteamVRTips&Tricks.INT.udn
new file mode 100644
--- /dev/null
+++ b/Engine/Documentation/Source/Platforms/SteamVR/Tps&Tricks/SteamVRTips&Tricks.INT.udn
@@ -0,0 +1,51 @@
+Title: SteamVR Tips & Tricks
+Crumbs:%ROOT%, Platforms
+Description: A collection of tips & tricks to help you unlock the full power of SteamVR
+Availability:Public
+Navigation:platform
+Platform:SteamVR
+
+## Overview 
+
+In the following section we will cover some tips and tricks that you can use to help unlock the full potential of UE4 and SteamVR.
+
+
+### Avoiding Display Flickering
+
+Each time the editor is launched the display will flicker when the SteamVR compositor takes over. 
+To stop this from happening you can do the following to keep the compositor running in the background at all times.
+
+1. Locate the **vrcompositor.exe** file that should be in **C:\ProgramFiles(x86)\Steam\steamapps\common\SteamVR\bin\win64**.
+	
+	![](T_Comp_EXE.png)
+
+1. Create a shortcut to vrcompositor.exe and then send that shortcut to your desktop.
+1. Right click on the shortcut that was placed on your desktop and then click on **Properties**.
+1. Inside of the **Shortcut** menu in the **Target** input add **-keepalive** after the .exe and press the **OK** button when done.  
+	![](T_KeepAlive_CLA.png)
+1. Now when opening up the UE4 editor you should no longer have a flickering problem.
+
+### VR Interaction Area
+
+
+
+
+
+
+
+
+
+
+
+
+
+
+
+
+
+
+
+
+
+
+
diff --git a/Engine/Documentation/Source/Platforms/VR/MotionController/MotionControllerComponent.INT.udn b/Engine/Documentation/Source/Platforms/VR/MotionController/MotionControllerComponent.INT.udn
new file mode 100644
--- /dev/null
+++ b/Engine/Documentation/Source/Platforms/VR/MotionController/MotionControllerComponent.INT.udn
@@ -0,0 +1,172 @@
+Availability:Public
+Title:Motion Controller Component Setup
+Crumbs:%ROOT%, Platforms
+Description:Information over how to setup Motion Controllers for VR interaction.
+Navigation:platform
+Version:4.9
+Skilllevel:Intermediate
+
+
+[TOC(start:2)]
+
+The release of Unreal Engine 4.9 introduces some major changes in how Unreal Engine interfaces with Motion Controllers.
+You can now add a Motion Controller as a component to your Pawn / Character, set which "Hand" it is supposed to represent and the Motion Controller is now ready to be used.
+In the following document we will cover how you can set up your Virtual Reality project to make use of Motion Controllers.
+
+[region:tip]
+Before going any further make sure that your VR Headset and Motion Controllers have been plugged into your PC and their required drivers installed. 
+If you have not done so already make sure to check out our other [VR Documents](Platforms/VR) as they contain valuable information on how to get various VR headsets to work with UE4.
+[/region]
+
+## Motion Controller Setup
+
+Setting up a VR project to make use of Motion Controllers can be accomplished in a few short steps. 
+In the following section we will cover how you can quickly get Motion Controllers working with your project.
+
+1. From the Content Browser open up the Player Pawn Blueprint.
+
+	![](Open_VR_Pawn.png)
+
+	[region:note]
+	Please note the following can be applied to any Pawn in any project and does not specifically have to be a VR project.
+	[/region]
+
+1. Once the Blueprint is opened, navigate to the **Viewport** tab and in the **Components** section click on the **Add Component** button to expose the components that can be added to this Blueprint. 
+
+	![](Find_Components.png)
+
+1. From the list that is shown search for the Motion Controller using the word **Motion** as the search term. 
+	When the Motion Controller component has been located, click on it to add it to the components list.
+
+	[OBJECT:EmbeddedVideo]
+	[PARAMLITERAL:width]
+	640
+	[/PARAMLITERAL]
+	[PARAMLITERAL:height]
+	360
+	[/PARAMLITERAL]
+	[PARAMLITERAL:videoid]
+	kcyHnfpmWbU
+	[/PARAMLITERAL]
+	[/OBJECT]
+
+1. Once the Motion Controller has been added, select it by clicking on it and then in the **Details** panel under the **Motion Controller** section make sure the **Hand** is set to **Left**.
+
+	![](Setting_Up_Components.png)
+
+1. If you ran the game now the motion controller would work but you would not see anything as by default no visuals are rendered. 
+	To fix this add a **Static Mesh Component** to the Components sections and then in the Details section of the Static Mesh Component under the Static Mesh section pick a Static Mesh to use from the dropdown menu.
+
+	[OBJECT:EmbeddedVideo]
+	[PARAMLITERAL:width]
+	640
+	[/PARAMLITERAL]
+	[PARAMLITERAL:height]
+	360
+	[/PARAMLITERAL]
+	[PARAMLITERAL:videoid]
+	x4vjDqnHQnk
+	[/PARAMLITERAL]
+	[/OBJECT]
+
+	[region:tip]
+	If you do not see any Static Mesh in the Content Browser you will need to either import one or enable **Engine Content** and then select one from there. 
+	Also if the Static Mesh you selected is to big or to small you can scale to your likening using the Scale section of the Transform menu.
+	[/region]
+
+1. Repeat this process again this time adding a Right Hand and when completed your Pawn Blueprint should look something like the image below.
+
+	![](Finshed_MC_Setup.png)
+
+	[region:note]
+	In the image above the the two cone shaped meshes have been moved apart to show that they are two of them. 
+	When setting this up for your project leave both mesh at 0,0,0 so that they are located right on top of one another.
+	[/region]
+
+## Working with Motion Controllers
+
+Working with Motion Controllers inside UE4 is similar to working with other input devices such as a game-pad or a keyboard. 
+If you search inside of the Blueprint Editor **Palette** using **MotionController** as the search term you will see all the available buttons and controls the Motion Controllers have available.
+
+![](Working_With_MC.png)
+
+When looking at the available Motion Controller input options there is a capital **L** or **R** next to the name of what that value is supposed to control. 
+This L or R means that this is for the **Left** or **Right** Motion Controller.
+
+![](MC_Left_Right.png)
+
+
+## Using Motion Controllers
+
+In the following section we will set the Motion Controllers up so that when you press the trigger button on the **Right-Handed** Motion Controller a Blueprint containing a Static Mesh is spawned into the world. 
+Then when you press the trigger button on the **Left-Handed** Motion Controller the Blueprint that contains the Static Mesh that was spawned into the world will be removed from the world.
+
+1. Inside of the Content Browser create a new Blueprint and name it **BP_Sphere** and then open it up by double clicking on it.  
+
+1. Once the BP_Sphere Blueprint is open, navigate to the **Viewport** tab and add a new **Static Mesh Component**.
+
+	![](Add_SM_Component.png)
+
+1. Move the Static Mesh Component on top of the **Default Screen Root** and then set the Static Mesh used to the **ArcadeEditorSphere** then compile and save your Blueprint.
+
+	[OBJECT:EmbeddedVideo]
+		[PARAMLITERAL:width]
+		640
+		[/PARAMLITERAL]
+		[PARAMLITERAL:height]
+		360
+		[/PARAMLITERAL]
+		[PARAMLITERAL:videoid]
+		PszzV_I0xso
+		[/PARAMLITERAL]
+	[/OBJECT]
+
+	[region:note]
+	If you can not find the **ArcadeEditorSphere** make sure that you have the **Engine Content** folder enabled by going to **View Options** and then clicking on the **Engine Content** folder to enable it showing up in the Content Browser.
+	[/region]
+
+1. Close down the BP_Sphere Blueprint and open up your **Pawn** Blueprint then navigate to the **Event Graph** tab.
+
+1. Create a new **Actor** variable with the name **Objects** and set it's **Variable Type** to array.
+	
+	![](Create_New_Var.png)
+
+1. To create the ability to add objects to the world, inside of the Event Graph add and connect the following Blueprint nodes together. When completed you should have something that looks like the following image.
+
+	![](Add_Objects_Setup.png)
+
+	[region:note]
+	The **Spawn Offset** was added to ensure that the objects could be seen while being placed in this example. This is not required for this to work and can be omitted by taking the **Return Value** from the **GetActorTransform** node and inputting that into the **Spawn Transform** input on the **SpawnActor** Blueprint node.
+	![](Add_Objects_Setup_Alt.png)
+	[/region]
+
+1. Once the Add Objects has been setup select the **Spawn Actor** Blueprint node and in the **Class** section, select the ** BP_Sphere** Blueprint that was created in step one.
+
+	![](Setting_SpawnActor_Up.png)
+
+1. To create the ability to remove the objects that were placed in the world, add and connect the following Blueprint nodes together. When completed you should have something that looks like the following image.
+
+	![](Remove_Objects_Setup.png)
+
+	[region:note]
+	The keyboard input keys were placed in the graph for quick testing. You do not have to include those inputs for this to work with Motion Controllers.
+	[/region]
+
+1. **Compile** then **Save** the changes to your Pawn Blueprint then press the **Play** button to preview the level. 
+	Once the level has started try adding objects to the world by pressing the **K** key on the keyboard or by pressing the Right trigger on the Right Handed Motion Controller. 
+	To remove objects from the world press the **T** key on the keyboard or press the **Left** trigger on the Left Handed Motion Controller which you can see demonstrated in the video below.
+
+	[OBJECT:EmbeddedVideo]
+		[PARAMLITERAL:width]
+		640
+		[/PARAMLITERAL]
+		[PARAMLITERAL:height]
+		360
+		[/PARAMLITERAL]
+		[PARAMLITERAL:videoid]
+		VXUowSUKT9I
+		[/PARAMLITERAL]
+	[/OBJECT]
+
+
+
