#UE4doc #docs:
Added engine version field
Removed Overview header
Changed font settings for multiple terms
Added multiple images and replaced a couple to reflect added/changed features of the Editor's UI
Added descriptions for buttons that have been added to the editor.

diff --git a/Engine/Documentation/Source/Engine/Animation/Persona/Toolbar/PersonaToolbar.INT.udn b/Engine/Documentation/Source/Engine/Animation/Persona/Toolbar/PersonaToolbar.INT.udn
--- a/Engine/Documentation/Source/Engine/Animation/Persona/Toolbar/PersonaToolbar.INT.udn
+++ b/Engine/Documentation/Source/Engine/Animation/Persona/Toolbar/PersonaToolbar.INT.udn
@@ -2,6 +2,7 @@ Availability:Public
 Title:Persona Toolbar
 Crumbs: %ROOT%, Engine, Engine/Animation, Engine/Animation/Persona
 Description:User guide for the Toolbar in the Persona animation editor.
+version: 4.9
 
 [VAR:Topic]
 [OBJECT:Topic]
@@ -28,7 +29,7 @@ Description:User guide for the Toolbar in the Persona animation editor.
 
 [TOC(start:2)]
 
-## Overview
+
 
 [EXCERPT:Overview]
 The **Toolbar** is displayed along the top of the Persona interface window by default. It contains a variety of common functions you will need when editing animation assets.
@@ -49,7 +50,7 @@ When in **Skeleton**, **Mesh**, or **Animation** modes, the Toolbar will look so
 
  However, when in **Graph** mode, the toolbar expands with further features, looking like this:
 
-![](Toolbar_2.png)
+![](PersonaToolbar.png)
 
 You may notice that many of the additional buttons are the same as those used in the Blueprint Editor window. 
 
@@ -61,26 +62,39 @@ Persona toolbar buttons provide easy access to switching between asset modes, su
 
 | Button																						| Description														|
 | --------------------------------------------------------------------------------------------- | ----------------------------------------------------------------- | 
-| ![Compile Successful Button](Engine/Blueprints/Editor/UIComponents/Toolbar/bp_compile_successful.png)	| [INCLUDE:Engine/Blueprints\Editor/UIComponents\Toolbar\Toolbaroptions#compilesuccessful]	|
+| ![Compile Successful Button](Engine/Blueprints/Editor/UIComponents/Toolbar/bp_compile_successful.png)		| [INCLUDE:Engine/Blueprints\Editor/UIComponents\Toolbar\Toolbaroptions#compilesuccessful]	|
 | ![Compile Needed Button](Engine/Blueprints/Editor/UIComponents/Toolbar/bp_needs_compile.png)				| [INCLUDE:Engine/Blueprints\Editor/UIComponents\Toolbar\Toolbaroptions#compileneeded]		|
 | ![Compile Warning Button](Engine/Blueprints/Editor/UIComponents/Toolbar/compile_warning.png)				| [INCLUDE:Engine/Blueprints\Editor/UIComponents\Toolbar\Toolbaroptions#compilewarning]		|
 | ![Compile Failed Button](Engine/Blueprints/Editor/UIComponents/Toolbar/compile_failed.png)				| [INCLUDE:Engine/Blueprints\Editor/UIComponents\Toolbar\Toolbaroptions#compilefailed]		|
-| ![Save Button](Engine/Blueprints/Editor/UIComponents/Toolbar/bp_save.png)								| [INCLUDE:Engine/Blueprints\Editor/UIComponents\Toolbar\Toolbaroptions#save]					|
+| ![Save Button](Engine/Blueprints/Editor/UIComponents/Toolbar/bp_save.png)									| [INCLUDE:Engine/Blueprints\Editor/UIComponents\Toolbar\Toolbaroptions#save]					|
 | ![Find in Content Browser Button](Engine/Blueprints/Editor/UIComponents/Toolbar/bp_find_in_cb.png)		| [INCLUDE:Engine/Blueprints\Editor/UIComponents\Toolbar\Toolbaroptions#findincb]				| 
 | ![Blueprint Properties Button](Engine/Blueprints/Editor/UIComponents/Toolbar/bp_properties.png)			| [INCLUDE:Engine/Blueprints\Editor/UIComponents\Toolbar\Toolbaroptions#blueprintprops]		|
-| ![Search Button](Engine/Blueprints/Editor/UIComponents/Toolbar/bp_search.png)							| [INCLUDE:Engine/Blueprints\Editor/UIComponents\Toolbar\Toolbaroptions#search]				|
-| ![Simulate Button](Engine/Blueprints/Editor/UIComponents/Toolbar/bp_simulate.png)						| [INCLUDE:Engine/UI/LevelEditor/Toolbar#simulate]			|
+| ![Blueprint Properties Button](Engine/Blueprints/Editor/UIComponents/Toolbar/defaults.png)				| Shows the Class Defaults Panel in the Details Tab |
+| ![Search Button](Engine/Blueprints/Editor/UIComponents/Toolbar/bp_search.png)								| [INCLUDE:Engine/Blueprints\Editor/UIComponents\Toolbar\Toolbaroptions#search]				|
+| ![Simulate Button](Engine/Blueprints/Editor/UIComponents/Toolbar/bp_simulate.png)							| [INCLUDE:Engine/UI/LevelEditor/Toolbar#simulate]			|
 | ![Play In Editor Button](Engine/Blueprints/Editor/UIComponents/Toolbar/bp_play_in.png)					| [INCLUDE:Engine/UI/LevelEditor/Toolbar#playin]				|
 | ![Pause Button](Engine/Blueprints/Editor/UIComponents/Toolbar/bp_pause.png)								| [INCLUDE:Engine/Blueprints\Editor/UIComponents\Toolbar\Toolbaroptions#pause]				|
-| ![Resume Button](Engine/Blueprints/Editor/UIComponents/Toolbar/bp_resume.png)							| [INCLUDE:Engine/Blueprints\Editor/UIComponents\Toolbar\Toolbaroptions#resume]				|
+| ![Resume Button](Engine/Blueprints/Editor/UIComponents/Toolbar/bp_resume.png)								| [INCLUDE:Engine/Blueprints\Editor/UIComponents\Toolbar\Toolbaroptions#resume]				|
 | ![Frame Skip Button](Engine/Blueprints/Editor/UIComponents/Toolbar/bp_frameskip.png)						| [INCLUDE:Engine/Blueprints\Editor/UIComponents\Toolbar\Toolbaroptions#frameskip]			|
-| ![Stop Button](Engine/Blueprints/Editor/UIComponents/Toolbar/bp_stop.png)								| [INCLUDE:Engine/Blueprints\Editor/UIComponents\Toolbar\Toolbaroptions#stop]					|
+| ![Stop Button](Engine/Blueprints/Editor/UIComponents/Toolbar/bp_stop.png)									| [INCLUDE:Engine/Blueprints\Editor/UIComponents\Toolbar\Toolbaroptions#stop]					|
 | ![Possess Button](Engine/Blueprints/Editor/UIComponents/Toolbar/bp_possess.png)							| [INCLUDE:Engine/Blueprints\Editor/UIComponents\Toolbar\Toolbaroptions#possess]				|
 | ![Eject Button](Engine/Blueprints/Editor/UIComponents/Toolbar/bp_eject.png)								| [INCLUDE:Engine/Blueprints\Editor/UIComponents\Toolbar\Toolbaroptions#eject]				|
 | ![Find Node Button](Engine/Blueprints/Editor/UIComponents/Toolbar/bp_find_node.png)						| [INCLUDE:Engine/Blueprints\Editor/UIComponents\Toolbar\Toolbaroptions#findnode]				|
-| ![Step Button](Engine/Blueprints/Editor/UIComponents/Toolbar/bp_step.png)								| [INCLUDE:Engine/Blueprints\Editor/UIComponents\Toolbar\Toolbaroptions#step]					|
-| ![Debug Dropdown](Engine/Blueprints/Editor/UIComponents/Toolbar/bp_debug_dropdown.png)					| [INCLUDE:Engine/Blueprints\Editor/UIComponents\Toolbar\Toolbaroptions#debugobject]			| 
-
+| ![Step Button](Engine/Blueprints/Editor/UIComponents/Toolbar/bp_step.png)									| [INCLUDE:Engine/Blueprints\Editor/UIComponents\Toolbar\Toolbaroptions#step]					|
+| ![Debug Dropdown](Engine/Blueprints/Editor/UIComponents/Toolbar/bp_debug_dropdown.png)					| [INCLUDE:Engine/Blueprints\Editor/UIComponents\Toolbar\Toolbaroptions#debugobject]			|
+| ![](button_Preview.png)																					| Displays a preview of the currently selected animation asset. |
+| ![](button_RefPose.png)																					| Displays the skeletal mesh in it's reference pose. |
+| ![](button_Import.png)																					| Brings up the Import dialogue. |
+| ![](button_ImportMesh.png)																				| Brings up the Import dialogue. |
+| ![](button_Reimport.png)																					| Brings up the Browse dialogue to select an asset to reimport. |
+| ![](button_ReimportMesh.png)																				| Brings up the Browse dialogue to select an asset to reimport. |
+| ![](button_Export.png)																					| Brings up the Export dialogue. |
+| ![](button_Record.png)																					| Allows for the recording of an animation to be saved at another location. |
+| ![](button_CreateAsset.png)																				| Brings a menu with multiple options for creating a new asset. |
+| ![](button_Compression.png)																				| Brings multiple options of compression for the currently selected animation. |
+| ![](button_AnimNotifies.png)																				| Brings up the Anim Notifies window, which displays all Anim Notifies. |
+| ![](button_Key.png)																						| Adds a key to the Additive Layer Track of the currently selected bone for the currently selected animation. |
+| ![](button_Apply.png)																						| Applys added keys to the currently selected animation. |
 
 ## Mode Switching Tabs
 
