#UE4doc #docs:
Added engine version field
Removed Overview header
Changed font settings for multiple terms
Replaced a few images due to updated UI
Added an image to go along with a new button
Added a field explaining what the Turn Table Speed is / does.

diff --git a/Engine/Documentation/Source/Engine/Animation/Persona/Viewport/PersonaViewport.INT.udn b/Engine/Documentation/Source/Engine/Animation/Persona/Viewport/PersonaViewport.INT.udn
--- a/Engine/Documentation/Source/Engine/Animation/Persona/Viewport/PersonaViewport.INT.udn
+++ b/Engine/Documentation/Source/Engine/Animation/Persona/Viewport/PersonaViewport.INT.udn
@@ -2,6 +2,7 @@ Availability:Public
 Title:Persona Viewport
 Crumbs: %ROOT%, Engine, Engine/Animation, Engine/Animation/Persona
 Description:User guide for the Viewport preview in the Persona animation editor.
+version: 4.9
 
 [VAR:Topic]
 [OBJECT:Topic]
@@ -25,10 +26,9 @@ Description:User guide for the Viewport preview in the Persona animation editor.
 
 [TOC (start:2 end:3)]
 
-## Overview
 
 [EXCERPT:Overview]
-The **Viewport** panel allows Animation Sequences, Blendspaces, and other aspects of the _Animation Blueprint_ to be previewed on a specific Skeletal Mesh. The panel consists of a 3D perspective viewport, a movement toolbar, a view options toolbar, and a timeline/control bar for controlling playback of the preview.
+The **Viewport** panel allows Animation Sequences, Blendspaces, and other aspects of the Animation Blueprint to be previewed on a specific Skeletal Mesh. The panel consists of a 3D perspective viewport, a movement toolbar, a view options toolbar, and a timeline/control bar for controlling playback of the preview.
 [/EXCERPT:Overview]
 
 ![Viewport Panel with Callouts](ViewportPanel.png)(w:720)
@@ -174,7 +174,7 @@ The **Preview LODs** Button ![](Persona_LODAuto.png) provides the following opti
 | ------ | ----------- |
 | **LOD Auto** | Automatically selects the Level Of Detail (LOD) (this is the default setting). |
 | **LOD 0** | Selects LOD 0 as the LOD. |
-| **Level of Detail Settings** | Opens a dialog box in where you can adjust LOD Settings (pictured below). |
+<!--| **Level of Detail Settings** | Opens a dialog box in where you can adjust LOD Settings (pictured below). |
 
 ![](LODSettingsMenu.png)
 
@@ -187,7 +187,7 @@ When an LOD is applied, either through the **Import LOD** option or through the
 ![](LODAdded.png)
 
 _The added LOD(s) can either be removed or an FBX asset can be imported for the LOD._
-
+-->
 
 ### Playback Speed
 
@@ -204,6 +204,22 @@ The **Playback Speed** Button ![](Persona_PlaybackSpeed.png) provides the follow
 | **x10.0** | Sets the playback speed to 10.0. |
 [/EXCERPT:viewoptions]
 
+### Turn Table Speed
+
+The **Turn Table Speed** button ![](Persona_TurnTableSpeed.png) provides the following options:
+| Option | Description |
+| ------ | ----------- |
+| **Play** | Plays the Turn Table (Spins the skeletal mesh) at the selected speed |
+| **Pause** | Pauses the Turn Table at the current rotation |
+| **Stop** | Stops the Turn Table and resets the rotation to the original |
+| **x0.1** | Sets the Turn Table speed to 0.1.|
+| **x0.25** | Sets the Turn Table speed to 0.25. |
+| **x0.5** | Sets the Turn Table speed to 0.5. |
+| **x1.0** | Sets the Turn Table speed to 1.0 (this is the default setting). |
+| **x2.0** | Sets the Turn Table speed to 2.0. |
+| **x5.0** | Sets the Turn Table speed to 5.0. |
+| **x10.0** | Sets the Turn Table speed to 10.0. |
+
 ## Viewport
 
 The **Viewport** displays a preview of the Skeletal Mesh, any options set in the Toolbar and menu bar, and any selected animation.
@@ -217,17 +233,17 @@ The **Timeline** and **Control bar** contain the following controls:
 
 | Control | Description |
 | ------- | ----------- |
-| ![Control bar - Timeline](CE_button_timeline.png) | The timeline. This shows and allows scrubbing through the frames of the _AnimationSequence_ or _AnimationMontage_. A red bar is displayed at the current position. **Drag** the red bar to scrub. |
-| ![Control bar - Skip to Beginning](CE_button_beginning.png) | Skips to the first frame in the _AnimationSequence_ or _AnimationMontage_. |
-| ![Control bar - Step Backward](CE_button_step_back.png) | Skips to the previous frame in the _AnimationSequence_ or _AnimationMontage_. |
-| ![Control bar - Play in Reverse](CE_button_play_back.png) ![Control bar - Pause](CE_button_pause.png) | Begins playback of the _AnimationSequence_ or _AnimationMontage_ in reverse. When playing, the pause button is displayed. |
-| ![Control bar - Play Forward](CE_button_play_forward.png) ![Control bar - Pause](CE_button_pause.png) | Begins playback of the _AnimationSequence_ or _AnimationMontage_. When playing, the pause button is displayed. |
-| ![Control bar - Step Forward](CE_button_step_forward.png) | Skips to the next frame in the _AnimationSequence_ or _AnimationMontage_. |
-| ![Control bar - Skip to End](CE_button_end.png) | Skips to the last frame in the _AnimationSequence_ or _AnimationMontage_. |
+| ![Control bar - Timeline](CE_button_timeline.png) | The timeline. This shows and allows scrubbing through the frames of the AnimationSequence or AnimationMontage. A red bar is displayed at the current position. **Drag** the red bar to scrub. |
+| ![Control bar - Skip to Beginning](CE_button_beginning.png) | Skips to the first frame in the AnimationSequence or AnimationMontage. |
+| ![Control bar - Step Backward](CE_button_step_back.png) | Skips to the previous frame in the AnimationSequence or AnimationMontage. |
+| ![Control bar - Play in Reverse](CE_button_play_back.png) ![Control bar - Pause](CE_button_pause.png) | Begins playback of the AnimationSequence or AnimationMontage in reverse. When playing, the pause button is displayed. |
+| ![Control bar - Play Forward](CE_button_play_forward.png) ![Control bar - Pause](CE_button_pause.png) | Begins playback of the AnimationSequence or AnimationMontage. When playing, the pause button is displayed. |
+| ![Control bar - Step Forward](CE_button_step_forward.png) | Skips to the next frame in the AnimationSequence or AnimationMontage. |
+| ![Control bar - Skip to End](CE_button_end.png) | Skips to the last frame in the AnimationSequence or AnimationMontage. |
 | ![Control bar - Enable Looping](CE_button_loop.png) ![Control bar - Disable Looping](CE_button_loop_enabled.png) | Toggles playback between play once and looping. When looping is enabled, the button turns red. |
 
 [REGION:note]
-When opening a _Blend Space_ or _Aim Offset_, the animation will play automatically if **Preview** is enabled. If you pause or disable the looping option in an _AnimationSequence_ or _AnimationMontage_ and then open a _Blend Space_ or _Aim Offset_, the animation will not automatically play. You will need to set an animation to play/looping in order for it to autostart when opening a _Blend Space_ or _Aim Offset_. 
+When opening a Blend Space or Aim Offset, the animation will play automatically if **Preview** is enabled. If you pause or disable the looping option in an AnimationSequence or AnimationMontage and then open a Blend Space or Aim Offset, the animation will not automatically play. You will need to set an animation to play/looping in order for it to autostart when opening a Blend Space or Aim Offset. 
 [/REGION]
 
 [/EXCERPT:PreviewControls]
@@ -240,7 +256,7 @@ When opening a _Blend Space_ or _Aim Offset_, the animation will play automatica
 ![Persona: Anim Preview Setup Panel](CE_anim_preview_setup_panel.png)
 
 The **Anim Preview Setup** panel is used to configure the preview viewport. It displays a list of all of the 
-Skeletal Meshes that use the Skeleton the _VimBlueprint_ being edited was created from. To set a mesh as the 
+Skeletal Meshes that use the Skeleton the VimBlueprint being edited was created from. To set a mesh as the 
 preview mesh used in the preview viewport, select the mesh in the list and click the **Use as preview mesh**. 
 To remove the current mesh being used in the preview viewport, click the **Clear preview** button. The list of 
 Skeletal Meshes can be filtered using the text box above the list.
