#UE4doc #docs:
Added engine version field
Removed Overview header
Changed font settings for multiple terms

diff --git a/Engine/Documentation/Source/Engine/Animation/PhysicallyDrivenAnimation/PhysicallyBasedAnimation.INT.udn b/Engine/Documentation/Source/Engine/Animation/PhysicallyDrivenAnimation/PhysicallyBasedAnimation.INT.udn
--- a/Engine/Documentation/Source/Engine/Animation/PhysicallyDrivenAnimation/PhysicallyBasedAnimation.INT.udn
+++ b/Engine/Documentation/Source/Engine/Animation/PhysicallyDrivenAnimation/PhysicallyBasedAnimation.INT.udn
@@ -7,15 +7,16 @@ Related: Engine/Content/Types/SkeletalMeshes
 Related: Engine/Animation/StateMachines
 Related: Gameplay/AI
 Related: Engine/Content/Tools/MayaRiggingTool
+version: 4.9
 
 [REGION:fullwidth]
 ![](PhysicsBlend.png)
 
-## Overview
+
 
 This document covers how to apply physics-driven animation to your characters and Skeletal Meshes. The concept is that you can blend in the result of a simulation along with your keyframed animation to have a naturally simulated feel for characters that need to exhibit a "ragdoll" effect.
 
-An example of this concept in action can be seen in the [Content Examples](Resources\ContentExamples\Animation) project, within the map named _PhysicalAnimation.umap_. In this map, we have a series of Skeletal Meshes to which animation has been applied. In each one, we can also interact in some way and see physics animation take over the existing motion.
+An example of this concept in action can be seen in the [Content Examples](Resources\ContentExamples\Animation) project, within the map named PhysicalAnimation.umap. In this map, we have a series of Skeletal Meshes to which animation has been applied. In each one, we can also interact in some way and see physics animation take over the existing motion.
 
 Be aware that the use of physics on a Skeletal Mesh in any form requires that the mesh have a Physics Asset set up and applied to it. Please see the associated documentation for the [Physics Asset Tool (PhAT)](Engine/Physics\PhAT).
 
@@ -26,13 +27,13 @@ Skeletal Meshes are used in this example for overall simplicity. The same techni
 
 ## Setup
 
-There are many ways to apply physics to a character and our example shows only a few possibilities. In our approach, the two primary tools you will need are the **Set All Bodies Simulate Physics** and **Set All Bodies Below Physics Blend Weight** nodes, which will generally be placed within your character's Animation Blueprint Event Graph.
+There are many ways to apply physics to a character and our example shows only a few possibilities. In our approach, the two primary tools you will need are the Set All Bodies Simulate Physics and Set All Bodies Below Physics Blend Weight nodes, which will generally be placed within your character's Animation Blueprint Event Graph.
 
 ![](PhysicsWeghtNodes.png)
 
 ### Set All Bodies Below Simulate Physics
 
-The purpose of the **Set All Bodies Below Simulate Physics** node is to recursively activate Physics Asset bodies on a Skeletal Mesh to begin simulating physics starting at a given bone and moving recursively down the bone chain. For example, if you tell the left clavicle to start simulating, then all of the bones descending from that in the skeletal hierarchy will simulate as well, resulting in a limp or ragdolling arm. Simply put, you can think of this node as an on/off switch to begin or stop simulating physics from a given bone.
+The purpose of the Set All Bodies Below Simulate Physics node is to recursively activate Physics Asset bodies on a Skeletal Mesh to begin simulating physics starting at a given bone and moving recursively down the bone chain. For example, if you tell the left clavicle to start simulating, then all of the bones descending from that in the skeletal hierarchy will simulate as well, resulting in a limp or ragdolling arm. Simply put, you can think of this node as an on/off switch to begin or stop simulating physics from a given bone.
 
 ### Set All Bodies Below Physics Blend Weight
 
@@ -44,9 +45,9 @@ This node simply controls how much influence the Physics Asset will have over th
 Hit-based physics reactions are a common use case for simulation on characters, such as when a character is struck by a projectile. At a high-level, this requires that you:
 
 * Get the name of the bone that was hit. This can be done via a trace for instant-effect weapons, or within a projectile class where applicable.
-* Transfer that bone name into the Character's Animation Blueprint for the Event Graph to use, generally through a **Set All Bodies Below Simulate Physics** node. This activates the simulation system.
-* Control the Physics Blend Weight property via a **Set All Below Physics Blend Weight** node. Generally, you want to quickly animate this going up to 1.0 and then back down to 0.0 so that the physics reaction blends in and then back out. This would generally be done within the Animation Blueprint's Event Graph.
-* Once the reaction is complete and the Physics Blend Weight returns to 0, you should use the **Set All Bodies Below Simulate Physics** node once more to deactivate the simulation.
+* Transfer that bone name into the Character's Animation Blueprint for the Event Graph to use, generally through a Set All Bodies Below Simulate Physics node. This activates the simulation system.
+* Control the Physics Blend Weight property via a Set All Below Physics Blend Weight node. Generally, you want to quickly animate this going up to 1.0 and then back down to 0.0 so that the physics reaction blends in and then back out. This would generally be done within the Animation Blueprint's Event Graph.
+* Once the reaction is complete and the Physics Blend Weight returns to 0, you should use the Set All Bodies Below Simulate Physics node once more to deactivate the simulation.
 
 [REGION:note]
 You can see an example of this technique in Example 1.2 of the PhysicalAnimation.umap level, found within the Content Examples project.
