Fixing minor errors that came up when publishing.

diff --git a/Engine/Documentation/Source/Platforms/SteamVR/BestPractices/SteamVRBestPractice.INT.udn b/Engine/Documentation/Source/Platforms/SteamVR/BestPractices/SteamVRBestPractice.INT.udn
--- a/Engine/Documentation/Source/Platforms/SteamVR/BestPractices/SteamVRBestPractice.INT.udn
+++ b/Engine/Documentation/Source/Platforms/SteamVR/BestPractices/SteamVRBestPractice.INT.udn
@@ -17,7 +17,7 @@ In the following sections we will cover all you need to know about how to get UE
 
 [region:note]
 Currently the SteamVR headset and controllers are not publicly available for purchase. 
-The only way to receive a development kit is to sign up as a SteamVR developer on the **[Valve](http://store.steampowered.com/universe/vr)** website. 
+The only way to receive a development kit is to sign up as a SteamVR developer on the [Valve](http://store.steampowered.com/universe/vr) website. 
 Also keep in mind that because SteamVR is still currently prototype hardware, images or information contained within this document could abruptly change or be slightly different than what you might have received or are currently using.
 As things like this happen this document will be updated so make sure to check back from time to time to familiarize yourself with anything new or any changes that might have happened.
 [/region]
@@ -29,7 +29,7 @@ Before we can start developing content to use with SteamVR there are a few thing
 In the following section we will go over what needs to be setup before you can begin to develop content for SteamVR.
 
 
-* Make sure to turn off **[Windows Aero](http://windows.microsoft.com/en-us/windows-vista/what-is-windows-aero)** to ensure that all of your development PC's power is available to use with VR.
+* Make sure to turn off [Windows Aero](http://windows.microsoft.com/en-us/windows-vista/what-is-windows-aero) to ensure that all of your development PC's power is available to use with VR.
 * From the Steam client make sure to download and install the **SteamVR** tools. You can find and install the tools by doing the following.
  1. From the Steam client go to **Library** > **Tools** then in the list of tools look for **SteamVR**.
  1. **Double - Click** on the **SteamVR** entry to start the tools install process and then follow the directions provide by Steam.
@@ -42,7 +42,7 @@ In the following section we will go over what needs to be setup before you can b
 
 ## What is SteamVR
 
-SteamVR or Steam Virtual Reality is fully featured, 360°, room scale VR experience that allows users to experience, interact, and even move around in virtual worlds. 
+SteamVR or Steam Virtual Reality is fully featured, 360 degree, room scale VR experience that allows users to experience, interact, and even move around in virtual worlds. 
 SteamVR can be utilized in either an active, moving around the room, or passive, sitting in a chair, setting allowing developers to choose which setting will give the best experience to users. 
 SteamVR also makes use of two wireless single handed controllers whose real world  movement, position, and rotation can be precisely tracked and then translated in to VR allowing users the ability to precisely interact with object in VR like never before. 
 If that was not enough, SteamVR makes use of the Lighthouse tracking system which utilizes room scanning lasers to precisely track the real world position, rotation, and movement of the user and then translates that data so that what the user is doing in the real world will match up to the VR world they are experiencing. 
diff --git a/Engine/Documentation/Source/Platforms/SteamVR/QuickStart/2/SteamVRQuickStart_2.INT.udn b/Engine/Documentation/Source/Platforms/SteamVR/QuickStart/2/SteamVRQuickStart_2.INT.udn
--- a/Engine/Documentation/Source/Platforms/SteamVR/QuickStart/2/SteamVRQuickStart_2.INT.udn
+++ b/Engine/Documentation/Source/Platforms/SteamVR/QuickStart/2/SteamVRQuickStart_2.INT.udn
@@ -1,7 +1,7 @@
 Title: 2. Setting up UE4 to work with SteamVR
 Crumbs:%ROOT%, Platforms
 Description: Setting up SteamVR and UE4 to work together
-Availability:Public
+Availability:Doc
 Navigation:platform
 Platform:SteamVR
 
diff --git a/Engine/Documentation/Source/Platforms/SteamVR/QuickStart/SteamVRQuickStart.INT.udn b/Engine/Documentation/Source/Platforms/SteamVR/QuickStart/SteamVRQuickStart.INT.udn
--- a/Engine/Documentation/Source/Platforms/SteamVR/QuickStart/SteamVRQuickStart.INT.udn
+++ b/Engine/Documentation/Source/Platforms/SteamVR/QuickStart/SteamVRQuickStart.INT.udn
@@ -1,7 +1,7 @@
 Title: SteamVR Quick Start
 Crumbs:%ROOT%, Platforms
 Description: A How To guide for setting up SteamVR to work with UE4
-Availability:Public
+Availability:Doc
 Navigation:platform
 Platform:SteamVR
 
diff --git a/Engine/Documentation/Source/Platforms/SteamVR/SteamVR.INT.udn b/Engine/Documentation/Source/Platforms/SteamVR/SteamVR.INT.udn
--- a/Engine/Documentation/Source/Platforms/SteamVR/SteamVR.INT.udn
+++ b/Engine/Documentation/Source/Platforms/SteamVR/SteamVR.INT.udn
@@ -12,7 +12,7 @@ Platform:SteamVR
 		![](%ROOT%/SteamVR_Icon.png)(convert:false)
 	[/PARAM]
 	[PARAM:title]
-		Oculus Rift Development
+		SteamVR Development
 	[/PARAM]
 	[PARAM:description]
 	[/PARAM]
diff --git a/Engine/Documentation/Source/Platforms/SteamVR/Tps&Tricks/SteamVRTips&Tricks.INT.udn b/Engine/Documentation/Source/Platforms/SteamVR/Tps&Tricks/SteamVRTips&Tricks.INT.udn
--- a/Engine/Documentation/Source/Platforms/SteamVR/Tps&Tricks/SteamVRTips&Tricks.INT.udn
+++ b/Engine/Documentation/Source/Platforms/SteamVR/Tps&Tricks/SteamVRTips&Tricks.INT.udn
@@ -25,23 +25,6 @@ To stop this from happening you can do the following to keep the compositor runn
 	![](T_KeepAlive_CLA.png)
 1. Now when opening up the UE4 editor you should no longer have a flickering problem.
 
-### VR Interaction Area
-
-
-
-
-
-
-
-
-
-
-
-
-
-
-
-
 
 
 
