#UE4doc #docs:
Changed a few words to match the changes in terminology used by the engine.

diff --git a/Engine/Documentation/Source/Engine/Animation/RetargetingDifferentSkeletons/RetargetingDifferentSkeletons.INT.udn b/Engine/Documentation/Source/Engine/Animation/RetargetingDifferentSkeletons/RetargetingDifferentSkeletons.INT.udn
--- a/Engine/Documentation/Source/Engine/Animation/RetargetingDifferentSkeletons/RetargetingDifferentSkeletons.INT.udn
+++ b/Engine/Documentation/Source/Engine/Animation/RetargetingDifferentSkeletons/RetargetingDifferentSkeletons.INT.udn
@@ -15,11 +15,11 @@ Note: This feature is not available before Engine Version 4.5
 [/REGION]
 
 
-**Animation Retargeting (Different Skeletons)** is a feature that allows animations to be reused between characters that **do not** share the same Skeleton asset, unlike [](Engine\Animation\AnimationRetargeting) which involves retargeting animations between characters that share the same Skeleton asset. 
+**Animation Retargeting (Different Skeletons)** is a feature that allows animations to be reused between characters that do not share the same Skeleton asset, unlike [](Engine\Animation\AnimationRetargeting) which involves retargeting animations between characters that share the same Skeleton asset. 
 
-The process of retargeting animations between characters that have different Skeletons uses an asset called a **Rig** to pass Bone information from one Skeleton to another. After selecting your Rig, you will share it between your Source Skeleton and a Target Skeleton (the Skeleton you wish to retarget animations to). Once you have done so, you will then match up the Target Skeleton's **Bones** to the **Nodes** stored in the Rig. 
+The process of retargeting animations between characters that have different Skeletons uses an asset called a **Rig** to pass Bone information from one Skeleton to another. After selecting your Rig, you will share it between your Source Skeleton and a Target Skeleton (the Skeleton you wish to retarget animations to). Once you have done so, you will then match up the Target Skeleton's Bones to the Nodes stored in the Rig. 
 
-Once all the Bones in the Target Skeleton have been setup in the Rig, you can then retarget animations that are used by your Source Skeleton to your Target Skeleton. You can even retarget entire **Animation Blueprints** along with all its associated data, which saves time if you want to quickly move all animation assets from one Skeleton to another. 
+Once all the Bones in the Target Skeleton have been setup in the Rig, you can then retarget animations that are used by your Source Skeleton to your Target Skeleton. You can even retarget entire Animation Blueprints along with all its associated data, which saves time if you want to quickly move all animation assets from one Skeleton to another. 
 
 ## Getting Started
 
@@ -48,7 +48,7 @@ The steps below will show you how to access the default Rig and retarget animati
 
 1.	Inside the **Content Browser** under the **Game/AnimStarterPack/UE4_Mannequin/Mesh** folder, open the **UE4_Mannequin_Skeleton** asset. 
 
-1.	Click the **Retarget Manager** button, then under **Set up Rig**, click the **Pick Asset** icon and select **Set Humanoid Rig**. 
+1.	Click the **Retarget Manager** button, then under **Set up Rig**, click the **Select Rig** dropdown and select **Select Humanoid Rig**. 
 
 	[REGION:lightbox]
 	[![](RT_Manager.jpg)(w:720)](RT_Manager.jpg)
@@ -60,7 +60,7 @@ The steps below will show you how to access the default Rig and retarget animati
 
 1.	In the **Content Browser** under the **Game/MixamoAnimPack/Mixamo_SWAT/Anims** folder, open the **SWAT_Skeleton** asset.
 
-1.	Click the **Retarget Manager** button, then under **Set up Rig**, click the **Pick Asset** icon and select **Set Humanoid Rig**. 
+1.	Click the **Retarget Manager** button, then under **Set up Rig**, click the **Select Rig** dropdown and select **Select Humanoid Rig**.
 
 	[REGION:lightbox]
 	[![](RetargetAdd3.png)(w:720)](RetargetAdd3.png)
