Removing the & out of the folder name as it was causing things to not work correctly.

diff --git a/Engine/Documentation/Source/Platforms/SteamVR/Tips_And_Tricks/SteamVRTips&Tricks.INT.udn b/Engine/Documentation/Source/Platforms/SteamVR/Tips_And_Tricks/SteamVRTips&Tricks.INT.udn
new file mode 100644
--- /dev/null
+++ b/Engine/Documentation/Source/Platforms/SteamVR/Tips_And_Tricks/SteamVRTips&Tricks.INT.udn
@@ -0,0 +1,34 @@
+Title: SteamVR Tips & Tricks
+Crumbs:%ROOT%, Platforms
+Description: A collection of tips & tricks to help you unlock the full power of SteamVR
+Availability:Public
+Navigation:platform
+Platform:SteamVR
+
+## Overview 
+
+In the following section we will cover some tips and tricks that you can use to help unlock the full potential of UE4 and SteamVR.
+
+
+### Avoiding Display Flickering
+
+Each time the editor is launched the display will flicker when the SteamVR compositor takes over. 
+To stop this from happening you can do the following to keep the compositor running in the background at all times.
+
+1. Locate the **vrcompositor.exe** file that should be in **C:\ProgramFiles(x86)\Steam\steamapps\common\SteamVR\bin\win64**.
+	
+	![](T_Comp_EXE.png)
+
+1. Create a shortcut to vrcompositor.exe and then send that shortcut to your desktop.
+1. Right click on the shortcut that was placed on your desktop and then click on **Properties**.
+1. Inside of the **Shortcut** menu in the **Target** input add **-keepalive** after the .exe and press the **OK** button when done.  
+	![](T_KeepAlive_CLA.png)
+1. Now when opening up the UE4 editor you should no longer have a flickering problem.
+
+
+
+
+
+
+
+
diff --git a/Engine/Documentation/Source/Platforms/SteamVR/Tps&Tricks/SteamVRTips&Tricks.INT.udn b/Engine/Documentation/Source/Platforms/SteamVR/Tps&Tricks/SteamVRTips&Tricks.INT.udn
deleted file mode 100644
