Adjusting the link to the tips and tricks section as it was broken during the re-name.

diff --git a/Engine/Documentation/Source/Platforms/SteamVR/SteamVR.INT.udn b/Engine/Documentation/Source/Platforms/SteamVR/SteamVR.INT.udn
--- a/Engine/Documentation/Source/Platforms/SteamVR/SteamVR.INT.udn
+++ b/Engine/Documentation/Source/Platforms/SteamVR/SteamVR.INT.udn
@@ -19,7 +19,7 @@ Platform:SteamVR
 	[PARAM:links]
 		* [](Platforms/SteamVR/BestPractices "%Platforms/SteamVR/BestPractices:description%")
 		* [](Platforms/SteamVR/QuickStart "%Platforms/SteamVR/QuickStart:description%")
-		* [](Platforms/SteamVR/Tps&Tricks "%Platforms/SteamVR/Tps&Tricks:description%")
+		* [](Platforms/SteamVR/Tips_And_Tricks "%Platforms/SteamVR/Tips_And_Tricks:description%")
 	[/PARAM]
 [/OBJECT]
 
