Fixing grammatical errors from forum thread

diff --git a/Engine/Documentation/Source/Engine/Rendering/Materials/HowTo/EmissiveGlow/EmissiveGlow.INT.udn b/Engine/Documentation/Source/Engine/Rendering/Materials/HowTo/EmissiveGlow/EmissiveGlow.INT.udn
--- a/Engine/Documentation/Source/Engine/Rendering/Materials/HowTo/EmissiveGlow/EmissiveGlow.INT.udn
+++ b/Engine/Documentation/Source/Engine/Rendering/Materials/HowTo/EmissiveGlow/EmissiveGlow.INT.udn
@@ -138,13 +138,13 @@ Texture exactly which parts should be illuminated. In the example above, we used
 
 ## Emissive Materials in Action
 
-1. With the room now setup, it is now time to add some meshes for testing.
+1. With the room setup, it is now time to add some meshes for testing.
 	Inside of the **Starter Content** folder, there is a folder labeled **Shapes**. Inside of that folder, you will find a mesh named **Shape_Cube**. Select the **Shape_Cube** with your 
 	**Left Mouse Button** and then drag the mesh from the **Content Browser** into the level releasing the **Left Mouse Button** when the mesh is over the level viewport. This will place the mesh in the level.
 	
 	![](EM_Adding_Test_Box.png)
 
-1. Now that the mesh has been added to level, it is now time to apply the Emissive Material to it. To do this, locate the **Emissive Material** folder in the **Content Browser**. 
+1. Now that the mesh has been added to level, it is time to apply the Emissive Material to it. To do this, locate the **Emissive Material** folder in the **Content Browser**. 
 	Once you find it, select the **Material Instance** with the **Left Mouse Button** and then drag the Material Instance from the **Content Browser** onto the **Shape_Cube** mesh.
 	
 	![](EM_Applying_Material_Instance.png)
diff --git a/Engine/Documentation/Source/Engine/Rendering/Materials/HowTo/Fresnel/Fresnel.INT.udn b/Engine/Documentation/Source/Engine/Rendering/Materials/HowTo/Fresnel/Fresnel.INT.udn
--- a/Engine/Documentation/Source/Engine/Rendering/Materials/HowTo/Fresnel/Fresnel.INT.udn
+++ b/Engine/Documentation/Source/Engine/Rendering/Materials/HowTo/Fresnel/Fresnel.INT.udn
@@ -15,7 +15,7 @@ Related: Engine/Rendering/Materials/ExpressionReference/Utility
 
 Often times when creating a film or cinematic, you need a way to better highlight the silhouette of a character. This lighting technique is often referred to as rim lighting, or edge lighting, and 
 involves adding additional lights to help better distinguish the Characters silhouette from the background. While this works quite nicely for film, adding additional lights to any 3D scene adds 
-extra complexity that can quickly become hard to manage. Unreal Engine 4 (UE4) offers a solution to this by allowing artist to simulate the same basic effect using the Material of the Character. 
+extra complexity that can quickly become difficult to manage. Unreal Engine 4 (UE4) offers a solution to this by allowing artist to simulate the same basic effect using the Material of the Character. 
 This way, the artist has more control over the look and feel of the effect and can even use it for things like adding lighting to an unlit object to save on performance. 
 
 
@@ -53,7 +53,7 @@ It is **highly recommended** that you use this version to begin with and then mo
 |--------- | ---------- |
 | **ExponentIn** | This controls the fall off of the Fresnel Effect. |
 | **BaseReflectFrctionIn** | This specifies the fraction of specular reflection when the surface is viewed from straight on. Setting this to a value of 1 effectively disables Fresnel |
-| **Normal** | Here you can input a Normal to effect the way the Fresnel effect is rendered. This property cannot be set by supplying either a Normal map or a Vector 3 input.|
+| **Normal** | Here you can input a Normal to affect the way the Fresnel effect is rendered. This property cannot be set by supplying either a Normal map or a Vector 3 input.|
 
 You have two options for adjusting the values for the Fresnel Material Expression.
 * Select the Fresnel Material Expression in the Material Graph and then input values via the **Details** Panel.
@@ -77,7 +77,7 @@ This node should only be used if the Fresnel that is under **Utility** does not
 | Property | Description|
 |--------- | ---------- |
 | **Normal Vector (V3)** |  Here you can input a Normal to effect the way the Fresnel effect is rendered. This property cannot be set by supplying either a Normal map or a Vector 3 input. |
-| **Camera Vector(V3)** |This allows you to effect the Camera that is used to determine the Dot Product between the Camera and the surface normal. |
+| **Camera Vector(V3)** |This allows you to affect the Camera that is used to determine the Dot Product between the Camera and the surface normal. |
 | **Invert Fresnel(B)** | This will allow you to invert the effect of the Fresnel. This is useful if you want the Fresnel to only be applied in the center or at the edges. |
 | **Power(S)** | This allows you to control the power of the Fresnel effect. Entering bigger numbers will tighten up the areas affected. Smaller numbers will allow more areas to be affected by the Fresnel.  |
 | **Use Cheap Contrast(B)** | This will allow the Fresnel to use a cheaper method for calculating contrast between lights and darks. |
@@ -114,7 +114,7 @@ Creating a Material that uses Fresnel can be accomplished in the following steps
 	When you search for the Fresnel Material Expression, you will find two different types that you can choose from. For this How - To, we will only be using the Fresnel that can be found under **Utility**.
 	[/region]
 
-1. Now that the required Material Expression nodes have been added to the Material Graph, it is now time to start connecting everything together and giving the parameters base values. 
+1. Now that the required Material Expression nodes have been added to the Material Graph, it is time to start connecting everything together and giving the parameters base values. 
 	For this example, we are going to be using the following Parameter names and default values. 
 	|Name | Value |
 	|-----| ------|
@@ -168,7 +168,7 @@ If it is perpendicular to the camera, then the Fresnel effect will take place. O
 However when a Normal map is introduced, we change the surface normal to no longer be smooth, allowing the Fresnel effect to take happen in places that it normally would not if the 
 surface Normal had been smooth. 
 
-For a visual example at what is going on, take a look at the following image showing the effect the surface normal has on the Fresnel. 
+For a visual example of what is going on, take a look at the following image showing the effect the surface normal has on the Fresnel. 
 On the left (image 1) is what the Fresnel node will output with a flat surface normal. Now take a look at the right hand side (image 2.)
 In this image, a Normal map was input which drastically changes the surface normal. 
 Because the surface normal is no longer flat everywhere, you get more of the Fresnel being displayed. 
@@ -181,7 +181,7 @@ In the image below, we can see how different Normal maps can affect the Fresnel
 
 The Material that was used to create the images above was the same Material that was created before but with the following modifications made. 
 
-* The Materials **Blend Mode** was changed from **Default** to **Translucent**
+* The Materials **Blend Mode** was changed from **Opaque** to **Translucent**
 * The final output for the Fresnel that was plugged into the **Emissive Color** input was also plugged into the **Opacity** input.
 * A new 2D Texture Sampler was added and plugged into the **Normal** input. **Tech_Hex_Tile_N** from the Starter Content was used as the Normal texture.
 * A new Scalar Parameter was added that is multiplied by a **Texture Coordinates** expression to implement some basic uniform titling on the Normal map. 
diff --git a/Engine/Documentation/Source/Engine/Rendering/Materials/HowTo/Masking/Masking.INT.udn b/Engine/Documentation/Source/Engine/Rendering/Materials/HowTo/Masking/Masking.INT.udn
--- a/Engine/Documentation/Source/Engine/Rendering/Materials/HowTo/Masking/Masking.INT.udn
+++ b/Engine/Documentation/Source/Engine/Rendering/Materials/HowTo/Masking/Masking.INT.udn
@@ -39,7 +39,7 @@ In the following image, we can see this in action.
 
 On the right hand side, outlined in green, is what the Mask Texture looks like. 
 Notice how only a few areas are painted White while all other areas are painted Black. 
-This is because the we only want to affect the areas that are in White.
+This is because we only want to affect the areas that are in White.
 On the left hand side, you can see what areas this Mask Texture will mask out when applied to the chair mesh.
 
 [region:note]
@@ -97,7 +97,7 @@ In this Material, the Mask Texture,**T_Chair_M**, is used to control every aspec
 From what sections should be metal or non-metal to what color those sections should have, the Mask Texture helps to define all of this.
 In the following image, you can see a breakdown of how each channel of the Mask Texture for SM_Chair works.
 On the Left-Hand side is what the Mask looks like as a Texture.
-The smaller images that you see on the Right-Hand side show each of the individual channels and what area those channels will effect on the chair Static Mesh. 
+The smaller images that you see on the Right-Hand side show each of the individual channels and what area those channels will affect on the chair Static Mesh. 
 
 ![](TM_Mask_Example.png)
 
