#UE4doc #docs - Replaced comparison images with Comparison Slider.

diff --git a/Engine/Documentation/Source/Engine/Rendering/LightingAndShadows/DistanceFieldAmbientOcclusion/DistanceFieldAmbientOcclusion.INT.udn b/Engine/Documentation/Source/Engine/Rendering/LightingAndShadows/DistanceFieldAmbientOcclusion/DistanceFieldAmbientOcclusion.INT.udn
--- a/Engine/Documentation/Source/Engine/Rendering/LightingAndShadows/DistanceFieldAmbientOcclusion/DistanceFieldAmbientOcclusion.INT.udn
+++ b/Engine/Documentation/Source/Engine/Rendering/LightingAndShadows/DistanceFieldAmbientOcclusion/DistanceFieldAmbientOcclusion.INT.udn
@@ -116,11 +116,14 @@ Foliage with LODs (Level of Detail models) can run into some problems with Dista
 
 To combat this you can use World Position Offset on those LODs to pull the vertices outside of the distance field.  Typically a small offset toward the camera is effective.  For billboards, use the new Pixel Depth Offset feature to create valid depth values which better represent the original 3d triangle mesh.  
 
-[REGION:lightbox]
-
-[![](0Billboard.jpg)(w:600)](0Billboard.jpg) [![](0BillboardWithPixelDepthOffset.jpg)(w:600)](0BillboardWithPixelDepthOffset.jpg)
-
-[/REGION]
+[OBJECT:ComparisonSlider]
+	[PARAM:before]
+	![No Pixel Depth Offset ](0Billboard.jpg)(w:600)
+	[/PARAM]
+	[PARAM:after]
+	![Pixel Depth Offset](0BillboardWithPixelDepthOffset.jpg)(w:600)
+	[/PARAM]
+[/OBJECT]
 _Distant tree billboards showing over-occlusion of sky lighting, which is improved with the use of Pixel Depth Offset._
 
 ## How does it work?
@@ -139,7 +142,15 @@ The second property that makes distance fields useful is that when you trace a r
 
 The entire level is represented by instanced distance fields stored in volume textures.  Landscape has a separate representation through heightfields.
 
-[![](TriangleScene.jpg)(w:400)](TriangleScene.jpg) [![](MeshDistanceFieldScene.jpg)(w:400)](MeshDistanceFieldScene.jpg)
+[OBJECT:ComparisonSlider]
+	[PARAM:before]
+	![Scene View](TriangleScene.jpg)(w:600)
+	[/PARAM]
+	[PARAM:after]
+	![Mesh Distance Fields View](MeshDistanceFieldScene.jpg)(w:600)
+	[/PARAM]
+[/OBJECT]
+
 
 [**Further reading on using distance fields for lighting.**](http://iquilezles.org/www/articles/raymarchingdf/raymarchingdf.htm)
 
