#UE4doc #docs - Re-organized images to be in image tables and updated outdated images.

diff --git a/Engine/Documentation/Source/Engine/Rendering/LightingAndShadows/RayTracedDistanceFieldShadowing/RayTracedDistanceFieldShadowing.INT.udn b/Engine/Documentation/Source/Engine/Rendering/LightingAndShadows/RayTracedDistanceFieldShadowing/RayTracedDistanceFieldShadowing.INT.udn
--- a/Engine/Documentation/Source/Engine/Rendering/LightingAndShadows/RayTracedDistanceFieldShadowing/RayTracedDistanceFieldShadowing.INT.udn
+++ b/Engine/Documentation/Source/Engine/Rendering/LightingAndShadows/RayTracedDistanceFieldShadowing/RayTracedDistanceFieldShadowing.INT.udn
@@ -92,19 +92,15 @@ Cascaded Shadow Maps are typically used to provide dynamic shadowing of a direct
 
 Ray Traced Distance Field shadows behave much more gracefully in the distance, doing shadowing work only for visible pixels.  Cascaded Shadow Maps can be used to cover regions near the camera, while Ray Traced Shadows will cover regions up until 'DistanceField Shadow Distance'.
 
-[REGION:fullwidth]
-	![](CSM2.jpg)
-[/REGION]
-[REGION:caption]
-A foliage scene with CSM covering up to 100 meters.
-[/REGION]
+[OBJECT:ComparisonSlider]
+	[PARAM:before]
+	![A foliage scene with CSM covering up to 100 meters.](CSM2.jpg)(w:800)
+	[/PARAM]
+	[PARAM:after]
+	![A foliage scene with CSM covering up to 100 meters, and Ray Traced Distance Field shadows covering up to 1.2 kilometers.](RTDF2.jpg)(w:800)
+	[/PARAM]
+[/OBJECT]
 
-[REGION:fullwidth]
-	![](RTDF2.jpg)
-[/REGION]
-[REGION:caption]
-A foliage scene with CSM covering up to 100 meters, and Ray Traced Distance Field shadows covering up to 1.2 kilometers.
-[/REGION]
 
 [_Video of changing time of day using Ray Traced Distance Field shadowing._](https://www.youtube.com/watch?v=4249b94KtyA&list=UUBobmJyzsJ6Ll7UbfhI4iwQ)
 
