#UE4doc #docs - Added engine version tag

diff --git a/Engine/Documentation/Source/Engine/Physics/Physics.INT.udn b/Engine/Documentation/Source/Engine/Physics/Physics.INT.udn
--- a/Engine/Documentation/Source/Engine/Physics/Physics.INT.udn
+++ b/Engine/Documentation/Source/Engine/Physics/Physics.INT.udn
@@ -2,6 +2,7 @@ Availability: Public
 Title:Physics Simulation
 Crumbs: %ROOT%, Engine
 Description:Subsystem used to calculate collision and simulate physical actors.
+Version: 4.9
 
 [VAR:Topic]
 [OBJECT:Topic]
