#UE4doc #docs - Added engine version and skill level tags, and updated images.

diff --git a/Engine/Documentation/Source/Engine/Rendering/LightingAndShadows/Basics/PlacingLights.INT.udn b/Engine/Documentation/Source/Engine/Rendering/LightingAndShadows/Basics/PlacingLights.INT.udn
--- a/Engine/Documentation/Source/Engine/Rendering/LightingAndShadows/Basics/PlacingLights.INT.udn
+++ b/Engine/Documentation/Source/Engine/Rendering/LightingAndShadows/Basics/PlacingLights.INT.udn
@@ -2,6 +2,8 @@ Availability: Public
 Title:Lighting Basics
 Crumbs: %ROOT%, Engine, Engine/Rendering/LightingAndShadows
 Description:The basics of placing lights into levels and setting them up.
+SkillLevel: Beginner
+Version: 4.9
 
 [VAR:playlistnav]
 [OBJECT:PlaylistNavigation]
