#UE4doc #docs - Added engine version tag.

diff --git a/Engine/Documentation/Source/Engine/Rendering/LightingAndShadows/LightMobility/LightMobility.INT.udn b/Engine/Documentation/Source/Engine/Rendering/LightingAndShadows/LightMobility/LightMobility.INT.udn
--- a/Engine/Documentation/Source/Engine/Rendering/LightingAndShadows/LightMobility/LightMobility.INT.udn
+++ b/Engine/Documentation/Source/Engine/Rendering/LightingAndShadows/LightMobility/LightMobility.INT.udn
@@ -2,6 +2,7 @@ Availability: Public
 Title:Light Mobility
 Crumbs: %ROOT%, Engine, Engine/Rendering/LightingAndShadows
 Description:Light Mobility and how it affects lighting and shadowing.
+Version: 4.9
 
 [VAR:Topic]
 [OBJECT:Topic]
