#UE4doc #docs
Minor revision to this section.
Availability: Docs

diff --git a/Engine/Documentation/Source/Programming/Tutorials/FirstPersonShooter/2/5/FirstPersonShooter_2_5.INT.udn b/Engine/Documentation/Source/Programming/Tutorials/FirstPersonShooter/2/5/FirstPersonShooter_2_5.INT.udn
--- a/Engine/Documentation/Source/Programming/Tutorials/FirstPersonShooter/2/5/FirstPersonShooter_2_5.INT.udn
+++ b/Engine/Documentation/Source/Programming/Tutorials/FirstPersonShooter/2/5/FirstPersonShooter_2_5.INT.udn
@@ -227,7 +227,7 @@ Go back to Visual Studio to add code to your FPSCharacter class.
 
 ## Testing Character Jumping
 
-It's now time to compile and test our newly implemented character movement functions.
+It's now time to compile and test your newly implemented character movement functions.
 
 1.	Save the `FPSCharacter` interface (\*.h) and implementation (\*.cpp) files in Visual Studio.
 
