#UE4doc #docs
Completed first draft of the 'Adding a Mesh to your Character' section.
Availability: Docs

diff --git a/Engine/Documentation/Source/Programming/Tutorials/FirstPersonShooter/2/6/FirstPersonShooter_2_6.INT.udn b/Engine/Documentation/Source/Programming/Tutorials/FirstPersonShooter/2/6/FirstPersonShooter_2_6.INT.udn
--- a/Engine/Documentation/Source/Programming/Tutorials/FirstPersonShooter/2/6/FirstPersonShooter_2_6.INT.udn
+++ b/Engine/Documentation/Source/Programming/Tutorials/FirstPersonShooter/2/6/FirstPersonShooter_2_6.INT.udn
@@ -23,6 +23,155 @@ Version: 4.9
 
 %Steps%
 
-<!--TODO: Add tutorial steps here-->
+During this step, you'll give your character a skeletal mesh. The `Character` class creates a `SkeletalMeshComponent` object for us by default, so all it needs to know is which `SkeletalMesh` asset to use.
+
+[REGION:note]
+Before getting started, download and extract our sample mesh from the following link:
+*	[Sample Mesh](GenericMale.zip "Sample Mesh")
+[/REGION]
+
+## Importing a Skeletal Mesh
+
+1.	**Right-click** inside the Content Browser's file box to open the **Import Asset** dialog box
+
+1.	**Click 'Import to /Game...'** to open the **Import** dialog box.
+
+	![](RightClickImport.png)
+
+1.	Locate and select the **GenericMale.fbx** mesh file.
+
+1.	**Click Open** to begin importing the mesh to your project.
+
+1.	In the **Content Browser**, the **FBX Import Options** dialog box appears. **Clicking Import** adds your mesh to the Project.
+
+1.	**Click** the **Save** button to save your imported mesh.
+
+## Creating a Blueprint for an FPS Character
+
+1.	**Right-click** in the **Content Browser** and select **New Folder**. 
+
+1.	Name this new folder "Blueprints".
+
+1.	**Double-click** on the folder to open it.
+
+1.	**Click** the **Add New** button and select **Blueprint Class**.
+
+	![](AddNewBPClass.png)
+
+1.	Expand the **All Classes** dropdown menu and type "FPSCharacter" into the search box.
+
+	![](AddFPSCharacterClass.png)
+
+1.	**Click** on the **FPSCharacter** and **click** the **Select** button.
+
+1.	Name the new Blueprint "BP_FPSCharacter".
+
+	![](NameBP_FPSCharacter.png)
+
+## Setting up a Third Person Mesh
+
+1.	 **Double-click** the **BP_FPSCharacter** icon to open it in the **Blueprint Editor**.
+
+1.	**Click** on the **Mesh** component in the **Components** tab.
+
+	![](MeshComponent.png)
+
+1.	Scroll down to the **Mesh** section of the **Details** tab and **click** on the dropdown menu that says "None".
+
+1.	Select the **GenericMale** skeletal mesh.
+
+	![](SelectSkeletalMesh.png)
+
+1.	Align the `SkeletalMeshComponent` to the `CapsuleComponent` by setting its Z Location to -88 in the **Details** tab.
+
+	![](NewZLoc_Mesh.png)
+
+	[REGION:note]
+	You can also move components around with widgets in the **Preview Viewport** rather than setting values in the **Details** tab.
+	[/REGION]
+
+1.	The `SkeletalMeshComponent` should now resemble the following:
+
+	![](SkeletalMeshLocation.png)
+
+	[REGION:note]
+	Generally, you should strive to contain your `SkeletalMeshComponent` within the `CapsuleComponent`, facing the same direction that the `ArrowComponent` is pointing, ensuring that your character moves correctly through the world.
+	[/REGION]
+
+1.	 Make sure to **Compile** and **Save** the **BP_FPSCharacter** Blueprint before closing the **Blueprint Editor**.
+
+## Setting the Default Pawn Class
+
+1.	Switch back to Visual Studio and open `FPSGameMode.cpp`.
+
+1.	Find the `FPSGameMode` constructor in `FPSGameMode.cpp` to locate the existing `DefaultPawnClass` assignment:
+
+		DefaultPawnClass = AFPSCharacter::StaticClass();
+
+1.	Replace the `DefaultPawnClass` assignment with the following code:
+
+		// set the default pawn class to the Blueprinted character
+		static ConstructorHelpers::FClassFinder<APawn> PlayerPawnObject(TEXT("Pawn'/Game/Blueprints/BP_FPSCharacter.BP_FPSCharacter_C'"));
+		if (PlayerPawnObject.Class != NULL)
+		{
+			DefaultPawnClass = PlayerPawnObject.Class;
+		}
+
+	[REGION:note]
+	This code will find the class generated by your blueprint and assign it as your default pawn class. Note that the "_C" suffix in the asset path distinguishes the actual class used by the game from the Blueprint asset. If you put your Blueprint in a different folder in the asset tree, you can get the full path by right-clicking on the asset in the Content Browser and selecting Copy Reference.
+	![](CopyAssetReference.png)
+	[/REGION]
+
+1.	At this point, you can remove `#include "FPSCharacter.h"` from the top of `FPSGameMode.cpp`. `FPSGameMode.cpp` should now look like the following:
+
+		// Fill out your copyright notice in the Description page of Project Settings.
+		
+		#include "FPSProject.h"
+		#include "FPSGameMode.h"
+		
+		AFPSGameMode::AFPSGameMode(const FObjectInitializer& ObjectInitializer) : Super(ObjectInitializer)
+		{
+			// set the default pawn class to the Blueprinted character
+			static ConstructorHelpers::FClassFinder<APawn> PlayerPawnObject(TEXT("Pawn'/Game/Blueprints/BP_FPSCharacter.BP_FPSCharacter_C'"));
+			if (PlayerPawnObject.Class != NULL)
+			{
+				DefaultPawnClass = PlayerPawnObject.Class;
+			}
+		}//AFPSGameMode Constructor
+		
+		void AFPSGameMode::StartPlay()
+		{
+			Super::StartPlay();
+		
+			StartMatch();
+		
+			if (GEngine != nullptr)
+			{
+				GEngine->AddOnScreenDebugMessage(-1, 5.f, FColor::Yellow, TEXT("We are using FPSGameMode."));
+			}
+		}//StartPlay()
+
+## Verifying the New Mesh in PIE Mode
+
+It's now time to compile and test your newly added mesh in-editor.
+
+1.	Save `FPCSGameMode.cpp` in Visual Studio.
+
+1.	Locate **FPSProject** in the **Solution Explorer**.
+
+1.	Right-click on **FPSProject** and select **Build** to compile your project.
+
+	![](BuildProject.png)
+
+1.	After the build finishes, open your **FPSProject** in Unreal Editor.
+
+1.	Click the **Play** button in the Level Editor Toolbar. You should be able to see your character's shadow as you move around.
+
+	[REGION:note]
+	To freely move the camera to see your character's mesh, press F1 to regain control of your mouse cursor, then click the **Eject** button in the toolbar.
+	![](FreelyMovingCamera.png)
+	[/REGION]
+
+1.	Click the **Stop** button in the Level Editor to exit Play in Editor (PIE) mode.
 
 %Steps%
\ No newline at end of file
