#UE4doc #docs
Completed first draft of the 'Changing the Camera View' section.
Availability: Docs

diff --git a/Engine/Documentation/Source/Programming/Tutorials/FirstPersonShooter/2/7/FirstPersonShooter_2_7.INT.udn b/Engine/Documentation/Source/Programming/Tutorials/FirstPersonShooter/2/7/FirstPersonShooter_2_7.INT.udn
--- a/Engine/Documentation/Source/Programming/Tutorials/FirstPersonShooter/2/7/FirstPersonShooter_2_7.INT.udn
+++ b/Engine/Documentation/Source/Programming/Tutorials/FirstPersonShooter/2/7/FirstPersonShooter_2_7.INT.udn
@@ -23,6 +23,228 @@ Version: 4.9
 
 %Steps%
 
-<!--TODO: Add tutorial steps here-->
+At the end of the previous step, the default camera was positioned inside the mesh's neck. During this step, you're going to set-up with a proper FPS camera that you can use to adjust the camera's properties (such as location and field of view).
+
+## Attaching a Camera Component
+
+1.	Switch back to Visual Studio and open `FPSCharacter.h` and add the following public property:
+
+		// FPS camera
+		UPROPERTY(VisibleAnywhere, BlueprintReadOnly, Category = Camera)
+		UCameraComponent* FirstPersonCameraComponent;
+
+1.	Add the following constructor to your `FPSCharacter.h` file:
+
+		// Constructor for the AFPSCharacter
+		AFPSCharacter(const FObjectInitializer& ObjectInitializer);
+
+1.	`FPSCharacter.h` should now look like the following:
+
+		// Fill out your copyright notice in the Description page of Project Settings.
+		
+		#pragma once
+		
+		#include "GameFramework/Character.h"
+		#include "FPSCharacter.generated.h"
+		
+		UCLASS()
+		class FPSPROJECT_API AFPSCharacter : public ACharacter
+		{
+			GENERATED_BODY()
+		
+		public:
+			// Sets default values for this character's properties
+			AFPSCharacter();
+		
+			// Constructor for the AFPSCharacter
+			AFPSCharacter(const FObjectInitializer& ObjectInitializer);
+		
+			// Called when the game starts or when spawned
+			virtual void BeginPlay() override;
+			
+			// Called every frame
+			virtual void Tick( float DeltaSeconds ) override;
+		
+			// Called to bind functionality to input
+			virtual void SetupPlayerInputComponent(class UInputComponent* InputComponent) override;
+		
+			// handles moving forward and backward
+			UFUNCTION()
+			void MoveForward(float Value);
+		
+			// handles moving right and left
+			UFUNCTION()
+			void MoveRight(float Value);
+		
+			// sets jump flag when key is pressed
+			UFUNCTION()
+			void OnStartJump();
+			
+			// clears jump flag when key is released
+			UFUNCTION()
+			void OnStopJump();
+		
+			// FPS camera
+			UPROPERTY(VisibleAnywhere, BlueprintReadOnly, Category = Camera)
+			UCameraComponent* FirstPersonCameraComponent;
+		};
+
+1.	Open `FPSCharacter.cpp` and add the following code:
+
+		// Constructor that creates the camera component
+		AFPSCharacter::AFPSCharacter(const FObjectInitializer& ObjectInitializer) : Super(ObjectInitializer)
+		{
+			// Create a camera Component
+			FirstPersonCameraComponent = ObjectInitializer.CreateDefaultSubobject<UCameraComponent>(this, TEXT("FirstPersonCamera"));
+		
+			if (FirstPersonCameraComponent != NULL)
+			{
+				FirstPersonCameraComponent->AttachParent = GetCapsuleComponent();
+			}
+		}
+
+	[REGION:note]
+	The aforementioned code creates a `CameraComponent` and attaches it to the `CapsuleComponent`.
+	[/REGION]
+
+1.	Add the following code to `AFPSCharacter::AFPSCharacter(const FObjectInitializer& ObjectInitializer)`:
+
+		// Position the camera a bit above the eyes
+		FirstPersonCameraComponent->RelativeLocation = FVector(0.f, 0.f, 50.f + BaseEyeHeight);
+		// Allow the pawn to control rotation
+		FirstPersonCameraComponent->bUsePawnControlRotation = true;
+
+	[REGION:note]
+	The aforementioned code adjusts the camera's position to be slightly above the character's eye position.
+	[/REGION]
+
+1.	`FPSCharacter.cpp` should now look like the following:
+
+		// Fill out your copyright notice in the Description page of Project Settings.
+		
+		#include "FPSProject.h"
+		#include "FPSCharacter.h"
+		
+		
+		// Sets default values
+		AFPSCharacter::AFPSCharacter()
+		{
+		 	// Set this character to call Tick() every frame.  You can turn this off to improve performance if you don't need it.
+			PrimaryActorTick.bCanEverTick = true;
+		
+		}
+		
+		// Constructor that creates the camera component
+		AFPSCharacter::AFPSCharacter(const FObjectInitializer& ObjectInitializer) : Super(ObjectInitializer)
+		{
+			// Create a camera Component
+			FirstPersonCameraComponent = ObjectInitializer.CreateDefaultSubobject<UCameraComponent>(this, TEXT("FirstPersonCamera"));
+		
+			if (FirstPersonCameraComponent != NULL)
+			{
+				FirstPersonCameraComponent->AttachParent = GetCapsuleComponent();
+				// Position the camera a bit above the eyes
+				FirstPersonCameraComponent->RelativeLocation = FVector(0.f, 0.f, 50.f + BaseEyeHeight);
+				// Allow the pawn to control rotation
+				FirstPersonCameraComponent->bUsePawnControlRotation = true;
+			}
+		}
+		
+		// Called when the game starts or when spawned
+		void AFPSCharacter::BeginPlay()
+		{
+			Super::BeginPlay();
+			
+			if (GEngine != nullptr)
+			{
+				GEngine->AddOnScreenDebugMessage(-1, 5.f, FColor::White, TEXT("We are using FPSCharacter."));
+			}
+		}
+		
+		// Called every frame
+		void AFPSCharacter::Tick( float DeltaTime )
+		{
+			Super::Tick( DeltaTime );
+		
+		}
+		
+		// Called to bind functionality to input
+		void AFPSCharacter::SetupPlayerInputComponent(class UInputComponent* InputComponent)
+		{
+			Super::SetupPlayerInputComponent(InputComponent);
+			
+			// set-up gameplay key bindings
+			InputComponent->BindAxis("MoveForward", this, &AFPSCharacter::MoveForward);
+			InputComponent->BindAxis("MoveRight", this, &AFPSCharacter::MoveRight);
+			// set-up gameplay mouse bindings
+			InputComponent->BindAxis("Turn", this, &AFPSCharacter::AddControllerYawInput);
+			InputComponent->BindAxis("LookUp", this, &AFPSCharacter::AddControllerPitchInput);
+			// set-up jump action bindings
+			InputComponent->BindAction("Jump", IE_Pressed, this, &AFPSCharacter::OnStartJump);
+			InputComponent->BindAction("Jump", IE_Released, this, &AFPSCharacter::OnStopJump);
+		}
+		
+		// Forward movement function definition
+		void AFPSCharacter::MoveForward(float Value)
+		{
+			if (Controller != nullptr && Value != 0.f)
+			{
+				//find out which way is forward
+				FRotator Rotation = Controller->GetControlRotation();
+				//limit pitch when walking for falling
+				if (GetCharacterMovement()->IsMovingOnGround() || GetCharacterMovement()->IsFalling())
+				{
+					Rotation.Pitch = 0.f;
+				}
+				//add movement in that direction
+				const FVector Direction = FRotationMatrix(Rotation).GetScaledAxis(EAxis::X);
+				AddMovementInput(Direction, Value);
+			}
+		}
+		
+		//Right movement function definition
+		void AFPSCharacter::MoveRight(float Value)
+		{
+			if (Controller != nullptr && Value != 0.f)
+			{
+				//find our which way is right
+				const FRotator Rotation = Controller->GetControlRotation();
+				const FVector Direction = FRotationMatrix(Rotation).GetScaledAxis(EAxis::Y);
+				//add movement in that direction
+				AddMovementInput(Direction, Value);
+			}
+		}
+		
+		//Function setting the jump flag
+		void AFPSCharacter::OnStartJump()
+		{
+			bPressedJump = true;
+		}
+		
+		//Function clearning the jump flag
+		void AFPSCharacter::OnStopJump()
+		{
+			bPressedJump = false;
+		}
+
+## Testing the New Camera
+
+It's now time to compile and test your newly implemented camera code.
+
+1.	Save the `FPSCharacter` interface (\*.h) and implementation (\*.cpp) files in Visual Studio.
+
+1.	Locate **FPSProject** in the **Solution Explorer**.
+
+1.	Right-click on **FPSProject** and select **Build** to compile your project.
+
+	![](BuildProject.png)
+
+1.	After the build finishes, open your **FPSProject** in Unreal Editor.
+
+1.	Click the **Play** button in the Level Editor Toolbar. Your camera should be above the character's head and if you look down, you should be able to see the top of your character's head.
+
+	![](NewCameraComponent.png)
+
+1.	Click the **Stop** button in the Level Editor to exit Play in Editor (PIE) mode.
 
 %Steps%
\ No newline at end of file
