#UE4doc #docs
Made progress on the 'Add a First Person Mesh to your Character' section.
Availability: Docs

diff --git a/Engine/Documentation/Source/Programming/Tutorials/FirstPersonShooter/2/8/FirstPersonShooter_2_8.INT.udn b/Engine/Documentation/Source/Programming/Tutorials/FirstPersonShooter/2/8/FirstPersonShooter_2_8.INT.udn
--- a/Engine/Documentation/Source/Programming/Tutorials/FirstPersonShooter/2/8/FirstPersonShooter_2_8.INT.udn
+++ b/Engine/Documentation/Source/Programming/Tutorials/FirstPersonShooter/2/8/FirstPersonShooter_2_8.INT.udn
@@ -23,6 +23,251 @@ Version: 4.9
 
 %Steps%
 
-<!--TODO: Add tutorial steps here-->
+A common approach for building FPS games is to use two separate character meshes where one is the full-body mesh and the other is the "weapon and hands" mesh. The full-body mesh is used to see the character from a third-person perspective; however, it is hidden when the player is viewing the game in first-person perspective. The "weapon and hands" mesh is typically attached to the camera and it is visible only to the player when the player is viewing the map in first-person perspective.
+
+## Adding the First Person Character Mesh
+
+1.	Switch back to Visual Studio and open `FPSCharacter.h` and add the following public property:
+
+		// Pawn Mesh: First person view (arms that are only seen by self)
+		UPROPERTY(VisibleDefaultsOnly, Category = Mesh)
+		USkeletalMeshComponent* FirstPersonMesh;
+
+1.	`FPSCharacter.h` should now look like the following:
+
+		// Fill out your copyright notice in the Description page of Project Settings.
+		
+		#pragma once
+		
+		#include "GameFramework/Character.h"
+		#include "FPSCharacter.generated.h"
+		
+		UCLASS()
+		class FPSPROJECT_API AFPSCharacter : public ACharacter
+		{
+			GENERATED_BODY()
+		
+		public:
+			// Sets default values for this character's properties
+			AFPSCharacter();
+		
+			// Constructor for the AFPSCharacter
+			AFPSCharacter(const FObjectInitializer& ObjectInitializer);
+		
+			// Called when the game starts or when spawned
+			virtual void BeginPlay() override;
+			
+			// Called every frame
+			virtual void Tick( float DeltaSeconds ) override;
+		
+			// Called to bind functionality to input
+			virtual void SetupPlayerInputComponent(class UInputComponent* InputComponent) override;
+		
+			// handles moving forward and backward
+			UFUNCTION()
+			void MoveForward(float Value);
+		
+			// handles moving right and left
+			UFUNCTION()
+			void MoveRight(float Value);
+		
+			// sets jump flag when key is pressed
+			UFUNCTION()
+			void OnStartJump();
+			
+			// clears jump flag when key is released
+			UFUNCTION()
+			void OnStopJump();
+		
+			// FPS camera
+			UPROPERTY(VisibleAnywhere, BlueprintReadOnly, Category = Camera)
+			UCameraComponent* FirstPersonCameraComponent;
+		
+			// Pawn Mesh: First person view (arms that are only seen by self)
+			UPROPERTY(VisibleDefaultsOnly, Category = Mesh)
+			USkeletalMeshComponent* FirstPersonMesh;
+		};
+
+1.	Open `FPSCharacter.cpp` and add the following code to create and configure the first person mesh:
+
+		// Create a mesh component that will be used when being viewed from a first person view when controlling this pawn
+		FirstPersonMesh = ObjectInitializer.CreateDefaultSubobject<USkeletalMeshComponent>(this, TEXT("FirstPersonMesh"));
+
+		if (FirstPersonMesh != NULL)
+		{
+			FirstPersonMesh->SetOnlyOwnerSee(true);  // only the owning player will see this mesh
+			FirstPersonMesh->AttachParent = FirstPersonCameraComponent;
+			FirstPersonMesh->bCastDynamicShadow = false;
+			FirstPersonMesh->CastShadow = false;
+		}
+
+	[REGION:note]
+	`SetOnlyOwnerSee` indicates that this mesh is only visible to the `PlayerController` that has possessed this `Character`. The aforementioned code is also attaching the mesh to the camera and disabling some environmental shadowing. Allowing the the camera-attahed arms to cast shadows would destroy the illusion of having a single mesh for the first person character.
+	[/REGION]
+
+1.	Add the following code to `FPSCharacter.cpp` to hide the existing third person-mesh from the owning character:
+
+		// everyone but the owner can see the regular body mesh
+		GetMesh()->SetOwnerNoSee(true);
+
+1.	`FPSCharacter.cpp` should now look like the following:
+
+		// Fill out your copyright notice in the Description page of Project Settings.
+		
+		#include "FPSProject.h"
+		#include "FPSCharacter.h"
+		
+		
+		// Sets default values
+		AFPSCharacter::AFPSCharacter()
+		{
+		 	// Set this character to call Tick() every frame.  You can turn this off to improve performance if you don't need it.
+			PrimaryActorTick.bCanEverTick = true;
+		
+		}
+		
+		// Constructor that creates the camera component
+		AFPSCharacter::AFPSCharacter(const FObjectInitializer& ObjectInitializer) : Super(ObjectInitializer)
+		{
+			// Create a camera Component
+			FirstPersonCameraComponent = ObjectInitializer.CreateDefaultSubobject<UCameraComponent>(this, TEXT("FirstPersonCamera"));
+		
+			if (FirstPersonCameraComponent != NULL)
+			{
+				FirstPersonCameraComponent->AttachParent = GetCapsuleComponent();
+				// Position the camera a bit above the eyes
+				FirstPersonCameraComponent->RelativeLocation = FVector(0.f, 0.f, 50.f + BaseEyeHeight);
+				// Allow the pawn to control rotation
+				FirstPersonCameraComponent->bUsePawnControlRotation = true;
+			}
+		
+			// Create a mesh component that will be used when being viewed from a first person view when controlling this pawn
+			FirstPersonMesh = ObjectInitializer.CreateDefaultSubobject<USkeletalMeshComponent>(this, TEXT("FirstPersonMesh"));
+		
+			if (FirstPersonMesh != NULL)
+			{
+				FirstPersonMesh->SetOnlyOwnerSee(true);  // only the owning player will see this mesh
+				FirstPersonMesh->AttachParent = FirstPersonCameraComponent;
+				FirstPersonMesh->bCastDynamicShadow = false;
+				FirstPersonMesh->CastShadow = false;
+			}
+		
+			// everyone but the owner can see the regular body mesh
+			GetMesh()->SetOwnerNoSee(true);
+		}
+		
+		// Called when the game starts or when spawned
+		void AFPSCharacter::BeginPlay()
+		{
+			Super::BeginPlay();
+			
+			if (GEngine != nullptr)
+			{
+				GEngine->AddOnScreenDebugMessage(-1, 5.f, FColor::White, TEXT("We are using FPSCharacter."));
+			}
+		}
+		
+		// Called every frame
+		void AFPSCharacter::Tick( float DeltaTime )
+		{
+			Super::Tick( DeltaTime );
+		
+		}
+		
+		// Called to bind functionality to input
+		void AFPSCharacter::SetupPlayerInputComponent(class UInputComponent* InputComponent)
+		{
+			Super::SetupPlayerInputComponent(InputComponent);
+			
+			// set-up gameplay key bindings
+			InputComponent->BindAxis("MoveForward", this, &AFPSCharacter::MoveForward);
+			InputComponent->BindAxis("MoveRight", this, &AFPSCharacter::MoveRight);
+			// set-up gameplay mouse bindings
+			InputComponent->BindAxis("Turn", this, &AFPSCharacter::AddControllerYawInput);
+			InputComponent->BindAxis("LookUp", this, &AFPSCharacter::AddControllerPitchInput);
+			// set-up jump action bindings
+			InputComponent->BindAction("Jump", IE_Pressed, this, &AFPSCharacter::OnStartJump);
+			InputComponent->BindAction("Jump", IE_Released, this, &AFPSCharacter::OnStopJump);
+		}
+		
+		// Forward movement function definition
+		void AFPSCharacter::MoveForward(float Value)
+		{
+			if (Controller != nullptr && Value != 0.f)
+			{
+				//find out which way is forward
+				FRotator Rotation = Controller->GetControlRotation();
+				//limit pitch when walking for falling
+				if (GetCharacterMovement()->IsMovingOnGround() || GetCharacterMovement()->IsFalling())
+				{
+					Rotation.Pitch = 0.f;
+				}
+				//add movement in that direction
+				const FVector Direction = FRotationMatrix(Rotation).GetScaledAxis(EAxis::X);
+				AddMovementInput(Direction, Value);
+			}
+		}
+		
+		//Right movement function definition
+		void AFPSCharacter::MoveRight(float Value)
+		{
+			if (Controller != nullptr && Value != 0.f)
+			{
+				//find our which way is right
+				const FRotator Rotation = Controller->GetControlRotation();
+				const FVector Direction = FRotationMatrix(Rotation).GetScaledAxis(EAxis::Y);
+				//add movement in that direction
+				AddMovementInput(Direction, Value);
+			}
+		}
+		
+		//Function setting the jump flag
+		void AFPSCharacter::OnStartJump()
+		{
+			bPressedJump = true;
+		}
+		
+		//Function clearning the jump flag
+		void AFPSCharacter::OnStopJump()
+		{
+			bPressedJump = false;
+		}
+
+1.	Save the `FPSCharacter` interface (\*.h) and implementation (\*.cpp) files in Visual Studio.
+
+1.	Locate **FPSProject** in the **Solution Explorer**.
+
+1.	Right-click on **FPSProject** and select **Build** to compile your project.
+
+	![](BuildProject.png)
+
+1.	After the build finishes, open your **FPSProject** in Unreal Editor.
+
+## Building the Mesh Blueprint
+
+[REGION:note]
+Before continuing, download and extract our sample mesh from the following link:
+["First Person Skeletal Mesh"](HeroFPP.zip "First Person Skeletal Mesh")
+[/REGION]
+
+1.	**Right-click** inside the Content Browser's file box to open the **Import Asset** dialog box
+
+1.	**Click 'Import to /Game...'** to open the **Import** dialog box.
+
+	![](RightClickImport.png)
+
+1.	Locate and select the **HeroFPP.fbx** mesh file.
+
+1.	**Click Open** to begin importing the mesh to your project.
+
+1.	In the **Content Browser**, the **FBX Import Options** dialog box appears. **Clicking Import** adds your mesh to the Project.
+
+1.	**Click** the **Save** button to save your imported mesh.
+
+
+
+## Adding the New Mesh to the Game
+
+
 
 %Steps%
\ No newline at end of file
