#UE4doc #docs
Completed first draft of Section 2, entitled 'Add a First Person Mesh'.
Availability: Docs

diff --git a/Engine/Documentation/Source/Programming/Tutorials/FirstPersonShooter/2/8/FirstPersonShooter_2_8.INT.udn b/Engine/Documentation/Source/Programming/Tutorials/FirstPersonShooter/2/8/FirstPersonShooter_2_8.INT.udn
--- a/Engine/Documentation/Source/Programming/Tutorials/FirstPersonShooter/2/8/FirstPersonShooter_2_8.INT.udn
+++ b/Engine/Documentation/Source/Programming/Tutorials/FirstPersonShooter/2/8/FirstPersonShooter_2_8.INT.udn
@@ -262,12 +262,56 @@ Before continuing, download and extract our sample mesh from the following link:
 
 1.	In the **Content Browser**, the **FBX Import Options** dialog box appears. **Clicking Import** adds your mesh to the Project.
 
+	[REGION:warning]
+	Disregard the following error regarding smoothing groups:
+	![](FBXWarning.png)  
+	This mesh still illustrates the first person mesh set-up and it will work with the animations you'll set-up in a later section.
+	[/REGION]
+
 1.	**Click** the **Save** button to save your imported mesh.
 
+1.	Navigate back to the **Blueprints** folder in the **Content Browser**.
+
+1.	 **Double-click** the **BP_FPSCharacter** icon to open it in the **Blueprint Editor**.
+
+1.	Locate the new **FirstPersonMesh** component in the **Components** tab.
+
+	![](LocateFPSMeshComponent.png)
+
+	[REGION:note]
+	The **FirstPersonMesh** component is a child of the **FirstPersonCameraComponent**, which means that it will always be attached to the camera.
+	[/REGION]
+
+1.	**Click** on the **FirstPersonMesh** in the **Components** menu.
+
+1.	Scroll down to the **Mesh** section of the **Details** tab and **click** on the dropdown menu that says "None".
 
+	![](SelectHeroFPPSkeletalMesh.png)
 
-## Adding the New Mesh to the Game
+1.	Select the **HeroFPP** skeletal mesh to add the arms to the **Viewport**.
 
+1.	The newly added **HeroFPP** skeletal mesh should look like the following inside the **Viewport**.
+
+	![](HeroFPPMesh_Viewport.png)
+
+1.	Adjust the newly added mesh's transform to be in front of the camera by setting its **Location** to {240, 0, 35} and by setting its **Rotation** to {-180, 50, -180}.
+
+	[REGION:lightbox]
+		[![](AdjustingArmsTransform.png)(w:750)](AdjustingArmsTransform.png)
+	[/REGION]
+	_Click on the image to zoom in._
+
+1.	 Make sure to **Compile** and **Save** the **BP_FPSCharacter** Blueprint before closing the **Blueprint Editor**.
+
+## Viewing the New Mesh In-Game
+
+1.	**Click** the **Play In** button in the **Level Editor Toolbar** to view the new mesh in-game.
+
+	[REGION:note]
+	To freely move the camera to see your character's mesh, press F1 to regain control of your mouse cursor, then click the **Eject** button in the toolbar.
+	![](FreeMovingCamera.png)
+	[/REGION]
 
+1.	Click the **Stop** button in the Level Editor to exit Play in Editor (PIE) mode.
 
 %Steps%
\ No newline at end of file
