#UE4doc #docs
Minor edit.

diff --git a/Engine/Documentation/Source/Programming/Tutorials/FirstPersonShooter/2/6/FirstPersonShooter_2_6.INT.udn b/Engine/Documentation/Source/Programming/Tutorials/FirstPersonShooter/2/6/FirstPersonShooter_2_6.INT.udn
--- a/Engine/Documentation/Source/Programming/Tutorials/FirstPersonShooter/2/6/FirstPersonShooter_2_6.INT.udn
+++ b/Engine/Documentation/Source/Programming/Tutorials/FirstPersonShooter/2/6/FirstPersonShooter_2_6.INT.udn
@@ -82,7 +82,7 @@ Before getting started, download and extract our sample mesh from the following
 
 	![](SelectSkeletalMesh.png)
 
-1.	Align the `SkeletalMeshComponent` to the `CapsuleComponent` by setting its Z Location to -88 in the **Details** tab.
+1.	Align the `SkeletalMeshComponent` to the `CapsuleComponent` by setting its Z **Location** to -88 in the **Details** tab.
 
 	![](NewZLoc_Mesh.png)
 
