#UE4doc #docs:
Removed Overview header
Added engine version field
Changed font settings for multiple terms
Fixed capitalization in some areas
Added a couple steps and edited a couple steps to align with the current engine version
Replaced a few pictures to reflect the latest UI changes

diff --git a/Engine/Documentation/Source/Engine/Animation/RootMotion/RootMotion.INT.udn b/Engine/Documentation/Source/Engine/Animation/RootMotion/RootMotion.INT.udn
--- a/Engine/Documentation/Source/Engine/Animation/RootMotion/RootMotion.INT.udn
+++ b/Engine/Documentation/Source/Engine/Animation/RootMotion/RootMotion.INT.udn
@@ -8,18 +8,19 @@ Related: Engine/Animation/PhysicallyDrivenAnimation
 Related: Engine/Content/Tools/MayaRiggingTool
 Related: Engine/Animation/AnimBlueprints
 Related: Engine/Animation/AnimBlueprints/EventGraph
+version: 4.9
 
 [TOC(start:2 end:2)]
 
-## Overview
+
 
 Typically in game animation, a character's collision capsule (or other shape) is driven through the scene by the controller. Data from this capsule is then used to drive animation. For instance, if the capsule is moving forward, the system then knows to play a running or walking animation on the character to give the appearance that the character is moving under its own power.
 
-However, this type of motion is not always ideal for every situation. In some cases, it makes sense for complex animations to actually drive the collision capsule, and not the other way around. This is where root motion handling becomes critical for your games.
+However, this type of motion is not always ideal for every situation. In some cases, it makes sense for complex animations to actually drive the collision capsule, and not the other way around. This is where Root Motion handling becomes critical for your games.
 
 Consider, for example, a special attack from a player in which the mesh had been pre-animated to lunge forward. If all character motion was based off of the player capsule, such animations would cause the character to step outside the capsule, effectively losing collision. Once the animation had played out, the player would slide back to their collision location. This would be problematic, as the capsule is generally used as the center of all calculations. Characters outside their capsule will pass through geometry and would not react properly to their environment. Plus, the slide back to their capsule at the end of the animation would be unrealistic.
  
-For those new to the concept, it may not be readily apparent why proper root motion would be important from just reading a description. Below, we have a character that is performing a pre-animated attack in which he lunges forward and slams down a hammer. Be aware that the animation of the character moving forward was not handled in game - it was created that way by an animation artist.
+For those new to the concept, it may not be readily apparent why proper Root Motion would be important from just reading a description. Below, we have a character that is performing a pre-animated attack in which he lunges forward and slams down a hammer. Be aware that the animation of the character moving forward was not handled in game - it was created that way by an animation artist.
 
 [REGION:warning]
 Early versions of the engine (Versions previous to 4.5) required the use of a AnimMontage (Please see the [AnimMontage documentation](https://docs.unrealengine.com/latest/INT/Engine/Animation/AnimMontage/index.html) for more information). 
@@ -47,14 +48,14 @@ Another example of **Root Motion** applied to a character can be found on the [A
 
 ## What Is Root Motion?
 
-Put simply, root motion is the motion of a character that is based off animation from the root bone of the skeleton. Most in-game animations are handled via cycles in which the character's root remains stationary. This cannot always be the case, however, as we can see in the example above. To handle this, we need to take the motion of the root away from the character and apply it to the character's capsule instead. That is the essence of what Root Motion does in UE4.
+Put simply, Root Motion is the motion of a character that is based off animation from the root bone of the skeleton. Most in-game animations are handled via cycles in which the character's root remains stationary. This cannot always be the case, however, as we can see in the example above. To handle this, we need to take the motion of the root away from the character and apply it to the character's capsule instead. That is the essence of what Root Motion does in UE4.
 
 [REGION:warning]
 In order to properly use Root Motion, it is important to note that the Root Bone of your character should be at the origin (0,0,0 with no rotation) as this allows the system to isolate physical movement (capsule) from animated movement (character). 
 [/REGION]
 
 
-The root motion of an animation can be visualized in Persona. Simply open up any animation in which the root bone is moving, and in the viewport choose **Show > Bones**. If you have not selected **Enable Root Motion** in the animation properties, you will notice a red line that appears as the character's root moves. This illustrates the root motion of the animation.
+The Root Motion of an animation can be visualized in Persona. Simply open up any animation in which the root bone is moving, and in the viewport choose **Show > Bones**. If you have not selected **Enable Root Motion** in the animation properties, you will notice a red line that appears as the character's root moves. This illustrates the Root Motion of the animation.
 
 [REGION:fullwidth]
 ![](RootMotionBone.png)
@@ -67,16 +68,16 @@ When you enable Root Motion, this red line will disappear. Instead, the characte
 [/REGION]
 
 
-So, what does this mean? Now that the system is taking the root motion of the character's animation into account, it is then able to reapply that same motion to the character's collision capsule. This means that it can still have the exact same motion that the animator intended, but it can also respond properly to in-game collisions and other physics and gameplay events. We will see this in action in the section below.
+So, what does this mean? Now that the system is taking the Root Motion of the character's animation into account, it is then able to reapply that same motion to the character's collision capsule. This means that it can still have the exact same motion that the animator intended, but it can also respond properly to in-game collisions and other physics and gameplay events. We will see this in action in the section below.
 
 
 ## Root Motion in Action
 
-Below is a breakdown of the problems with not handling root motion, and the benefits root motion has on your animations. We will be using the lunging animation shown above.
+Below is a breakdown of the problems with not handling Root Motion, and the benefits Root Motion has on your animations. We will be using the lunging animation shown above.
 
 ### Without Root Motion
 
-Without root motion, such animations would pull the character away from his capsule, as seen below. Note the unsightly snap back to the capsule location. This happens because once the animation is done playing, the character controller pulls the character back to their original position.
+Without Root Motion, such animations would pull the character away from his capsule, as seen below. Note the unsightly snap back to the capsule location. This happens because once the animation is done playing, the character controller pulls the character back to their original position.
 
 [OBJECT:EmbeddedVideo]
 [PARAMLITERAL:width]
@@ -111,7 +112,7 @@ ov9pyx4MAOo
 
 ### How Root Motion Handling Solves the Problem
 
-Once your animation is set up to utilize root motion, then the movement that was set up by the animator temporarily becomes the driving force of the capsule. This allows us to continue playing from the end point of the animation. You can see that the second time we trigger the attack, we take off from the new position. Of course, we turned around first so we did not wander off camera.
+Once your animation is set up to utilize Root Motion, then the movement that was set up by the animator temporarily becomes the driving force of the capsule. This allows us to continue playing from the end point of the animation. You can see that the second time we trigger the attack, we take off from the new position. Of course, we turned around first so we did not wander off camera.
 
 
 
@@ -149,10 +150,10 @@ YuD29-Xr7Oc
 
 ## Setting Up Root Motion
 
-Setting up root motion in UE4 is very easy, however, the steps will be slightly different depending on what version of UE4 you are using. At a high level, the process pre-4.6 includes:
+Setting up Root Motion in UE4 is very easy, however, the steps will be slightly different depending on what version of UE4 you are using. At a high level, the process pre-4.6 includes:
 
 1. (UE4 4.6 and above) - Open the Animation in which you would like to use and make sure that the **Enable Root Motion** property is checked
-1. Assign the animation which requires root motion control to an AnimMontage asset.
+1. Assign the animation which requires Root Motion control to an AnimMontage asset.
 1. Within that AnimMontage, be sure to name the slot that contains the motion.
 1. (If you are using UE4 Pre-4.6) Within the AnimMontage, make sure that the **Enable Root Motion Translation** and **Enable Root Motion Rotation** properties are checked.
 1. In the AnimGraph, make sure that the animation flows through a **Slot** node, and that the Slot node is assigned the same slot name as provided to your AnimMontage.
@@ -179,7 +180,7 @@ So, we start by creating a function within our BlueprintInterface. In our case,
 
 ![](UseHammerFunction.png)
 
-Next, in the Character Blueprint, we set up the following special network. Note that handling for 3rd person control was already in place - we are just focusing on what we need to get the hammer animation working. In this case, we set up an input event such that when we press **R**, we call the _UseHammer_ function on our interface.
+Next, in the Character Blueprint, we set up the following special network. Note that handling for 3rd person control was already in place - we are just focusing on what we need to get the hammer animation working. In this case, we set up an input event such that when we press **R**, we call the UseHammer function on our interface.
 
 [REGION:fullwidth]
 ![](PressRUseHammer.png)
@@ -211,24 +212,30 @@ Before we do, we need to make sure Root Motion is enabled on our Animation. Doub
 
 Once this is done, we need to create the AnimMontage.
 
-The easiest way to do this is to **right-click** (**Ctrl + left-click** on Mac) on the animation that requires Root Motion, and choose Create AnimMontage.
+The easiest way to do this is to **Right-click** (**Ctrl + Left-click** on Mac) on the animation that requires Root Motion, and choose Create AnimMontage.
 
  ![](CreateAnimMontage.png)
 
-Name your AnimMontage and then **double-click** it to open it within Persona. In the center of the Editor, near the bottom you will see a SlotName field. Enter a meaningful name for the slot. If this is part of a set of animations that will play together, you should choose that name. A common slot name is _FullBody_, which is what we will use.
+Name your AnimMontage and then **Double-click** it to open it within Persona. In the center of the Editor, near the bottom you will see a SlotName dropdown. Select the eyeglass icon under it to bring up the Anim Slot Manager. Oress the **Add Slot** button and name the slot.  If this is part of a set of animations that will play together, you should choose that name. A common slot name is FullBody, which is what we will use.
+
+After that, select your new slot from the dropdown menu inside of the Montage.
 
 [REGION:fullwidth]
 ![](SettingUpSlotinMontage.png)
 [/REGION]
 
-Next, we go back to the **Animation Blueprint**. Just after our UseHammer function event, we can now create a **Montage Play**. We need to make sure to associate this with the proper AnimMontage.
+Next, we go back to the Animation Blueprint. Just after our UseHammer function event, we can now create a **Montage Play**. We need to make sure to associate this with the proper AnimMontage.
 
 ![](PlayMontage.png)
 
-We are almost finished! Finally, we just need to make sure that our slot is being calculated in our final animation. Within the **Animation Blueprint**, we need to go into the AnimGraph tab. Create a **Slot** node, and give it the name you used for the slot within the AnimMontage. In our case, that would be _FullBody_. This needs to go between any existing State Machines or other animations and your Final Animation Pose node.
+We are almost finished! Finally, we just need to make sure that our slot is being calculated in our final animation. Within the Animation Blueprint, we need to go into the AnimGraph tab. Create a **Slot** node, and give it the name you used for the slot within the AnimMontage. In our case, that would be FullBody. This needs to go between any existing State Machines or other animations and your Final Animation Pose node.
 
 ![](InsertSlotNode.png)
 
+[REGION:note]
+There may be a warning displayed on the New State Machine node. This is due to it not being completely set up, but this warning is not important for the purposes of this tutorial.
+[/REGION]
+
 That is all there is to it! Now that your animation is playing back via an AnimMontage, with that Montage being set to utilize Root Motion, everything will just work!
 
 
@@ -250,7 +257,7 @@ PWB_mqjz3iA
 
 [INCLUDE:Engine/Animation/AnimMontage#MontageReplication]
 
-The physics state of the Character is taken into account during Root Motion. For example, if the character physics is _Walking_ or _Falling_, then the Z-axis of Root Motion will be ignored, and gravity will be applied. Characters will fall, go down slopes or up stairs. If the character physics is _Flying_, then the full root motion is applied and gravity is ignored.
+The physics state of the Character is taken into account during Root Motion. For example, if the character physics is Walking or Falling, then the Z-axis of Root Motion will be ignored, and gravity will be applied. Characters will fall, go down slopes or up stairs. If the character physics is Flying, then the full Root Motion is applied and gravity is ignored.
 
 
 
