#UE4doc #docs:
Unbolded a few terms

diff --git a/Engine/Documentation/Source/Engine/Animation/Sequences/Editor/AnimSequenceEditor.INT.udn b/Engine/Documentation/Source/Engine/Animation/Sequences/Editor/AnimSequenceEditor.INT.udn
--- a/Engine/Documentation/Source/Engine/Animation/Sequences/Editor/AnimSequenceEditor.INT.udn
+++ b/Engine/Documentation/Source/Engine/Animation/Sequences/Editor/AnimSequenceEditor.INT.udn
@@ -10,7 +10,7 @@ Description:Guide to using Persona to edit Animation Sequences and add animation
 [TOC(start:2)]
 
 [EXCERPT:overview]
-The **Animation Sequence Editor** provides functionality for previewing and editing individual **Animation Sequence** assets. The properties, including compression scheme, of the **Animation Sequence** can be set, and animation notification events - also called Notifies - can be added to allow camera effects, particle effects, sounds, and more to be triggered based on the animation.
+The **Animation Sequence Editor** provides functionality for previewing and editing individual **Animation Sequence** assets. The properties, including compression scheme, of the Animation Sequence can be set, and animation notification events - also called Notifies - can be added to allow camera effects, particle effects, sounds, and more to be triggered based on the animation.
 [/EXCERPT:overview]
 
 ## Interface
diff --git a/Engine/Documentation/Source/Engine/Animation/Sequences/Notifies/AnimNotifies.INT.udn b/Engine/Documentation/Source/Engine/Animation/Sequences/Notifies/AnimNotifies.INT.udn
--- a/Engine/Documentation/Source/Engine/Animation/Sequences/Notifies/AnimNotifies.INT.udn
+++ b/Engine/Documentation/Source/Engine/Animation/Sequences/Notifies/AnimNotifies.INT.udn
@@ -29,16 +29,16 @@ Description:Notifies are a system for setting up and receiving events in Animati
 ## Overview
 
 [EXCERPT:Intro]
-**Animation Notifications** (**AnimNotifies** or just **Notifies**) provide a way for animation programmers to setup events to occur at specific points during an **Animation Sequence**. **Notifies** are commonly used to add effects like footstep sounds during a walk or run animation or spawning a particle system during animation. However, they can have any number of different uses, as the system can be extended with custom notification types to cater to the needs of any type of game.
+**Animation Notifications** (**AnimNotifies** or just **Notifies**) provide a way for animation programmers to setup events to occur at specific points during an **Animation Sequence**. Notifies are commonly used to add effects like footstep sounds during a walk or run animation or spawning a particle system during animation. However, they can have any number of different uses, as the system can be extended with custom notification types to cater to the needs of any type of game.
 [/EXCERPT:Intro]
 
-Creating and editing **Notifies** is performed within the **Notify Editor** in **Persona**. For more information, refer to the [Editing Animation Sequences](Engine/Animation/Sequences/Editor) documentation. 
+Creating and editing Notifies is performed within the **Notify Editor** in **Persona**. For more information, refer to the [Editing Animation Sequences](Engine/Animation/Sequences/Editor) documentation. 
 
 ## Animation Notify Types
 
-Different types of **Notifies** cause different events to be triggered. Camera effects, particle effects, sounds, etc. can all be triggered at any point during an animation. 
+Different types of Notifies cause different events to be triggered. Camera effects, particle effects, sounds, etc. can all be triggered at any point during an animation. 
 
-Below you will find the currently available **Notifies** and how they are used. 
+Below you will find the currently available Notifies and how they are used. 
 
 
 ### Play Particle Effect
