#UE4doc #docs:
Added engine version field
Changed font settings for multiple terms
Replaced a picture to reflect the current engine UI

diff --git a/Engine/Documentation/Source/Engine/Animation/Sequences/AnimationSequences.INT.udn b/Engine/Documentation/Source/Engine/Animation/Sequences/AnimationSequences.INT.udn
--- a/Engine/Documentation/Source/Engine/Animation/Sequences/AnimationSequences.INT.udn
+++ b/Engine/Documentation/Source/Engine/Animation/Sequences/AnimationSequences.INT.udn
@@ -8,6 +8,7 @@ Related: Engine/Animation/PhysicallyDrivenAnimation
 Related: Engine/Animation/StateMachines
 Related: Gameplay/AI
 Related: Engine/Content/Tools/MayaRiggingTool
+version: 4.9
 
 [VAR:Topic]
 [OBJECT:Topic]
@@ -47,10 +48,10 @@ Qo4hH0h9ZZA
 [/PARAMLITERAL]
 [/OBJECT]
 
-Each **Animation Sequence** asset targets a specific _Skeleton_ and can only be played on that _Skeleton_. This means, in order to share animations between multiple _Skeletal Meshes_, each of the meshes must use the same _Skeleton_ asset.
+Each Animation Sequence asset targets a specific Skeleton and can only be played on that Skeleton. This means, in order to share animations between multiple Skeletal Meshes, each of the meshes must use the same Skeleton asset.
 
 [REGION:tip]
-An example of a **Skeleton** using an **Animation Sequence** can also be found on the [Animation Content Examples](Resources\ContentExamples\Animation) page under section 1.1. 
+An example of a Skeleton using an Animation Sequence can also be found on the [Animation Content Examples](Resources\ContentExamples\Animation) page under section 1.1. 
 [/REGION]
 
 ## Editing
@@ -63,7 +64,7 @@ See [](Engine/Animation/Sequences/Editor) for more details.
 
 [INCLUDE:Engine/Animation/Sequences/UserGuide#Intro]
 
-See the [](Engine/Animation/Sequences/UserGuide) for instructions on using **Animation Sequences** within **Animation Blueprints**.
+See the [](Engine/Animation/Sequences/UserGuide) for instructions on using Animation Sequences within Animation Blueprints.
 
 ## Notifies
 
@@ -82,7 +83,7 @@ See [](Engine/Animation/Sequences/Curves) for descriptions and usage information
 
 ## Compression
 
-**Animation Sequence** assets when opened up in **Persona** can have compression applied to them. 
+Animation Sequence assets when opened up in **Persona** can have compression applied to them. 
 
 To do this, from the **Menu Bar** select **Asset** then select **Apply Compression** (as seen below). 
 
