added description

diff --git a/Engine/Documentation/Source/Engine/Animation/AnimationBlending/AnimationBlending.INT.udn b/Engine/Documentation/Source/Engine/Animation/AnimationBlending/AnimationBlending.INT.udn
--- a/Engine/Documentation/Source/Engine/Animation/AnimationBlending/AnimationBlending.INT.udn
+++ b/Engine/Documentation/Source/Engine/Animation/AnimationBlending/AnimationBlending.INT.udn
@@ -1,6 +1,7 @@
 Availability:Public
 Title: Blending Animations
 Crumbs: %ROOT%, Engine, Engine/Animation
+description:Transitioning smoothly between two animations on a single Skeletal Mesh
 version: 4.9
 
 
