fixed broken link

diff --git a/Engine/Documentation/Source/Engine/Components/Camera/CameraComponents.INT.udn b/Engine/Documentation/Source/Engine/Components/Camera/CameraComponents.INT.udn
--- a/Engine/Documentation/Source/Engine/Components/Camera/CameraComponents.INT.udn
+++ b/Engine/Documentation/Source/Engine/Components/Camera/CameraComponents.INT.udn
@@ -19,7 +19,7 @@ A related property that can be set for any **Pawn** is `bTakeCameraControlWhenPo
 
 ## Spring Arm Component
 
-The **SpringArmComponent** tries to maintain its children at a fixed distance from the parent, but will retract the children if there is a collision, and spring back when there is no collision. Typically, the SpringArmComponent is used as a "camera boom" to keep the follow camera for a player from colliding into the world (without the SpringArmComponent, the CameraComponent will remain at the fixed distance specified regardless of objects that could potentially get in the way of it and the object it is attached to), see [Using a Spring Arm Component](Gameplay/HowTo/UsingCameras/Blueprints#usingaspringarmcomponent). 
+The **SpringArmComponent** tries to maintain its children at a fixed distance from the parent, but will retract the children if there is a collision, and spring back when there is no collision. Typically, the SpringArmComponent is used as a "camera boom" to keep the follow camera for a player from colliding into the world (without the SpringArmComponent, the CameraComponent will remain at the fixed distance specified regardless of objects that could potentially get in the way of it and the object it is attached to), see [Using a Spring Arm Component](Gameplay/HowTo/UsingCameras/SpringArmComponents). 
 
 ![](camera_springArm.png)
 
