#UE4doc #docs:
Resubmitting due to unsubmitted changes

diff --git a/Engine/Documentation/Source/Engine/Animation/Sequences/Curves/Curves.INT.udn b/Engine/Documentation/Source/Engine/Animation/Sequences/Curves/Curves.INT.udn
--- a/Engine/Documentation/Source/Engine/Animation/Sequences/Curves/Curves.INT.udn
+++ b/Engine/Documentation/Source/Engine/Animation/Sequences/Curves/Curves.INT.udn
@@ -2,9 +2,10 @@ Availability:Public
 Title:Animation Curves
 Crumbs: %ROOT%, Engine, Engine/Animation, Engine/Animation/Sequences
 Description:Curve system for driving properties of Morph Targets and Materials throughout animation playback.
+version: 4.9
 
 [EXCERPT:Intro]
-Curves provide a way to change the value of a Material parameter or a Morph Target while an animation is playing back. Their workflow is very simple, requiring that you simply designate the asset you want to modify (either a Material or a Morph Target), name the curve accordingly, and then adjust keyframe values across the duration of an animation.
+**Animation Curves** provide a way to change the value of a Material parameter or a Morph Target while an animation is playing back. Their workflow is very simple, requiring that you simply designate the asset you want to modify (either a Material or a Morph Target), name the curve accordingly, and then adjust keyframe values across the duration of an animation.
 [/EXCERPT:Intro]
 
 
@@ -30,13 +31,13 @@ TKAcNubFOH8
 [/PARAMLITERAL]
 [/OBJECT]
 
-We can open that Animation Sequence in Persona and look at it by **double-click** it in the **Content Browser**.
+We can open that Animation Sequence in Persona and look at it by **Double-clicking** it in the **Content Browser**.
 
 [REGION:fullwidth]
 ![](Persona_HammerSlam.png)
 [/REGION]
 
-The character's head has a simple Material with a parameter named _Blush_. Here is a look at the Material Graph:
+The character's head has a simple Material with a parameter named **Blush**. Here is a look at the Material Graph:
 
 [REGION:fullwidth]
 ![](BlushMaterial.jpg)
@@ -46,7 +47,7 @@ The character's head has a simple Material with a parameter named _Blush_. Here
 We have also made a [Material Instance](Engine/Rendering/Materials\MaterialInstances) from this Material and applied it to our character.
 [/REGION]
 
-In Persona, we start by clicking the ![](button_AddFloatTrack.png) button located in the Curves area. We name the Curve _Blush_. This name _must_ coincide with the name used in the Material Parameter.
+In Persona, we start by clicking the ![](button_AddFloatTrack.png) button located in the Curves area. From this, navigate to the **Add Variable Curve...** field and select **Create Curve**. We name the Curve "Blush". This name must coincide with the name used in the Material Parameter.
 
 ![](AddingBlushCurve.png)
 
@@ -58,11 +59,11 @@ You can expand the curve area size by clicking the small eye icon next to the dr
 
 ![](ExpandingCurve.png)
 
-You can **shift+click** in the curve area to create a new keyframe. Keyframe values and times can be adjusted using the fields that appear when a keyframe is selected.
+You can **Shift+Left-click** in the curve area to create a new keyframe. Keyframe values and times can be adjusted using the fields that appear when a keyframe is selected.
 
 ![](CreateMaterialCurve.png)
 
-You can also **right click** on a keyframe and change its tangent type:
+You can also **Right-click** on a keyframe and change its tangent type:
 
 
 ![](CurveType.png)
