#UE4doc #docs:
Added engine version field
Changed font settings for multiple terms
Fixed capitalization on multiple terms
Replaced multiple pictures due to changes to the UI
Changed a couple steps to align with the latest version of the engine.

diff --git a/Engine/Documentation/Source/Engine/Animation/Sequences/Editor/AnimSequenceEditor.INT.udn b/Engine/Documentation/Source/Engine/Animation/Sequences/Editor/AnimSequenceEditor.INT.udn
--- a/Engine/Documentation/Source/Engine/Animation/Sequences/Editor/AnimSequenceEditor.INT.udn
+++ b/Engine/Documentation/Source/Engine/Animation/Sequences/Editor/AnimSequenceEditor.INT.udn
@@ -2,6 +2,7 @@ Availability:Public
 Title:Editing Animation Sequences
 Crumbs: %ROOT%, Engine, Engine/Animation, Engine/Animation/Sequences
 Description:Guide to using Persona to edit Animation Sequences and add animation notification events.
+version: 4.9
 
 [REGION:banner]
 ![](persona_sequence_banner.png)(convert:false)
@@ -19,27 +20,27 @@ The **Animation Sequence Editor** consists of the following:
 
 ![](AnimationSequenceEditor.png)
 
-1. **Toolbar** - Displays the name of the _AnimationSequence_ currently being edited.
-1. **Track List** - Editable list of tracks for creating/editing notifies. See below for instructions on creating and removing tracks as well as creating and working with notifies.
+1. **Name** - Displays the name of the AnimationSequence currently being edited.
+1. **Notifies** - Editable list of tracks for creating/editing notifies. See below for instructions on creating and removing tracks as well as creating and working with notifies.
+1. **Curves** - A graph-like area for adding and editing animation curves
+1. **Tracks** - A field used for editing Animation Layers. For more information, please visit the [](Engine/Animation/AnimHowTo/LayerEditing) page.
 1. **Timeline** - Displays information about, and provides controls for, playback of the preview in the **Viewport** panel of **Persona**.
 
 
 ## Notifies
 
-An animation _Notify_ is an event along a particular animation at which you can cause something to happen. For instance, the moment a character's foot hits the floor, you can create a Notify along the timeline to play a footstep sound. Other notifies are available as well, allowing for particle, camera, and sound effects. You may also create custom notifies that trigger your own events.
+An **Animation Notify** is an event along a particular animation at which you can cause something to happen. For instance, the moment a character's foot hits the floor, you can create a Notify along the timeline to play a footstep sound. Other notifies are available as well, allowing for particle, camera, and sound effects. You may also create custom notifies that trigger your own events.
 
 For more information on Animation Notifies, please see the [](Engine/Animation/Sequences\Notifies) page.
 
 
 ## Notify Tracks
 
-Notifies are stored in **Tracks** within **Animation Sequences**. A track is essentially a container that can hold any number of Notifies, and an **Animation Sequences** can have any number of tracks. Using multiple tracks makes it easier to organize and visualize Notifies, especially in cases where an **Animation Sequences** requires large numbers or several different types of Notifies.
+Notifies are stored in **Tracks** within Animation Sequences. A track is essentially a container that can hold any number of Notifies, and an Animation Sequences can have any number of tracks. Using multiple tracks makes it easier to organize and visualize Notifies, especially in cases where an Animation Sequences requires large numbers or several different types of Notifies.
 
-Every **Animation Sequences** has a single track by default; however, additional tracks can be added to the sequence and existing tracks can be removed using the controls to the right of the tracks.
+Every Animation Sequences has a single track by default; however, additional tracks can be added to the sequence and existing tracks can be removed using the controls to the right of the tracks.
 
-**Adding a new track**
-
-In the **Track List**, click the ![](button_track_add.png) button next to an existing track:
+To add a new Track; In the **Track List**, click the ![](button_track_add.png) button next to an existing track:
 
 ![](AddTrack.png)
 
@@ -47,9 +48,7 @@ A new track is added above the existing track in the list:
 
 ![](TrackAdded.png)
 
-**Removing an existing track**
-
-In the **Track List**, click the ![](button_track_remove.png) button next to the track you want to remove:
+To remove an existing Track; In the **Track List**, click the ![](button_track_remove.png) button next to the track you want to remove:
 
 ![](RemoveTrack.png)
 
@@ -63,11 +62,11 @@ Notifies can be added to tracks, moved around inside a track, and moved back and
 
 ### Creating a Notify
 
-Notifies of any existing type or custom Notifies can be added through the context menu. **Right-clicking** on a track displays the context menu. All available types of Notifies are listed in the **Add a notify** flyout menu. Select the type of notify to add to the track:
+Notifies of any existing type or custom Notifies can be added through the context menu. **Right-clicking** on a track displays the context menu. All available types of Notifies are listed in the **Add Notify...** context menu. Select the type of notify to add to the track:
 
 ![](AddNotify.png)
 
-The new Notify is added to the track at the location where you **right-clicked**:
+The new Notify is added to the track at the location where you **Right-clicked**:
 
 ![](NotifyAdded.png)
 
@@ -94,24 +93,24 @@ Moving notifies between tracks is extremely easy. In fact, it is essentially the
 
 ### Copy/Pasting Notifies
 
-By **right-clicking** on a notify, you can copy it to the clipboard for pasting elsewhere along the animation, or even into entirely different Animation Sequences. 
+By **Right-clicking** on a notify, you can copy it to the clipboard for pasting elsewhere along the animation, or even into entirely different Animation Sequences. 
 
 ![](CopyAnimNotify.png)
 
-You can paste by also **right-clicking** within a notify track. When pasting, you may choose between 2 options:
+You can paste by also **Right-clicking** within a notify track. When pasting, you may choose between 2 options:
 
 ![](PasteOptions.png)
 
 * **Paste** - Paste the notify at the current location of the time slider.
 * **Paste at Original Time** - Pastes the notify at the same time as when it was copied. 
 
-You can also perform actions for multiple notifies by **Shift + Clicking** on a notify to add it to a selection or **Ctrl + Clicking** to toggle a notify from a selection. 
+You can also perform actions for multiple notifies by **Shift + Left-clicking** on a notify to add it to a selection or **Ctrl + Left-clicking** to toggle a notify from a selection. 
 
 You can drag the selection around and it will remain within the notify panel while keeping snapping intact. 
 
 ![](MultipleNotifies.png)
 
-When multiple notifies are copied, by **right-clicking** in the timeline to bring up the context menu, you will notice additional paste options. 
+When multiple notifies are copied, by **Right-clicking** in the timeline to bring up the context menu, you will notice additional paste options. 
 
 ![](CopyPasteMenu.png)
 
@@ -136,7 +135,7 @@ The information bar in the **Timeline** displays the following playback informat
 
 | Info | Description |
 | ---- | ----------- |
-| **Animation** | The name of the _AnimationSequence_ currently being edited. |
+| **Animation** | The name of the AnimationSequence currently being edited. |
 | **Percentage** | The current playback position expressed as a percentage of the whole animation. |
 | **Current Time** | The current playback position expressed in seconds. |
 | **Current Frame** | The current playback position expressed in frames. |
