#DocFix Correcting grammatical errors on this page reporting in the forums

diff --git a/Engine/Documentation/Source/Engine/Rendering/Materials/HowTo/DetailTexturing/DetailTexturing.INT.udn b/Engine/Documentation/Source/Engine/Rendering/Materials/HowTo/DetailTexturing/DetailTexturing.INT.udn
--- a/Engine/Documentation/Source/Engine/Rendering/Materials/HowTo/DetailTexturing/DetailTexturing.INT.udn
+++ b/Engine/Documentation/Source/Engine/Rendering/Materials/HowTo/DetailTexturing/DetailTexturing.INT.udn
@@ -95,7 +95,7 @@ Once selected, a new **Material** will be created. Name the Material **Detail_Ma
 
 	![](DT_Adding_Textures.png)
 
-1. With all of the required Textures added to the Material Graph, it is now time to add the **Detail Texturing** Material Function. To do this, first search in the **Palette** for **Detail Texturing**.
+1. With all of the required Textures added to the Material Graph, it is time to add the **Detail Texturing** Material Function. To do this, first search in the **Palette** for **Detail Texturing**.
 
 	![](DT_Search_DT_MF.png)
 
@@ -103,12 +103,12 @@ Once selected, a new **Material** will be created. Name the Material **Detail_Ma
 
 	![](DT_Adding_DT_MF.png)
 
-1. With the **Detail Texturing** Material Function now added, it is now time to start hooking the required inputs and outputs so the Material will work. For this example, we are going to be using the **T_Brick_Clay_Old** Textures as our base Textures and the **T_Ground_Gravel** as our Detail Textures. This means that the Diffuse and Normal maps for **T_Brick_Clay_Old** should go into the Diffuse and Normal inputs and the **T_Ground_Gravel** Diffuse and Normal maps should go into the Detail Diffuse and Normal inputs. When completed, your Material Graph should look like this.
+1. With the **Detail Texturing** Material Function added, it is time to start hooking the required inputs and outputs so the Material will work. For this example, we are going to be using the **T_Brick_Clay_Old** Textures as our base Textures and the **T_Ground_Gravel** as our Detail Textures. This means that the Diffuse and Normal maps for **T_Brick_Clay_Old** should go into the Diffuse and **T_Ground_Moss** Normal inputs and the **T_Ground_Gravel** Diffuse and Detail Normal maps should go into the Detail Diffuse and Normal inputs. When completed, your Material Graph should look like this.
 
 	![](DT_Hooked_Up_Textures.png)
 
 	[region:note]
-	Remember that the **T_Ground_Gravel** Diffuse and Normal maps Texture Samplers need to be converted into **Texture Objects** or the Material will not work. For more information on how to do this, please refer to the **Converting Texture Sample to Texture Objects** section above. 
+	Remember that the **T_Ground_Gravel** Diffuse and **T_Ground_Moss** Normal map Texture Samplers need to be converted into **Texture Objects** or the Material will not work. For more information on how to do this, please refer to the **Converting Texture Sample to Texture Objects** section above. 
 	[/region]
 
 1. Now that the Textures have been hooked up, it is time to add some Constant Material Expressions so that things like the scale and intensity of the Detail Textures can be tweaked. To find the Constant Material Expression, search for Constant in the Material Palette and then drag the Constant from the **Palette** to the Material Graph. Once you have added one Constant Material Expression, duplicate it by first selecting it and then press CTRL + W on the keyboard two times. In the end, you want to have three Constant Material Expressions. 
