[UE-20015] Improvements to U* macro documentation. (Contributed by [REDACTED_PERSON], pull request 1464.)

diff --git a/Engine/Documentation/Source/Programming/UnrealArchitecture/Reference/Interfaces/Interfaces.INT.udn b/Engine/Documentation/Source/Programming/UnrealArchitecture/Reference/Interfaces/Interfaces.INT.udn
--- a/Engine/Documentation/Source/Programming/UnrealArchitecture/Reference/Interfaces/Interfaces.INT.udn
+++ b/Engine/Documentation/Source/Programming/UnrealArchitecture/Reference/Interfaces/Interfaces.INT.udn
@@ -12,6 +12,72 @@ Declaring an interface class is very similar to declaring a regular class. The o
 	UINTERFACE([specifier, specifier, ...], [meta(key=value, key=value, ...)])
 	class ClassName : ParentName
 
-## Interface Specifiers
+## Interface Declaration
+
+In your .h file (e.g. `ItemInterface.h`):
+
+    #pragma once
+
+    #include "ItemInterface.generated.h"
+
+    UINTERFACE(Blueprintable)
+    class UItemInterface : public UInterface
+    {
+        GENERATED_UINTERFACE_BODY()
+    };
+
+    class IItemInterface
+    {    GENERATED_IINTERFACE_BODY()
+
+    public:
+
+        UFUNCTION(BlueprintNativeEvent, BlueprintCallable, Category="Item")
+        FItemStruct GetItem() const;
+    };
+
+In your .cpp file (e.g. `ItemInterface.cpp`):
+
+    #include "ItemInterface.h"
+
+    UItemInterface::UItemInterface(const FObjectInitializer& ObjectInitializer) : Super(ObjectInitializer)
+    {
+    }
+
+The declaration consists of a standard C++ class declaration for the class. Above the standard declaration, descriptors such as class specifiers and metadata are passed to the UINTERFACE macro. These are used to create the UInterface for the interface being declared, which can be thought of as the engine's specialized representation of the class. Also, the GENERATED_UINTERFACE_BODY() macro must be placed at the very beginning of the class body.
+
+For interfaces that are to be recognized by the engine, you must create two classes, a UINTERFACE class that inherits from UInterface, and a pure interface.
+
+The UINTERFACE class should be prefixed with U, as it extends from UObject. The Interface should be prefixed with I.
+
+The `Blueprintable` specifier is required if you want to allow Blueprints to implement this interface.
+
+### Interface Specifiers
 
 [INCLUDE:Programming/UnrealArchitecture/Reference/Interfaces/Specifiers#main(offset:1)]
+
+## Implementing Your Interface in C++
+
+When declaring your class in C++, simply inherit from your interface (not the UInterface class though).
+
+    class AItem : public AActor, public IItemInterface
+    {
+        GENERATED_BODY()
+
+    public:
+
+        FItemStruct GetItem_Implementation() const;
+    };
+
+## Determining If A Given Actor Implements Your Interface
+
+For compatability with both C++ and Blueprint classes that implement your interface, use the following code:
+
+    obj->GetClass()->ImplementsInterface(UItemInterface::StaticClass());
+
+[REGION:note]
+You must use the UInterface version of the class to check if the interface is implemented.
+[/REGION]
+
+## Blueprint Implementable Classes
+
+If you want to allow Blueprints to implement this interface, you must use the `Blueprintable` metadata specifier, and every function in your interface must be a `BlueprintNativeEvent` or a `BlueprintImplementableEvent`. Otherwise, a Blueprint class implementing the interface would be unable to implement all of the methods of the interface.
diff --git a/Engine/Documentation/Source/Programming/UnrealArchitecture/Reference/Interfaces/Specifiers/InterfaceSpecifiers.INT.udn b/Engine/Documentation/Source/Programming/UnrealArchitecture/Reference/Interfaces/Specifiers/InterfaceSpecifiers.INT.udn
--- a/Engine/Documentation/Source/Programming/UnrealArchitecture/Reference/Interfaces/Specifiers/InterfaceSpecifiers.INT.udn
+++ b/Engine/Documentation/Source/Programming/UnrealArchitecture/Reference/Interfaces/Specifiers/InterfaceSpecifiers.INT.udn
@@ -4,14 +4,15 @@ Description: Keywords used when declaring gameplay interfaces to specify how the
 Availability:Public
 
 [EXCERPT:main]
-When declaring interfaces, specifiers can be added to the declaration to control how the interface behaves with various 
+When declaring interfaces, specifiers can be added to the declaration to control how the interface behaves with various
 aspects of the engine and editor. Interfaces use a subset of the specifiers available when declaring classes.
 
-
+[REGION:buttonlist]
 * [](Programming/UnrealArchitecture/Reference/Classes/Specifiers/Blueprintable)
 * [](Programming/UnrealArchitecture/Reference/Classes/Specifiers/DependsOn)
 * [](Programming/UnrealArchitecture/Reference/Classes/Specifiers/MinimalAPI)
 * [](Programming/UnrealArchitecture/Reference/Classes/Specifiers/NotBlueprintable)
+[/REGION]
 
 ## Metadata Specifiers
 [INCLUDE:Programming/UnrealArchitecture/Reference/Metadata#MetadataUsage]
diff --git a/Engine/Documentation/Source/Programming/UnrealArchitecture/Reference/Metadata/BlueprintInternalUseOnly/FunctionMetadataSpecifier_BlueprintInternalUseOnly.INT.udn b/Engine/Documentation/Source/Programming/UnrealArchitecture/Reference/Metadata/BlueprintInternalUseOnly/FunctionMetadataSpecifier_BlueprintInternalUseOnly.INT.udn
--- a/Engine/Documentation/Source/Programming/UnrealArchitecture/Reference/Metadata/BlueprintInternalUseOnly/FunctionMetadataSpecifier_BlueprintInternalUseOnly.INT.udn
+++ b/Engine/Documentation/Source/Programming/UnrealArchitecture/Reference/Metadata/BlueprintInternalUseOnly/FunctionMetadataSpecifier_BlueprintInternalUseOnly.INT.udn
@@ -1,7 +1,12 @@
 Title: BlueprintInternalUseOnly
 Crumbs: %ROOT%, Programming, Programming/UnrealArchitecture/Reference, Programming/UnrealArchitecture/Reference/Functions, Programming/UnrealArchitecture/Reference/Functions/Specifiers
-Description: 
+Description:
 Availability:Public
 Version: 4.9
 
 This function is an internal implementation detail, used to implement another function or node.  It is never directly exposed in a graph.
+
+## Usage
+
+    UFUNCTION(meta=(BlueprintInternalUseOnly))
+    ReturnType FunctionName([Parameter, Parameter, ...])
diff --git a/Engine/Documentation/Source/Programming/UnrealArchitecture/Reference/Metadata/BlueprintProtected/FunctionMetadataSpecifier_BlueprintProtected.INT.udn b/Engine/Documentation/Source/Programming/UnrealArchitecture/Reference/Metadata/BlueprintProtected/FunctionMetadataSpecifier_BlueprintProtected.INT.udn
--- a/Engine/Documentation/Source/Programming/UnrealArchitecture/Reference/Metadata/BlueprintProtected/FunctionMetadataSpecifier_BlueprintProtected.INT.udn
+++ b/Engine/Documentation/Source/Programming/UnrealArchitecture/Reference/Metadata/BlueprintProtected/FunctionMetadataSpecifier_BlueprintProtected.INT.udn
@@ -1,7 +1,12 @@
 Title: BlueprintProtected
 Crumbs: %ROOT%, Programming, Programming/UnrealArchitecture/Reference, Programming/UnrealArchitecture/Reference/Functions, Programming/UnrealArchitecture/Reference/Functions/Specifiers
-Description: 
+Description:
 Availability:Public
 Version: 4.9
 
- This function can only be called on 'this' in a blueprint. It cannot be called on another instance.
+This function can only be called on 'this' in a blueprint. It cannot be called on another instance.
+
+## Usage
+
+    UFUNCTION(meta=(BlueprintProtected))
+    ReturnType FunctionName([Parameter, Parameter, ...])
diff --git a/Engine/Documentation/Source/Programming/UnrealArchitecture/Reference/Metadata/BlueprintSpawnableComponent/ClassMetadataSpecifier_BlueprintSpawnableComponent.INT.udn b/Engine/Documentation/Source/Programming/UnrealArchitecture/Reference/Metadata/BlueprintSpawnableComponent/ClassMetadataSpecifier_BlueprintSpawnableComponent.INT.udn
--- a/Engine/Documentation/Source/Programming/UnrealArchitecture/Reference/Metadata/BlueprintSpawnableComponent/ClassMetadataSpecifier_BlueprintSpawnableComponent.INT.udn
+++ b/Engine/Documentation/Source/Programming/UnrealArchitecture/Reference/Metadata/BlueprintSpawnableComponent/ClassMetadataSpecifier_BlueprintSpawnableComponent.INT.udn
@@ -1,7 +1,15 @@
 Title: BlueprintSpawnableComponent
 Crumbs: %ROOT%, Programming, Programming/UnrealArchitecture/Reference, Programming/UnrealArchitecture/Reference/Classes, Programming/UnrealArchitecture/Reference/Classes/Specifiers
-Description: 
+Description:
 Availability:Public
 Version: 4.9
 
- If present, the component class can be spawned by a Blueprint.
+If present, the component class can be spawned by a Blueprint.
+
+## Usage
+
+    UCLASS(meta=(BlueprintSpawnableComponent))
+    class ClassName : UActorComponent
+    {
+        ...
+    }
diff --git a/Engine/Documentation/Source/Programming/UnrealArchitecture/Reference/Metadata/CannotImplementInterfaceBlueprint/InterfaceMetadataSpecifier_CannotImplementInterfaceInBlueprint.INT.udn b/Engine/Documentation/Source/Programming/UnrealArchitecture/Reference/Metadata/CannotImplementInterfaceBlueprint/InterfaceMetadataSpecifier_CannotImplementInterfaceInBlueprint.INT.udn
--- a/Engine/Documentation/Source/Programming/UnrealArchitecture/Reference/Metadata/CannotImplementInterfaceBlueprint/InterfaceMetadataSpecifier_CannotImplementInterfaceInBlueprint.INT.udn
+++ b/Engine/Documentation/Source/Programming/UnrealArchitecture/Reference/Metadata/CannotImplementInterfaceBlueprint/InterfaceMetadataSpecifier_CannotImplementInterfaceInBlueprint.INT.udn
@@ -1,7 +1,17 @@
 Title: CannotImplementInterfaceInBlueprint
 Crumbs: %ROOT%, Programming, Programming/UnrealArchitecture/Reference, Programming/UnrealArchitecture/Reference/Interfaces, Programming/UnrealArchitecture/Reference/Interfaces/Specifiers
-Description: 
+Description:
 Availability:Public
 Version: 4.9
 
 This interface cannot be implemented by a Blueprint (e.g., it has only non-exposed C++ member methods).
+
+Any interfaces that define a function that isn't a [BlueprintNativeEvent](Programming/UnrealArchitecture/Reference/Functions/Specifiers/BlueprintImplementableEvent) or [BlueprintImplementableEvent](Programming/UnrealArchitecture/Reference/Functions/Specifiers/BlueprintNativeEvent) must be declared as `CannotImplementInterfaceInBlueprint`.
+
+## Usage
+
+    UINTERFACE(meta=(CannotImplementInterfaceInBlueprint))
+    class ClassName : UInterface
+    {
+        GENERATED_UINTERFACE_BODY()
+    };
diff --git a/Engine/Documentation/Source/Programming/UnrealArchitecture/Reference/Metadata/DeprecatedFunction/FunctionMetadataSpecifier_DeprecatedFunction.INT.udn b/Engine/Documentation/Source/Programming/UnrealArchitecture/Reference/Metadata/DeprecatedFunction/FunctionMetadataSpecifier_DeprecatedFunction.INT.udn
--- a/Engine/Documentation/Source/Programming/UnrealArchitecture/Reference/Metadata/DeprecatedFunction/FunctionMetadataSpecifier_DeprecatedFunction.INT.udn
+++ b/Engine/Documentation/Source/Programming/UnrealArchitecture/Reference/Metadata/DeprecatedFunction/FunctionMetadataSpecifier_DeprecatedFunction.INT.udn
@@ -1,7 +1,12 @@
 Title: DeprecatedFunction
 Crumbs: %ROOT%, Programming, Programming/UnrealArchitecture/Reference, Programming/UnrealArchitecture/Reference/Functions, Programming/UnrealArchitecture/Reference/Functions/Specifiers
-Description: 
+Description:
 Availability:Public
 Version: 4.9
 
-This function is deprecated, any blueprint references to it cause a compilation warning.
+This function is deprecated, any blueprint references to it cause a compilation warning. You can customize the deprecation message using the [DeprecationMessage](Programming/UnrealArchitecture/Reference/Metadata/DeprecationMessage) metadata specifier.
+
+## Usage
+
+    UFUNCTION(meta=(DeprecatedFunction))
+    ReturnType FunctionName([Parameter, Parameter, ...])
diff --git a/Engine/Documentation/Source/Programming/UnrealArchitecture/Reference/Metadata/DeprecationMessage/FunctionMetadataSpecifier_DeprecationMessage.INT.udn b/Engine/Documentation/Source/Programming/UnrealArchitecture/Reference/Metadata/DeprecationMessage/FunctionMetadataSpecifier_DeprecationMessage.INT.udn
--- a/Engine/Documentation/Source/Programming/UnrealArchitecture/Reference/Metadata/DeprecationMessage/FunctionMetadataSpecifier_DeprecationMessage.INT.udn
+++ b/Engine/Documentation/Source/Programming/UnrealArchitecture/Reference/Metadata/DeprecationMessage/FunctionMetadataSpecifier_DeprecationMessage.INT.udn
@@ -1,7 +1,12 @@
 Title: DeprecationMessage
 Crumbs: %ROOT%, Programming, Programming/UnrealArchitecture/Reference, Programming/UnrealArchitecture/Reference/Functions, Programming/UnrealArchitecture/Reference/Functions/Specifiers
-Description: 
+Description:
 Availability:Public
 Version: 4.9
 
-Supplies the custom message to use for deprecation.
+Supplies the custom message to use for deprecation. Must be combined with the [DeprecatedFunction](Programming/UnrealArchitecture/Reference/Metadata/DeprecatedFunction) metadata specifier.
+
+## Usage
+
+    UFUNCTION(meta=(DeprecatedFunction,DeprecationMessage="This function is deprecated, please use OtherFunctionName instead."))
+    ReturnType FunctionName([Parameter, Parameter, ...])
diff --git a/Engine/Documentation/Source/Programming/UnrealArchitecture/Reference/Metadata/MetadataSpecifiers.INT.udn b/Engine/Documentation/Source/Programming/UnrealArchitecture/Reference/Metadata/MetadataSpecifiers.INT.udn
--- a/Engine/Documentation/Source/Programming/UnrealArchitecture/Reference/Metadata/MetadataSpecifiers.INT.udn
+++ b/Engine/Documentation/Source/Programming/UnrealArchitecture/Reference/Metadata/MetadataSpecifiers.INT.udn
@@ -3,13 +3,16 @@ Crumbs:%ROOT%, Programming, Programming/UnrealArchitecture/Reference
 Description: Metadata keywords used when declaring classes, functions, and interfaces to specify how they behave with various aspects of the engine and editor.
 Availability:Public
 
-When declaring classes, functions, and interfaces, metadata specifiers can be added to the declaration to control how they behave with various 
+When declaring classes, functions, and interfaces, metadata specifiers can be added to the declaration to control how they behave with various
 aspects of the engine and editor.
 
 [EXCERPT:MetadataUsage]
-The usage of metadata specifiers differs from regular class, function, and interface specifiers.  
+The usage of metadata specifiers differs from regular class, function, and interface specifiers.
 [/EXCERPT:MetadataUsage]
 
+To use metadata specifies, use the meta specifier in a `UCLASS`, `UINTERFACE`, or `UFUNCTION` macro:
+
+    UCLASS(meta=(BlueprintSpawnableComponent))
 
 ## Class Metadata Specifiers
 [EXCERPT:ClassMetadata]
diff --git a/Engine/Documentation/Source/Programming/UnrealArchitecture/Reference/Metadata/UnsafeDuringActorConstruction/FunctionMetadataSpecifier_UnsafeDuringActorConstruction.INT.udn b/Engine/Documentation/Source/Programming/UnrealArchitecture/Reference/Metadata/UnsafeDuringActorConstruction/FunctionMetadataSpecifier_UnsafeDuringActorConstruction.INT.udn
--- a/Engine/Documentation/Source/Programming/UnrealArchitecture/Reference/Metadata/UnsafeDuringActorConstruction/FunctionMetadataSpecifier_UnsafeDuringActorConstruction.INT.udn
+++ b/Engine/Documentation/Source/Programming/UnrealArchitecture/Reference/Metadata/UnsafeDuringActorConstruction/FunctionMetadataSpecifier_UnsafeDuringActorConstruction.INT.udn
@@ -1,7 +1,12 @@
 Title: UnsafeDuringActorConstruction
 Crumbs: %ROOT%, Programming, Programming/UnrealArchitecture/Reference, Programming/UnrealArchitecture/Reference/Functions, Programming/UnrealArchitecture/Reference/Functions/Specifiers
-Description: 
+Description:
 Availability:Public
 Version: 4.9
 
 This function is not safe to call during Actor construction.
+
+## Usage
+
+    UFUNCTION(meta=(UnsafeDuringActorConstruction))
+    ReturnType FunctionName([Parameter, Parameter, ...])
diff --git a/Engine/Documentation/Source/Programming/UnrealArchitecture/Reference/Properties/Properties.INT.udn b/Engine/Documentation/Source/Programming/UnrealArchitecture/Reference/Properties/Properties.INT.udn
--- a/Engine/Documentation/Source/Programming/UnrealArchitecture/Reference/Properties/Properties.INT.udn
+++ b/Engine/Documentation/Source/Programming/UnrealArchitecture/Reference/Properties/Properties.INT.udn
@@ -6,18 +6,17 @@ Version: 4.9
 
 ## Property Declaration
 
-UProperty variables are declared using standard C++ syntax with additional descriptors, such as variable specifiers and metadata placed above the declaration.
+Properties are declared using standard C++ variable syntax, preceeded by the UPROPERTY macro which defines property metadata and variable specifiers.
 
 	UPROPERTY([specifier, specifier, ...], [meta(key=value, key=value, ...)])
 	Type VariableName;
 
 ## Core Data Types
 
-
 ### Integers
 The convention for integral data types is "int" or "uint" followed by the size in bits.
 
-|Variable Type | Description | 
+|Variable Type | Description |
 | ------------| -------------|
 | **uint8**	 | 8-bit  unsigned|
 | **uint16**	| 16-bit unsigned|
@@ -35,6 +34,7 @@ Unreal uses the standard C++ floating point types, float and double.
 Boolean types can be represented either with the C++ bool keyword or as a bitfield.
 
 	uint32 bIsHungry : 1;
+    bool bIsThirsty;
 
 ### Strings
 Unreal supports 3 core types of strings.
@@ -48,9 +48,7 @@ For most uses, Unreal relies on the TCHAR type for characters.  The TEXT() macro
 
 For more on the three string types, when to use each one, and how to work with them, see the **[String Handling documentation](Programming/UnrealArchitecture/StringHandling)**.
 
-
-
-### Property Specifiers
+## Property Specifiers
 
 [INCLUDE:Programming/UnrealArchitecture/Reference/Properties/Specifiers#main(offset:2)]
 
diff --git a/Engine/Documentation/Source/Programming/UnrealArchitecture/Reference/Properties/Specifiers/BlueprintAssignable/PropertySpecifier_BlueprintAssignable.INT.udn b/Engine/Documentation/Source/Programming/UnrealArchitecture/Reference/Properties/Specifiers/BlueprintAssignable/PropertySpecifier_BlueprintAssignable.INT.udn
--- a/Engine/Documentation/Source/Programming/UnrealArchitecture/Reference/Properties/Specifiers/BlueprintAssignable/PropertySpecifier_BlueprintAssignable.INT.udn
+++ b/Engine/Documentation/Source/Programming/UnrealArchitecture/Reference/Properties/Specifiers/BlueprintAssignable/PropertySpecifier_BlueprintAssignable.INT.udn
@@ -1,7 +1,27 @@
 Title: BlueprintAssignable
 Crumbs: %ROOT%, Programming, Programming/UnrealArchitecture/Reference, Programming/UnrealArchitecture/Reference/Properties, Programming/UnrealArchitecture/Reference/Properties/Specifiers
-Description: 
+Description:
 Availability:Public
 Version: 4.9
 
-Multicast Delegates only.  Property should be exposed for assigning in Blueprints.
+Usable with Multicast Delegates only. Exposes the property for assigning in Blueprints.
+
+## Usage
+
+    UPROPERTY(BlueprintAssignable)
+    Type VariableName;
+
+## Example
+
+    DECLARE_DYNAMIC_MULTICAST_DELEGATE_OneParam(FOnItemRemoved, FItemStruct, RemovedItem);
+
+    UCLASS(meta=(BlueprintSpawnableComponent))
+    class UInventoryComponent : public UActorComponent
+    {
+        GENERATED_BODY()
+
+    public:
+
+        UPROPERTY(BlueprintAssignable)
+        FOnItemRemoved OnItemRemoved;
+    };
diff --git a/Engine/Documentation/Source/Programming/UnrealArchitecture/Reference/Properties/Specifiers/BlueprintCallable/PropertySpecifier_BlueprintCallable.INT.udn b/Engine/Documentation/Source/Programming/UnrealArchitecture/Reference/Properties/Specifiers/BlueprintCallable/PropertySpecifier_BlueprintCallable.INT.udn
--- a/Engine/Documentation/Source/Programming/UnrealArchitecture/Reference/Properties/Specifiers/BlueprintCallable/PropertySpecifier_BlueprintCallable.INT.udn
+++ b/Engine/Documentation/Source/Programming/UnrealArchitecture/Reference/Properties/Specifiers/BlueprintCallable/PropertySpecifier_BlueprintCallable.INT.udn
@@ -1,7 +1,12 @@
 Title: BlueprintCallable
 Crumbs: %ROOT%, Programming, Programming/UnrealArchitecture/Reference, Programming/UnrealArchitecture/Reference/Properties, Programming/UnrealArchitecture/Reference/Properties/Specifiers
-Description: 
+Description:
 Availability:Public
 Version: 4.9
 
- Multicast Delegates only.  Property should be exposed for calling in Blueprint code.
+Multicast Delegates only.  Property should be exposed for calling in Blueprint code.
+
+## Usage
+
+    UPROPERTY(BlueprintCallable)
+    Type VariableName;
diff --git a/Engine/Documentation/Source/Programming/UnrealArchitecture/Reference/Properties/Specifiers/BlueprintReadOnly/PropertySpecifier_BlueprintReadOnly.INT.udn b/Engine/Documentation/Source/Programming/UnrealArchitecture/Reference/Properties/Specifiers/BlueprintReadOnly/PropertySpecifier_BlueprintReadOnly.INT.udn
--- a/Engine/Documentation/Source/Programming/UnrealArchitecture/Reference/Properties/Specifiers/BlueprintReadOnly/PropertySpecifier_BlueprintReadOnly.INT.udn
+++ b/Engine/Documentation/Source/Programming/UnrealArchitecture/Reference/Properties/Specifiers/BlueprintReadOnly/PropertySpecifier_BlueprintReadOnly.INT.udn
@@ -1,7 +1,12 @@
 Title: BlueprintReadOnly
 Crumbs: %ROOT%, Programming, Programming/UnrealArchitecture/Reference, Programming/UnrealArchitecture/Reference/Properties, Programming/UnrealArchitecture/Reference/Properties/Specifiers
-Description: 
+Description:
 Availability:Public
 Version: 4.9
 
-This property can be read by Blueprints, but not modified.
+This property can be read by Blueprints, but not modified. This operator is incompatible with the BlueprintReadWrite specifier.
+
+## Usage
+
+    UPROPERTY(BlueprintReadOnly)
+    Type VariableName;
diff --git a/Engine/Documentation/Source/Programming/UnrealArchitecture/Reference/Properties/Specifiers/BlueprintReadWrite/PropertySpecifier_BlueprintReadWrite.INT.udn b/Engine/Documentation/Source/Programming/UnrealArchitecture/Reference/Properties/Specifiers/BlueprintReadWrite/PropertySpecifier_BlueprintReadWrite.INT.udn
--- a/Engine/Documentation/Source/Programming/UnrealArchitecture/Reference/Properties/Specifiers/BlueprintReadWrite/PropertySpecifier_BlueprintReadWrite.INT.udn
+++ b/Engine/Documentation/Source/Programming/UnrealArchitecture/Reference/Properties/Specifiers/BlueprintReadWrite/PropertySpecifier_BlueprintReadWrite.INT.udn
@@ -1,7 +1,12 @@
 Title: BlueprintReadWrite
 Crumbs: %ROOT%, Programming, Programming/UnrealArchitecture/Reference, Programming/UnrealArchitecture/Reference/Properties, Programming/UnrealArchitecture/Reference/Properties/Specifiers
-Description: 
+Description:
 Availability:Public
 Version: 4.9
 
-This property can be read or written from a Blueprint.
+This property can be read or written from a Blueprint. This operator is incompatible with the BlueprintReadOnly specifier.
+
+## Usage
+
+    UPROPERTY(BlueprintReadWrite)
+    Type VariableName;
diff --git a/Engine/Documentation/Source/Programming/UnrealArchitecture/Reference/Properties/Specifiers/Category/PropertySpecifier_Category.INT.udn b/Engine/Documentation/Source/Programming/UnrealArchitecture/Reference/Properties/Specifiers/Category/PropertySpecifier_Category.INT.udn
--- a/Engine/Documentation/Source/Programming/UnrealArchitecture/Reference/Properties/Specifiers/Category/PropertySpecifier_Category.INT.udn
+++ b/Engine/Documentation/Source/Programming/UnrealArchitecture/Reference/Properties/Specifiers/Category/PropertySpecifier_Category.INT.udn
@@ -1,9 +1,21 @@
 Title: Category
 Crumbs: %ROOT%, Programming, Programming/UnrealArchitecture/Reference, Programming/UnrealArchitecture/Reference/Properties, Programming/UnrealArchitecture/Reference/Properties/Specifiers
-Description: 
+Description:
 Availability:Public
 Version: 4.9
 
-Specifies the category of the property. 
+Specifies the category of the property. Define nested categories using the | operator.
 
-* Usage: Category=CategoryName.
+## Usage
+
+    UPROPERTY(Category=CategoryName)
+    Type VariableName;
+
+
+    UPROPERTY(Category="CategoryName|SubCategoryName")
+    Type VariableName;
+
+## Example
+
+    UPROPERTY(Category="Player")
+    float Health;
diff --git a/Engine/Documentation/Source/Programming/UnrealArchitecture/Reference/Properties/Specifiers/Const/PropertySpecifier_Const.INT.udn b/Engine/Documentation/Source/Programming/UnrealArchitecture/Reference/Properties/Specifiers/Const/PropertySpecifier_Const.INT.udn
--- a/Engine/Documentation/Source/Programming/UnrealArchitecture/Reference/Properties/Specifiers/Const/PropertySpecifier_Const.INT.udn
+++ b/Engine/Documentation/Source/Programming/UnrealArchitecture/Reference/Properties/Specifiers/Const/PropertySpecifier_Const.INT.udn
@@ -1,7 +1,7 @@
 Title: Const
 Crumbs: %ROOT%, Programming, Programming/UnrealArchitecture/Reference, Programming/UnrealArchitecture/Reference/Properties, Programming/UnrealArchitecture/Reference/Properties/Specifiers
-Description: 
+Description:
 Availability:Public
 Version: 4.9
 
-This variable is const and should be exported as const.
+This variable is `const` and should be exported as `const`. `const` properties will not be modifiable in the editor.
diff --git a/Engine/Documentation/Source/Programming/UnrealArchitecture/Reference/Properties/Specifiers/EditAnywhere/PropertySpecifier_EditAnywhere.INT.udn b/Engine/Documentation/Source/Programming/UnrealArchitecture/Reference/Properties/Specifiers/EditAnywhere/PropertySpecifier_EditAnywhere.INT.udn
--- a/Engine/Documentation/Source/Programming/UnrealArchitecture/Reference/Properties/Specifiers/EditAnywhere/PropertySpecifier_EditAnywhere.INT.udn
+++ b/Engine/Documentation/Source/Programming/UnrealArchitecture/Reference/Properties/Specifiers/EditAnywhere/PropertySpecifier_EditAnywhere.INT.udn
@@ -1,7 +1,12 @@
 Title: EditAnywhere
 Crumbs: %ROOT%, Programming, Programming/UnrealArchitecture/Reference, Programming/UnrealArchitecture/Reference/Properties, Programming/UnrealArchitecture/Reference/Properties/Specifiers
-Description: 
+Description:
 Availability:Public
 Version: 4.9
 
-Indicates that this property can be edited by property windows in the editor.
+Indicates that this property can be edited by property windows, on archetypes and instances.
+
+## Usage
+
+    UPROPERTY(EditAnywhere)
+    Type VariableName;
diff --git a/Engine/Documentation/Source/Programming/UnrealArchitecture/Reference/Properties/Specifiers/EditDefaultsOnly/PropertySpecifier_EditDefaultsOnly.INT.udn b/Engine/Documentation/Source/Programming/UnrealArchitecture/Reference/Properties/Specifiers/EditDefaultsOnly/PropertySpecifier_EditDefaultsOnly.INT.udn
--- a/Engine/Documentation/Source/Programming/UnrealArchitecture/Reference/Properties/Specifiers/EditDefaultsOnly/PropertySpecifier_EditDefaultsOnly.INT.udn
+++ b/Engine/Documentation/Source/Programming/UnrealArchitecture/Reference/Properties/Specifiers/EditDefaultsOnly/PropertySpecifier_EditDefaultsOnly.INT.udn
@@ -1,7 +1,12 @@
 Title: EditDefaultsOnly
 Crumbs: %ROOT%, Programming, Programming/UnrealArchitecture/Reference, Programming/UnrealArchitecture/Reference/Properties, Programming/UnrealArchitecture/Reference/Properties/Specifiers
-Description: 
+Description:
 Availability:Public
 Version: 4.9
 
-Indicates that this property can be edited by property windows, but only on archetypes.
+Indicates that this property can be edited by property windows, but only on archetypes. This operator is incompatible with the Visible* specifiers.
+
+## Usage
+
+    UPROPERTY(EditDefaultsOnly)
+    Type VariableName;
diff --git a/Engine/Documentation/Source/Programming/UnrealArchitecture/Reference/Properties/Specifiers/EditFixedSize/PropertySpecifier_EditFixedSize.INT.udn b/Engine/Documentation/Source/Programming/UnrealArchitecture/Reference/Properties/Specifiers/EditFixedSize/PropertySpecifier_EditFixedSize.INT.udn
--- a/Engine/Documentation/Source/Programming/UnrealArchitecture/Reference/Properties/Specifiers/EditFixedSize/PropertySpecifier_EditFixedSize.INT.udn
+++ b/Engine/Documentation/Source/Programming/UnrealArchitecture/Reference/Properties/Specifiers/EditFixedSize/PropertySpecifier_EditFixedSize.INT.udn
@@ -1,7 +1,12 @@
 Title: EditFixedSize
 Crumbs: %ROOT%, Programming, Programming/UnrealArchitecture/Reference, Programming/UnrealArchitecture/Reference/Properties, Programming/UnrealArchitecture/Reference/Properties/Specifiers
-Description: 
+Description:
 Availability:Public
 Version: 4.9
 
 Only useful for dynamic arrays. This will prevent the user from changing the length of an array via the Unreal Editor property window.
+
+## Usage
+
+    UPROPERTY(EditFixedSize)
+    Type VariableName;
diff --git a/Engine/Documentation/Source/Programming/UnrealArchitecture/Reference/Properties/Specifiers/EditInline/PropertySpecifier_EditInline.INT.udn b/Engine/Documentation/Source/Programming/UnrealArchitecture/Reference/Properties/Specifiers/EditInline/PropertySpecifier_EditInline.INT.udn
--- a/Engine/Documentation/Source/Programming/UnrealArchitecture/Reference/Properties/Specifiers/EditInline/PropertySpecifier_EditInline.INT.udn
+++ b/Engine/Documentation/Source/Programming/UnrealArchitecture/Reference/Properties/Specifiers/EditInline/PropertySpecifier_EditInline.INT.udn
@@ -1,7 +1,12 @@
 Title: EditInline
 Crumbs: %ROOT%, Programming, Programming/UnrealArchitecture/Reference, Programming/UnrealArchitecture/Reference/Properties, Programming/UnrealArchitecture/Reference/Properties/Specifiers
-Description: 
+Description:
 Availability:Public
 Version: 4.9
 
-Allows the user to edit the properties of the object referenced by this variable within Unreal Editor's property inspector (only useful for object references, including arrays of object reference). 
+Allows the user to edit the properties of the object referenced by this variable within Unreal Editor's property inspector (only useful for object references, including arrays of object reference).
+
+## Usage
+
+    UPROPERTY(EditInline)
+    Type VariableName;
diff --git a/Engine/Documentation/Source/Programming/UnrealArchitecture/Reference/Properties/Specifiers/EditInstanceOnly/PropertySpecifier_EditInstanceOnly.INT.udn b/Engine/Documentation/Source/Programming/UnrealArchitecture/Reference/Properties/Specifiers/EditInstanceOnly/PropertySpecifier_EditInstanceOnly.INT.udn
--- a/Engine/Documentation/Source/Programming/UnrealArchitecture/Reference/Properties/Specifiers/EditInstanceOnly/PropertySpecifier_EditInstanceOnly.INT.udn
+++ b/Engine/Documentation/Source/Programming/UnrealArchitecture/Reference/Properties/Specifiers/EditInstanceOnly/PropertySpecifier_EditInstanceOnly.INT.udn
@@ -1,7 +1,12 @@
 Title: EditInstanceOnly
 Crumbs: %ROOT%, Programming, Programming/UnrealArchitecture/Reference, Programming/UnrealArchitecture/Reference/Properties, Programming/UnrealArchitecture/Reference/Properties/Specifiers
-Description: 
+Description:
 Availability:Public
 Version: 4.9
 
-Indicates that this property can be edited by property windows, but only on instances, not on archetypes.
+Indicates that this property can be edited by property windows, but only on instances, not on archetypes. This operator is incompatible with the Visible* specifiers.
+
+## Usage
+
+    UPROPERTY(EditInstanceOnly)
+    Type VariableName;
diff --git a/Engine/Documentation/Source/Programming/UnrealArchitecture/Reference/Properties/Specifiers/Transient/PropertySpecifier_Transient.INT.udn b/Engine/Documentation/Source/Programming/UnrealArchitecture/Reference/Properties/Specifiers/Transient/PropertySpecifier_Transient.INT.udn
--- a/Engine/Documentation/Source/Programming/UnrealArchitecture/Reference/Properties/Specifiers/Transient/PropertySpecifier_Transient.INT.udn
+++ b/Engine/Documentation/Source/Programming/UnrealArchitecture/Reference/Properties/Specifiers/Transient/PropertySpecifier_Transient.INT.udn
@@ -1,7 +1,12 @@
 Title: Transient
 Crumbs: %ROOT%, Programming, Programming/UnrealArchitecture/Reference, Programming/UnrealArchitecture/Reference/Properties, Programming/UnrealArchitecture/Reference/Properties/Specifiers
-Description: 
+Description:
 Availability:Public
 Version: 4.9
 
 Property is transient: should not be saved, zero-filled at load time.
+
+## Usage
+
+    UPROPERTY(Transient)
+    Type VariableName;
diff --git a/Engine/Documentation/Source/Programming/UnrealArchitecture/Reference/Properties/Specifiers/VisibleAnywhere/PropertySpecifier_VisibleAnywhere.INT.udn b/Engine/Documentation/Source/Programming/UnrealArchitecture/Reference/Properties/Specifiers/VisibleAnywhere/PropertySpecifier_VisibleAnywhere.INT.udn
--- a/Engine/Documentation/Source/Programming/UnrealArchitecture/Reference/Properties/Specifiers/VisibleAnywhere/PropertySpecifier_VisibleAnywhere.INT.udn
+++ b/Engine/Documentation/Source/Programming/UnrealArchitecture/Reference/Properties/Specifiers/VisibleAnywhere/PropertySpecifier_VisibleAnywhere.INT.udn
@@ -1,7 +1,12 @@
 Title: VisibleAnywhere
 Crumbs: %ROOT%, Programming, Programming/UnrealArchitecture/Reference, Programming/UnrealArchitecture/Reference/Properties, Programming/UnrealArchitecture/Reference/Properties/Specifiers
-Description: 
+Description:
 Availability:Public
 Version: 4.9
 
-Indicates that this property is visible in property windows, but cannot be edited at all.
+Indicates that this property is visible in property windows, but cannot be edited at all. This operator is incompatible with the Edit* specifiers.
+
+## Usage
+
+    UPROPERTY(VisibleAnywhere)
+    Type VariableName;
diff --git a/Engine/Documentation/Source/Programming/UnrealArchitecture/Reference/Properties/Specifiers/VisibleDefaultsOnly/PropertySpecifier_VisibleDefaultsOnly.INT.udn b/Engine/Documentation/Source/Programming/UnrealArchitecture/Reference/Properties/Specifiers/VisibleDefaultsOnly/PropertySpecifier_VisibleDefaultsOnly.INT.udn
--- a/Engine/Documentation/Source/Programming/UnrealArchitecture/Reference/Properties/Specifiers/VisibleDefaultsOnly/PropertySpecifier_VisibleDefaultsOnly.INT.udn
+++ b/Engine/Documentation/Source/Programming/UnrealArchitecture/Reference/Properties/Specifiers/VisibleDefaultsOnly/PropertySpecifier_VisibleDefaultsOnly.INT.udn
@@ -1,7 +1,12 @@
 Title: VisibleDefaultsOnly
 Crumbs: %ROOT%, Programming, Programming/UnrealArchitecture/Reference, Programming/UnrealArchitecture/Reference/Properties, Programming/UnrealArchitecture/Reference/Properties/Specifiers
-Description: 
+Description:
 Availability:Public
 Version: 4.9
 
-Indicates that this property is only visible in property windows for archetypes, and cannot be edited.
+Indicates that this property is only visible in property windows for archetypes, and cannot be edited. This operator is incompatible with the Edit* specifiers.
+
+## Usage
+
+    UPROPERTY(VisibleDefaultsOnly)
+    Type VariableName;
diff --git a/Engine/Documentation/Source/Programming/UnrealArchitecture/Reference/Properties/Specifiers/VisibleInstanceOnly/PropertySpecifier_VisibleInstanceOnly.INT.udn b/Engine/Documentation/Source/Programming/UnrealArchitecture/Reference/Properties/Specifiers/VisibleInstanceOnly/PropertySpecifier_VisibleInstanceOnly.INT.udn
--- a/Engine/Documentation/Source/Programming/UnrealArchitecture/Reference/Properties/Specifiers/VisibleInstanceOnly/PropertySpecifier_VisibleInstanceOnly.INT.udn
+++ b/Engine/Documentation/Source/Programming/UnrealArchitecture/Reference/Properties/Specifiers/VisibleInstanceOnly/PropertySpecifier_VisibleInstanceOnly.INT.udn
@@ -1,7 +1,12 @@
 Title: VisibleInstanceOnly
 Crumbs: %ROOT%, Programming, Programming/UnrealArchitecture/Reference, Programming/UnrealArchitecture/Reference/Properties, Programming/UnrealArchitecture/Reference/Properties/Specifiers
-Description: 
+Description:
 Availability:Public
 Version: 4.9
 
-Indicates that this property is only visible in property windows for instances, not for archetypes, and cannot be edited.
+Indicates that this property is only visible in property windows for instances, not for archetypes, and cannot be edited. This operator is incompatible with the Edit* specifiers.
+
+## Usage
+
+    UPROPERTY(VisibleInstanceOnly)
+    Type VariableName;
