#UE4doc #docs:
-In middle of making changes, checking back in for now-
Added engine version field
Removed Overview header
Changed font settings for multiple terms
Replaced a couple pictures due to changes to the editor's UI

diff --git a/Engine/Documentation/Source/Engine/Animation/Sequences/Notifies/AnimNotifies.INT.udn b/Engine/Documentation/Source/Engine/Animation/Sequences/Notifies/AnimNotifies.INT.udn
--- a/Engine/Documentation/Source/Engine/Animation/Sequences/Notifies/AnimNotifies.INT.udn
+++ b/Engine/Documentation/Source/Engine/Animation/Sequences/Notifies/AnimNotifies.INT.udn
@@ -2,6 +2,7 @@ Availability:Public
 Title:Animation Notifications (Notifies)
 Crumbs: %ROOT%, Engine, Engine/Animation, Engine/Animation/Sequences
 Description:Notifies are a system for setting up and receiving events in Animation Sequences to perform external actions.
+version: 4.9
 
 [VAR:Topic]
 [OBJECT:Topic]
@@ -26,7 +27,6 @@ Description:Notifies are a system for setting up and receiving events in Animati
 
 [TOC (start:2 end:2)] 
 
-## Overview
 
 [EXCERPT:Intro]
 **Animation Notifications** (**AnimNotifies** or just **Notifies**) provide a way for animation programmers to setup events to occur at specific points during an **Animation Sequence**. Notifies are commonly used to add effects like footstep sounds during a walk or run animation or spawning a particle system during animation. However, they can have any number of different uses, as the system can be extended with custom notification types to cater to the needs of any type of game.
@@ -43,11 +43,11 @@ Below you will find the currently available Notifies and how they are used.
 
 ### Play Particle Effect
 
-The **Play Particle Effect** notify spawns and plays a particle system at a given socket location and at a specific point along an animation's timeline. This is useful for such effects as causing a cloud of dust to appear around a player's feet when they move, fire to spew from an exhaust pipe when the player uses a Nitro-Boost power-up, or other such effects.
+The **PlayParticleEffect** notify spawns and plays a particle system at a given socket location and at a specific point along an animation's timeline. This is useful for such effects as causing a cloud of dust to appear around a player's feet when they move, fire to spew from an exhaust pipe when the player uses a Nitro-Boost power-up, or other such effects.
 
 ![](PlayParticleEffectNotify.png)
 
-A **Play Particle Effect** notify does not use the **Animation Blueprint** to control behavior. Instead, it uses a few simple properties defined in the **Details** panel to determine what **Particle System** is used and from what **Socket** the particles will originate.
+A PlayParticleEffect notify does not use the Animation Blueprint to control behavior. Instead, it uses a few simple properties defined in the **Details** panel to determine what Particle System is used and from what Socket the particles will originate.
 
 ![](PlayParticleEffectDetails.png)
 
@@ -67,7 +67,7 @@ Once set up, the associated Particle System will play, using the Socket location
 
 ### Play Sound
 
-The **Play Sound** notify provides a simple means to get your sound effects synced with animation. Much like _Play Particle Effect_, this notify has a few simple properties that can be set in the **Details** panel, which allow for fast and easy sound effect setup. This is useful for such things as mechanical whirrs when a robot moves, the rustle of cloth during character motion, and similar types of effects.
+The **PlaySound** notify provides a simple means to get your sound effects synced with animation. Much like Play Particle Effect, this notify has a few simple properties that can be set in the **Details** panel, which allow for fast and easy sound effect setup. This is useful for such things as mechanical whirrs when a robot moves, the rustle of cloth during character motion, and similar types of effects.
 
 ![](PlaySound.png)
 
@@ -85,7 +85,7 @@ The **Play Sound** notify uses a few simple properties defined in the **Details*
 
 ### Custom
 
-The **Custom** notify serves as a catch-all for anything an artist may want to do within a **Animation Blueprint** at a certain point in an animation. Custom notifies are created by **Right-Clicking** on a notify track and choosing **New Notify** from the context menu.
+The **Custom** notify serves as a catch-all for anything an artist may want to do within a Animation Blueprint at a certain point in an animation. Custom notifies are created by **Right-clicking** on a notify track and choosing **New Notify** from the context menu.
 
 ![](MySpecialNotify.png)
 
@@ -99,9 +99,10 @@ After the name is typed in, you will see the new notify appear in the timeline.
 
 Custom notifies will not have any properties to edit in the **Details** panel. Instead, they simply allow you to create a notify event of the same name in the **Event Graph** or any of the **Transitions** in a state machine of the Skeletal Mesh's Animation Blueprint.
 
-![](SwitchToAnimBP.png)
+Switch to the **Graph** edit mode
+<!--![](SwitchToAnimBP.png) Commented out picture as it didn't seem useful--> 
 
-**Right-Click** anywhere in the **Event Graph** and there is a section for adding Animation Notifies events.
+**Right-click** anywhere in the **Event Graph** and there is a section for adding Animation Notifies events.
 
 ![](rightClickAddNotify.png)
 
@@ -109,27 +110,27 @@ Custom notifies will not have any properties to edit in the **Details** panel. I
 
 ## Animation Notify States 
 
-**AnimNotifyStates** (**Notify States**) work much like the standard notifies above.  They have 3 distinct events: a start, the "tick," and an end. They start straightforward, firing at the moments the notify begins and ends, and the **Event Graphs** in them fire off when their time comes up in the animation. The tick fires off every animation update until the end event is hit. The major difference between the standard Notifies is that **AnimNotifyStates** are self-contained **Blueprints**.
+**AnimNotifyStates** (**Notify States**) work much like the standard notifies above.  They have 3 distinct events: a start, the tick, and an end. They start straightforward, firing at the moments the notify begins and ends, and the **Event Graphs** in them fire off when their time comes up in the animation. The tick fires off every animation update until the end event is hit. The major difference between the standard Notifies is that AnimNotifyStates are self-contained Blueprints.
 
-To use an **AnimNotifyState**:
+To use an AnimNotifyState:
 
-1. Create a new **AnimNotifyState** asset in the **Content Browser** by creating a new **Blueprint** that inherits from `AnimNotifyState`:
+1. Create a new AnimNotifyState asset in the **Content Browser** by creating a new **Blueprint** that inherits from `AnimNotifyState`:
 
 	![](newBP.png)
 
 	![](AnimNotifyStateCreate.png)
 
-1. On the **Animation Sequence** or **Animation Montage**, you can **Right-Click** in the notify area and now add your new **AnimNoitfyState**.
+1. On the Animation Sequence or Animation Montage, you can **Right-click** in the notify area and now add your new AnimNoitfyState.
 
 	![](addNotifyState.png)
 
-1. You can set the length of the **AnimNotifyState** by **Right-Clicking** on it and choosing **Set AnimNotifyState Duration**.
+1. You can set the length of the AnimNotifyState by **Right-clicking** on it and choosing **Set AnimNotifyState Duration**.
 
 	![](AnimNotifyStateSetLength.png)
 
 	![](duration.png)
 
-1. Finally from the **Content Browser**, **Double Left-Click** your **AnimNotifyState** to bring up its **Blueprint** so you can start defining what if does.
+1. Finally from the **Content Browser**, **Double Left-click** your AnimNotifyState to bring up its Blueprint so you can start defining what if does.
 	
 	[REGION:lightbox]
 		[![](animNotifyStateImplementFunction.png)(w:860)](animNotifyStateImplementFunction.png)
@@ -189,7 +190,7 @@ To use an **AnimTrail** emitter you must do the following;
 1.       Create a **Particle System** that contains an emitter with **AnimTrail** type data. 
 2.       In the **Animation Editor**, add an **AnimNotifyState** of type **Trail**. 
 
-_Setting Up the Particle System_
+Setting Up the Particle System
 
 Create a **Particle System** as usual and add a new emitter with type data **AnimTrail** to it. 
 Several properties of the type data module will affect how the trail behaves. Below is what you should see in **Cascade** when setting up the system. Note, you will see nothing in the preview window as trails must be used in conjunction with a **Skeletal Mesh** component as explained later. 
@@ -204,9 +205,9 @@ Below are the properties inside the **AnimTrail** type data section of the parti
 
 | **Property** | **Description** |
 | --- | --- |
-| **Dead Trails On Deactivate** | If _true_, all trails are marked as dead when the system is deactivated as opposed to allowing the particles to expire naturally. This should generally be left as _true_.  |
-| **Enable Previous Tangent Recalculation** | If _true_, recalculates the previous tangent when a new particle is spawned. Makes the system more expensive but better looking as tangents more accurately match the trails curve.  |
-| **Tangent Recalculation Every Frame** | If _true_, recalculates tangents every frame. Slower but allows the trail particles to use velocity and acceleration etc. |
+| **Dead Trails On Deactivate** | If true, all trails are marked as dead when the system is deactivated as opposed to allowing the particles to expire naturally. This should generally be left as true.  |
+| **Enable Previous Tangent Recalculation** | If true, recalculates the previous tangent when a new particle is spawned. Makes the system more expensive but better looking as tangents more accurately match the trails curve.  |
+| **Tangent Recalculation Every Frame** | If true, recalculates tangents every frame. Slower but allows the trail particles to use velocity and acceleration etc. |
 | **Tiling Distance** | The estimated distance in world units the trail has to travel before the 2nd UV set tiles. If 0.0, the 2nd UV set will not tile. |
 
 **Tessellation Parameters**
@@ -220,7 +221,7 @@ However, the more tessellated a trail is, the more expensive it is.
 | **Tangent Tessellation Step Size** | The (approx.) angle difference in degrees that the tangent must change from one particle to another to warrant another tessellation point. This allows you to ensure that any sharp turns in the trail are nice and smooth.  |
 | **Width Tessellation Step Size** | The (approx.) world unit change in the width of the trail that warrants an additional tessellation point.|
  
-_Setting Up the AnimNotify_
+Setting Up the AnimNotify
 
 Now that we have a particle system, we must use it with an **AnimNotify**. 
 Open the **Animation Editor**, add or use an existing **Timeline** under the **Notifies** section and **Right-Click** on it to bring up the menu. The **Trails** notify can be found under the **Add Notify State** section. 
