#UE4doc #docs - Removed Overview header and added engine version tag.

diff --git a/Engine/Documentation/Source/Engine/Blueprints/UserGuide/BlueprintComms/InterfaceMigration/InterfaceMigration.INT.udn b/Engine/Documentation/Source/Engine/Blueprints/UserGuide/BlueprintComms/InterfaceMigration/InterfaceMigration.INT.udn
--- a/Engine/Documentation/Source/Engine/Blueprints/UserGuide/BlueprintComms/InterfaceMigration/InterfaceMigration.INT.udn
+++ b/Engine/Documentation/Source/Engine/Blueprints/UserGuide/BlueprintComms/InterfaceMigration/InterfaceMigration.INT.udn
@@ -2,11 +2,11 @@ Availability:Public
 Title:Blueprint Interface Migration
 Crumbs: %ROOT%, Engine, Engine/Blueprints, Engine/Blueprints/UserGuide, Engine/Blueprints/UserGuide/BlueprintComms
 Description:The migration path for moving Blueprint Interface functions into Blueprints.
-
+Version: 4.9
 
 [TOC(start:2)]
 
-## Overview
+
 
 If you have been using Blueprint Interfaces to have Blueprints communicate with each other, you can easily convert the functions in the Blueprint Interfaces to normal Blueprint functions, so your Blueprints can communicate with each other directly.
 
