#UE4doc#docs - Added Engine version tag, bolded links where needed, and removed the Overview header.

diff --git a/Engine/Documentation/Source/Engine/Blueprints/UserGuide/BlueprintCommsUsage/BlueprintCommsUsage.INT.udn b/Engine/Documentation/Source/Engine/Blueprints/UserGuide/BlueprintCommsUsage/BlueprintCommsUsage.INT.udn
--- a/Engine/Documentation/Source/Engine/Blueprints/UserGuide/BlueprintCommsUsage/BlueprintCommsUsage.INT.udn
+++ b/Engine/Documentation/Source/Engine/Blueprints/UserGuide/BlueprintCommsUsage/BlueprintCommsUsage.INT.udn
@@ -6,6 +6,7 @@ Related: Engine\Blueprints\UserGuide\BlueprintComms
 Related: Engine/Blueprints/UserGuide/EventDispatcher
 Related: Engine/Blueprints/UserGuide/CastNodes
 Related: Engine/Blueprints/UserGuide/Types/Interface/UsingInterfaces
+Version: 4.9
 
 [VAR:Topic]
 [OBJECT:Topic]
@@ -50,9 +51,9 @@ Related: Engine/Blueprints/UserGuide/Types/Interface/UsingInterfaces
 
 [TOC(start:2 end:2)]
 
-When you are working with a **Blueprint**, you may sometimes find that you want to access the Functions, Events, or Variables that are contained in another **Blueprint**. Referred to as **Blueprint Communication** there are several methods in which you can access and communicate between multiple **Blueprints** depending upon your needs.
+When you are working with a Blueprint, you may sometimes find that you want to access the Functions, Events, or Variables that are contained in another Blueprint. This is referred to as **Blueprint Communication**. There are several methods in which you can access and communicate between multiple Blueprints depending upon your needs.
 
-Outlined below are the different methods of communication and sample use cases, refer to the links in each section for further documentation. 
+Outlined below are the different methods of communication and sample use cases. Refer to the links in each section for further documentation. 
 
 ## Direct Blueprint Communication
 
@@ -62,10 +63,10 @@ Outlined below are the different methods of communication and sample use cases,
 to each other at a certain point. 
 	
 This type of communication is always one-to-one, meaning one Blueprint (referred to as the **Working Blueprint**) is requesting to access another Blueprint (referred to as the **Target 
-Blueprint**). The simplest way to use **Blueprint Comms** is to get a reference to the **Target Blueprint** through a publicly exposed object variable and then specify which instance of 
+Blueprint**). The simplest way to use Blueprint Comms is to get a reference to the Target Blueprint through a publicly exposed object variable and then specify which instance of 
 that Blueprint you want to access.
 
-Refer to [](Engine\Blueprints\UserGuide\BlueprintComms) for more detailed information. 
+Refer to **[](Engine\Blueprints\UserGuide\BlueprintComms)** for more detailed information. 
 
 ### When to Use It...
 
@@ -74,7 +75,7 @@ Below are some examples when you might use **Blueprint Comms**:
 *	**You have an Actor in your level that you want to communicate with another Actor in your level.**
 	*	Interacting with a switch in your level opens a specific door or turns on a specific light (each of which are separate Blueprints). 
 
-Refer to [Setting Up Direct Blueprint Communications](Engine\Blueprints\BP_HowTo\DirectBlueprintComs) for a walk through of how to set one up. 
+Refer to [**Setting Up Direct Blueprint Communications**](Engine\Blueprints\BP_HowTo\DirectBlueprintComs) for a walk through of how to set one up. 
 
 
 	
@@ -94,7 +95,7 @@ To sum up, Event Dispatchers are used to tell anyone interested that something j
 or at different points in gameplay.
 
 
-Refer to [](Engine/Blueprints/UserGuide/EventDispatcher) for more detailed information. 
+Refer to **[](Engine/Blueprints/UserGuide/EventDispatcher)** for more detailed information. 
 
 ### When to Use It...
  
@@ -107,7 +108,7 @@ The examples below would be a case where an **Event Dispatcher** could be used:
 	*	You spawn a boss and an Event is fired when the boss is killed which spawns a reward in the world.
 	*	You spawn an item in your level (weapon, health, boost, etc.) and through an Event, notify the item and character when it is picked up.
 
-Refer to [Setting Up Event Dispatchers](Engine\Blueprints\BP_HowTo\EventDispatcher) for a walk through of how to set one up. 
+Refer to [**Setting Up Event Dispatchers**](Engine\Blueprints\BP_HowTo\EventDispatcher) for a walk through of how to set one up. 
 
 
 ## Blueprint Interfaces
@@ -117,7 +118,7 @@ Refer to [Setting Up Event Dispatchers](Engine\Blueprints\BP_HowTo\EventDispatch
 **Blueprint Interfaces** allow for a common method of interacting with multiple types of objects that all share some specific functionality. This means you can have completely different types of 
 objects, such as a car and a tree, that share one specific thing; they can both be shot by weapon fire and take damage. 
 
-By creating a **Blueprint Interface** that contains an **OnTakeWeaponFire** function and having both the car and the tree implement that **Blueprint Interface**, you can treat the car and the tree 
+By creating a Blueprint Interface that contains an **OnTakeWeaponFire** function and having both the car and the tree implement that Blueprint Interface, you can treat the car and the tree 
 as the same and simply call the **OnTakeWeaponFire** function when either of them is shot and have them respond differently.
 
 This is still a one-to-one style of communication, just like direct Blueprint communications, but unlike direct Blueprint communications, the Target may or may not receive the message.  An 
@@ -125,7 +126,7 @@ interface is a contract of sorts.  It says "if you implement this interface, you
 Blueprint communications and interfaces is that with direct communications, you have to know specifically whether or not the Blueprint on the receiving end cares about a function or not.  
 Interface Messages you can send to anything, and they'll respond if they care, and ignore the message if they don't. Going back to the weapon fire example, you may not want your walls to react to damage, so they just wouldn't implement the interface with **OnTakeWeaponFire**. 
 
-Refer to [](Engine/Blueprints/UserGuide/Types/Interface/UsingInterfaces) for more detailed information. 
+Refer to **[](Engine/Blueprints/UserGuide/Types/Interface/UsingInterfaces)** for more detailed information. 
 
 ### When to Use It...
 
@@ -136,7 +137,7 @@ Each of the examples below are scenarios where a **Blueprint Interface** might b
 	*	Player has a "use" button, pressing the use button opens doors, turns off lights, picks up items, etc. 
 	*	Enemy has a special ability that changes based on player health. 
 
-Refer to [Setting Up Blueprint Interfaces](Engine\Blueprints\BP_HowTo\Interfaces) for a walkthrough of how to set one up. 
+Refer to [**Setting Up Blueprint Interfaces**](Engine\Blueprints\BP_HowTo\Interfaces) for a walkthrough of how to set one up. 
 
 
 ## Casting For Communications
@@ -169,7 +170,7 @@ Here are a few examples when you could use **Blueprint Casting**:
 	*	You have several lights in your level and you want them all to turn on or off when an event occurs. 
 [/EXCERPT:Castusage]
 
-Refer to [Setting Up Blueprint Casts](Engine\Blueprints\BP_HowTo\Casting) for a walk through of how to set one up. 
+Refer to [**Setting Up Blueprint Casts**](Engine\Blueprints\BP_HowTo\Casting) for a walk through of how to set one up. 
 
 
 
