#UE4doc#docs - made edits reported by forum user

diff --git a/Engine/Documentation/Source/Engine/Rendering/Materials/HowTo/DetailTexturing/DetailTexturing.INT.udn b/Engine/Documentation/Source/Engine/Rendering/Materials/HowTo/DetailTexturing/DetailTexturing.INT.udn
--- a/Engine/Documentation/Source/Engine/Rendering/Materials/HowTo/DetailTexturing/DetailTexturing.INT.udn
+++ b/Engine/Documentation/Source/Engine/Rendering/Materials/HowTo/DetailTexturing/DetailTexturing.INT.udn
@@ -95,7 +95,7 @@ Once selected, a new **Material** will be created. Name the Material **Detail_Ma
 
 	![](DT_Adding_Textures.png)
 
-1. With all of the required Textures added to the Material Graph, it is time to add the **Detail Texturing** Material Function. To do this, first search in the **Palette** for **Detail Texturing**.
+1. With all of the required Textures added to the Material Graph, it is now time to add the **Detail Texturing** Material Function. To do this, first search in the **Palette** for **Detail Texturing**.
 
 	![](DT_Search_DT_MF.png)
 
@@ -103,7 +103,7 @@ Once selected, a new **Material** will be created. Name the Material **Detail_Ma
 
 	![](DT_Adding_DT_MF.png)
 
-1. With the **Detail Texturing** Material Function added, it is time to start hooking the required inputs and outputs so the Material will work. For this example, we are going to be using the **T_Brick_Clay_Old** Textures as our base Textures and the **T_Ground_Gravel** as our Detail Textures. This means that the Diffuse and Normal maps for **T_Brick_Clay_Old** should go into the Diffuse and **T_Ground_Moss** Normal inputs and the **T_Ground_Gravel** Diffuse and Detail Normal maps should go into the Detail Diffuse and Normal inputs. When completed, your Material Graph should look like this.
+1. With the **Detail Texturing** Material Function now added, it is now time to start hooking the required inputs and outputs so the Material will work. For this example, we are going to be using the **T_Brick_Clay_Old** Textures as our base Textures and the **T_Ground_Gravel** as our Detail Textures. This means that the Diffuse and Normal maps for **T_Brick_Clay_Old** should go into the Diffuse and Normal inputs and the **T_Ground_Gravel** Diffuse and **T_Ground_Moss** Normal maps should go into the Detail Diffuse and Detail Normal inputs. When completed, your Material Graph should look like this.
 
 	![](DT_Hooked_Up_Textures.png)
 
