#UE4doc#docs - Added line to Fresnel instructions to: Start by Updating the "Translucency Lighting Mode" by changing it from "Volumetric Non Directional" to "Surface Translucency Volume"

diff --git a/Engine/Documentation/Source/Engine/Rendering/Materials/HowTo/Fresnel/Fresnel.INT.udn b/Engine/Documentation/Source/Engine/Rendering/Materials/HowTo/Fresnel/Fresnel.INT.udn
--- a/Engine/Documentation/Source/Engine/Rendering/Materials/HowTo/Fresnel/Fresnel.INT.udn
+++ b/Engine/Documentation/Source/Engine/Rendering/Materials/HowTo/Fresnel/Fresnel.INT.udn
@@ -193,7 +193,9 @@ The Fresnel Material Expression can also be supplied with a Normal map that can
 This is particularly useful if you want to use the Fresnel to display something like an energy shield on a Material but still have that Material look normal when the shield is not active. 
 To achieve this using the Material above, the following things need to be done.
 
-1. First change the **Base Color** and **Roughness** to use the **T_Brick_Clay_New_D** texture from the Starter Content. 
+1. Start by Updating the **Translucency Lighting Mode** by changing it from **Volumetric Non Directional** to **Surface Translucency Volume**
+
+1. Change the **Base Color** and **Roughness** to use the **T_Brick_Clay_New_D** texture from the Starter Content. 
 
 1. Next change the Normal input so that it uses T_Brick_Clay_New_N, which can also be found in the Starter Content. 
 
