#UE4doc #docs - updated images, fixed step that was missing information and bolded links where needed.

diff --git a/Engine/Documentation/Source/Engine/Blueprints/UserGuide/BlueprintCommsUsage/BPComHowTo/BPComHowTo.INT.udn b/Engine/Documentation/Source/Engine/Blueprints/UserGuide/BlueprintCommsUsage/BPComHowTo/BPComHowTo.INT.udn
--- a/Engine/Documentation/Source/Engine/Blueprints/UserGuide/BlueprintCommsUsage/BPComHowTo/BPComHowTo.INT.udn
+++ b/Engine/Documentation/Source/Engine/Blueprints/UserGuide/BlueprintCommsUsage/BPComHowTo/BPComHowTo.INT.udn
@@ -15,17 +15,21 @@ Below we have two **Blueprints** in our level that we want to have communicate w
 
 ![](2_1.png)(w:400)
 
-* The cube above is a **Blueprint** created using the **Shape_Cube** mesh and its collision has been set to **OverlapOnlyPawn** so it acts as a trigger.
+* The cube above is a **Blueprint** created using the **Shape_Cube** mesh and its collision has been set to **OverlapOnlyPawn** so it acts as a trigger. Also enable **Generate Overlap Events**
 * The sparks above are the **Blueprint_Effect_Sparks** asset (included with starter content).
 
 Using **Direct Blueprint Communication**, we would do the following: 
 
 1. In the **Shape_Cube** Blueprint, under **MyBlueprint** click the ![](Engine\Blueprints\Editor\UIComponents\MyBlueprint\plus_Button.png) button on the Variables category: ![](Engine\Blueprints\Editor\UIComponents\MyBlueprint\myblueprint_variable.png). 
 
-1. In the **Details** panel, under **Variable Type** select the type of Blueprint you wish to access (Reference types are denoted by an appended **_C**).
+1. In the **Details** panel, under **Variable Type** select the type of Blueprint you wish to access.
 
 	![](2_3.png)
 
+	Hover over the **Blueprint_Effects_Sparks** Variable Type and select **Reference** from the list.
+
+	![](2_3_1.png)
+
 	Here we are stating that we want to access a **Blueprint_Effect_Sparks** Blueprint.
 
 1. In the **Details** panel, update the sections outlined below. 
@@ -80,7 +84,7 @@ In the example below, inside our **MyCharacter** Blueprint, we have stated that
 
 We can then access the **Blueprint_Effect_Smoke** Blueprint (contained in the yellow boxes above) and get the Smoke Effect and Smoke Audio components from our **Target Blueprint** and **Deactivate** them when **F** is pressed a second time (which is what the **Flip/Flop** node does). Therefore we are accessing one Blueprint from inside another Blueprint through **Direct Blueprint Communication**. 
 
-Refer to the [](Engine\Blueprints\UserGuide\BlueprintComms) documentation for more information. 
+Refer to the **[](Engine\Blueprints\UserGuide\BlueprintComms)** documentation for more information. 
 
 
 ### Blueprint Casting
@@ -172,7 +176,7 @@ In the image above for number 2, a **Widget Blueprint** is created and assigned
 [/EXCERPT:TargetCasting]
 
 ### Event Dispatchers
-Below we have a bush Actor **Blueprint** in our level that we want to recieve communication from the **Character Blueprint** when the player presses a button to set the bush on fire, then destroy the fire and the bush after a few seconds. We are wanting to communicate from the **Character Blueprint** to the **Level Blueprint** which can be done through using an **Event Dispatcher**. 
+Below we have a bush **Actor Blueprint** in our level that we want to recieve communication from the **Character Blueprint** when the player presses a button to set the bush on fire, then destroy the fire and the bush after a few seconds. We are wanting to communicate from the **Character Blueprint** to the **Level Blueprint** which can be done through using an **Event Dispatcher**. 
 
 ![](Event1.png)(w:340)
 
@@ -224,7 +228,7 @@ If you compile and save then play in the editor, you should see that when you pr
 
 This is of course a simple example and you would probably want more checking to occur (if the player is near the bush allow them to set it on fire) as well only allowing the player to start the fire once, but it illustrates how you can execute events inside the **Level Blueprint** by way of a **Character Blueprint** using an **Event Dispatcher**. The example above also shows how you can assign an **Event Dispatcher** to a Spawned Actor and execute Events when something occurs to that Actor (in this case, when it is destroyed).  
 
-Refer to the [](Engine\Blueprints\UserGuide\EventDispatcher) documentation for more information. 
+Refer to the **[](Engine\Blueprints\UserGuide\EventDispatcher)** documentation for more information. 
 
 
 ### Blueprint Interfaces
@@ -400,6 +404,6 @@ Here is how we would set up passing through those variables:
 
 	The second **Branch** node checks if the **PlayerHealth** value is 0 and if it is, text is printed to the screen that states "1 More Hit" before the character is "dead". This is by no means a perfect Health/Damage setup, but it illustrates how you can pass two variables through an Interface and how you could then use those variables in another Blueprint. The "PlayerHealth" value in this example could be passed to a HUD and updated to reflect current Health for example.  
 
-Refer to the [](Engine\Blueprints\UserGuide\Types\Interface) documentation for more information. 
+Refer to the **[](Engine\Blueprints\UserGuide\Types\Interface)** documentation for more information. 
 
 
