fixed typo

diff --git a/Engine/Documentation/Source/Engine/Rendering/PostProcessEffects/Blendables/Blendables.INT.udn b/Engine/Documentation/Source/Engine/Rendering/PostProcessEffects/Blendables/Blendables.INT.udn
--- a/Engine/Documentation/Source/Engine/Rendering/PostProcessEffects/Blendables/Blendables.INT.udn
+++ b/Engine/Documentation/Source/Engine/Rendering/PostProcessEffects/Blendables/Blendables.INT.udn
@@ -125,6 +125,6 @@ This was implemented in C++ code and only for this blendable type but later we s
 any data loss. In order to avoid a lot of assets spaming the content packages we would create the objects as part of the level. 
 * We want to further polish the blueprint interaction to make it easier to use.
 * We could easily expose a Blendable array in the world settings and the project settings.
-* In order to get more transpareceny which Blendable are getting applied we should have some debug view showing the weights and asset/object names and types.
+* In order to get more transparency which Blendable are getting applied we should have some debug view showing the weights and asset/object names and types.
 
 
