#UE4doc #docs
Completed first draft of 'Adding Projectiles' section.
Availability: Docs

diff --git a/Engine/Documentation/Source/Programming/Tutorials/FirstPersonShooter/3/1/FirstPersonShooter_3_1.INT.udn b/Engine/Documentation/Source/Programming/Tutorials/FirstPersonShooter/3/1/FirstPersonShooter_3_1.INT.udn
--- a/Engine/Documentation/Source/Programming/Tutorials/FirstPersonShooter/3/1/FirstPersonShooter_3_1.INT.udn
+++ b/Engine/Documentation/Source/Programming/Tutorials/FirstPersonShooter/3/1/FirstPersonShooter_3_1.INT.udn
@@ -22,6 +22,717 @@ Version: 4.9
 
 %Steps%
 
-<!--TODO: Add tutorial steps here-->
+Now that you've set-up your character, it's time to implement a simple projectile weapon so that when you fire, a simple grenade-like projectile will shoot from the center of the screen and fly until it collides with the world. During this step, you're going to add input and create a new code class for your projectile.
+
+## Adding Fire Action Mapping
+
+1.	In the **Edit** menu, click on **Project Settings**.
+
+1.	Under the **Engine** heading on the left side of the **Project Settings** tab, click on **Input**.
+
+1.	Under **Bindings**, click on the plus sign next to **Action Mappings**.
+
+1.	Click on the arrow to the left of **Action Mappings**.
+
+1.	Type "Fire" into the text field that appears, then click on the arrow to the left of the text box to expand the action binding options.
+
+1.	In the dropdown menu, select **Left Mouse Button** from the **Mouse** dropdown list.
+
+1.	Your input settings should now look like the following:
+
+	![](ActionMapping.png)
+
+1.	Close the **Project Settings** menu.
+
+## Adding a Projectile Class
+
+1.	In the File menu, select **New C++ Class...** to choose your new parent class.
+
+1.	The **Choose Parent Class** menu will open. Scroll down, select **Actor** as the parent class and click **Next**.
+
+	![](AddProjectileClass.png)
+
+1.	Name the new class "FPSProjectile," then click **Create**.
+
+	![](CreateProjectileClass.png)
+
+## Adding a USphere Component
+
+1.	Locate the `FPSProjectile` class declaration in the **Solution Explorer** and open `FPSProjectile.h`.
+
+1.	Add a reference to a `USphereComponent` in the `FPSProjectile` interface.
+
+		// Sphere collision component
+		UPROPERTY(VisibleDefaultsOnly, Category = Projectile)
+		USphereComponent* CollisionComp;
+
+1.	Add the following constructor to `FPSProjectile.h`:
+
+		// Constructor for adding a sphere as a simple collision representation
+		AFPSProjectile(const FObjectInitializer& ObjectInitializer);
+
+1.	`FPSProjectile.h` should now look like the following:
+
+		// Fill out your copyright notice in the Description page of Project Settings.
+		
+		#pragma once
+		
+		#include "GameFramework/Actor.h"
+		#include "FPSProjectile.generated.h"
+		
+		UCLASS()
+		class FPSPROJECT_API AFPSProjectile : public AActor
+		{
+			GENERATED_BODY()
+			
+		public:	
+			// Sets default values for this actor's properties
+			AFPSProjectile();
+		
+			// Constructor for adding a sphere as a simple collision representation
+			AFPSProjectile(const FObjectInitializer& ObjectInitializer);
+		
+			// Called when the game starts or when spawned
+			virtual void BeginPlay() override;
+			
+			// Called every frame
+			virtual void Tick( float DeltaSeconds ) override;
+		
+			// Sphere collision component
+			UPROPERTY(VisibleDefaultsOnly, Category = Projectile)
+			USphereComponent* CollisionComp;
+		};
+
+1.	Locate the `FPSProjectile` class implementation file in the **Solution Explorer** and open `FPSProjectile.cpp`.
+
+1.	Add the following code to `FPSProjectile.cpp`:
+
+		// Constructor
+		AFPSProjectile::AFPSProjectile(const FObjectInitializer& ObjectInitializer) : Super(ObjectInitializer)
+		{
+			// Use a sphere as a simple collision representation
+			CollisionComp = ObjectInitializer.CreateDefaultSubobject<USphereComponent>(this, TEXT("SphereComp"));
+		
+			if (CollisionComp != NULL)
+			{
+				CollisionComp->InitSphereRadius(15.f);
+				RootComponent = CollisionComp;
+			}
+		}
+
+	[REGION:note]
+	You're making `CollisionComponent` a `RootComponent` since the simulation will drive it.
+	[/REGION]
+
+1. `FPSProjectile.cpp` should now look like the following:
+
+		// Fill out your copyright notice in the Description page of Project Settings.
+		
+		#include "FPSProject.h"
+		#include "FPSProjectile.h"
+		
+		// Sets default values
+		AFPSProjectile::AFPSProjectile()
+		{
+		 	// Set this actor to call Tick() every frame.  You can turn this off to improve performance if you don't need it.
+			PrimaryActorTick.bCanEverTick = true;
+		
+		}
+		
+		// Constructor
+		AFPSProjectile::AFPSProjectile(const FObjectInitializer& ObjectInitializer) : Super(ObjectInitializer)
+		{
+			// Use a sphere as a simple collision representation
+			CollisionComp = ObjectInitializer.CreateDefaultSubobject<USphereComponent>(this, TEXT("SphereComp"));
+		
+			if (CollisionComp != NULL)
+			{
+				CollisionComp->InitSphereRadius(15.f);
+				RootComponent = CollisionComp;
+			}
+		}
+		
+		// Called when the game starts or when spawned
+		void AFPSProjectile::BeginPlay()
+		{
+			Super::BeginPlay();
+			
+		}
+		
+		// Called every frame
+		void AFPSProjectile::Tick( float DeltaTime )
+		{
+			Super::Tick( DeltaTime );
+		
+		}
+
+## Adding a Projectile Movement Component
+
+1.	Locate the `FPSProjectile` class interface file in the **Solution Explorer** and open `FPSProjectile.h`.
+
+1.	Add the following public reference to `FPSProjectile.h`:
+
+		// Projectile Movement Component
+		UPROPERTY(VisibleAnywhere, BlueprintReadOnly, Category = Movement)
+		UProjectileMovementComponent* ProjectileMovement;
+
+1.	`FPSProjectile.h` should now look like the following:
+
+		// Fill out your copyright notice in the Description page of Project Settings.
+		
+		#pragma once
+		
+		#include "GameFramework/Actor.h"
+		#include "FPSProjectile.generated.h"
+		
+		UCLASS()
+		class FPSPROJECT_API AFPSProjectile : public AActor
+		{
+			GENERATED_BODY()
+			
+		public:	
+			// Sets default values for this actor's properties
+			AFPSProjectile();
+		
+			// Constructor for adding a sphere as a simple collision representation
+			AFPSProjectile(const FObjectInitializer& ObjectInitializer);
+		
+			// Called when the game starts or when spawned
+			virtual void BeginPlay() override;
+			
+			// Called every frame
+			virtual void Tick( float DeltaSeconds ) override;
+		
+			// Sphere collision component
+			UPROPERTY(VisibleDefaultsOnly, Category = Projectile)
+			USphereComponent* CollisionComp;
+		
+			// Projectile Movement Component
+			UPROPERTY(VisibleAnywhere, BlueprintReadOnly, Category = Movement)
+			UProjectileMovementComponent* ProjectileMovement;
+		};
+
+1.	Locate the `FPSProjectile` class implementation file in the **Solution Explorer** and open `FPSProjectile.cpp`.
+
+1.	Add the following lines of code to the `FPSProjectile` constructor in `FPSProjectile.cpp`:
+
+		// Use a ProjectileMovementComponent to govern this projectile's movement
+		ProjectileMovement = ObjectInitializer.CreateDefaultSubobject<UProjectileMovementComponent>(this, TEXT("ProjectileComp"));
+
+		if (ProjectileMovement != NULL)
+		{
+			ProjectileMovement->UpdatedComponent = CollisionComp;
+			ProjectileMovement->InitialSpeed = 3000.f;
+			ProjectileMovement->MaxSpeed = 3000.f;
+			ProjectileMovement->bRotationFollowsVelocity = true;
+			ProjectileMovement->bShouldBounce = true;
+			ProjectileMovement->Bounciness = 0.3f;
+		}
+
+1.	`FPSProjectile.cpp` should now look like the following:
+
+		// Fill out your copyright notice in the Description page of Project Settings.
+		
+		#include "FPSProject.h"
+		#include "FPSProjectile.h"
+		
+		// Sets default values
+		AFPSProjectile::AFPSProjectile()
+		{
+		 	// Set this actor to call Tick() every frame.  You can turn this off to improve performance if you don't need it.
+			PrimaryActorTick.bCanEverTick = true;
+		
+		}
+		
+		// Constructor
+		AFPSProjectile::AFPSProjectile(const FObjectInitializer& ObjectInitializer) : Super(ObjectInitializer)
+		{
+			// Use a sphere as a simple collision representation
+			CollisionComp = ObjectInitializer.CreateDefaultSubobject<USphereComponent>(this, TEXT("SphereComp"));
+		
+			if (CollisionComp != NULL)
+			{
+				CollisionComp->InitSphereRadius(15.f);
+				RootComponent = CollisionComp;
+		
+				// Use a ProjectileMovementComponent to govern this projectile's movement
+				ProjectileMovement = ObjectInitializer.CreateDefaultSubobject<UProjectileMovementComponent>(this, TEXT("ProjectileComp"));
+		
+				if (ProjectileMovement != NULL)
+				{
+					ProjectileMovement->UpdatedComponent = CollisionComp;
+					ProjectileMovement->InitialSpeed = 3000.f;
+					ProjectileMovement->MaxSpeed = 3000.f;
+					ProjectileMovement->bRotationFollowsVelocity = true;
+					ProjectileMovement->bShouldBounce = true;
+					ProjectileMovement->Bounciness = 0.3f;
+				}
+			}
+		}
+		
+		// Called when the game starts or when spawned
+		void AFPSProjectile::BeginPlay()
+		{
+			Super::BeginPlay();
+			
+		}
+		
+		// Called every frame
+		void AFPSProjectile::Tick( float DeltaTime )
+		{
+			Super::Tick( DeltaTime );
+		
+		}
+
+## Setting the Projectile's Initial Velocity
+
+1.	Locate the `FPSProjectile` class interface file in the **Solution Explorer** and open `FPSProjectile.h`.
+
+1.	Add the following function declaration in `FPSProjectile.h`:
+
+		// Initiallizes velocity of the projectile in the shoot direction
+		void InitVelocity(const FVector& ShootDirection);
+
+	[REGION:note]
+	The aforementioned code is the declaration for a function that will launch the projectile.
+	[/REGION]
+
+1.	`FPSProjectile.h` should now look like the following:
+
+		// Fill out your copyright notice in the Description page of Project Settings.
+		
+		#pragma once
+		
+		#include "GameFramework/Actor.h"
+		#include "FPSProjectile.generated.h"
+		
+		UCLASS()
+		class FPSPROJECT_API AFPSProjectile : public AActor
+		{
+			GENERATED_BODY()
+			
+		public:	
+			// Sets default values for this actor's properties
+			AFPSProjectile();
+		
+			// Constructor for adding a sphere as a simple collision representation
+			AFPSProjectile(const FObjectInitializer& ObjectInitializer);
+		
+			// Called when the game starts or when spawned
+			virtual void BeginPlay() override;
+			
+			// Called every frame
+			virtual void Tick( float DeltaSeconds ) override;
+		
+			// Initiallizes velocity of the projectile in the shoot direction
+			void InitVelocity(const FVector& ShootDirection);
+		
+			// Sphere collision component
+			UPROPERTY(VisibleDefaultsOnly, Category = Projectile)
+			USphereComponent* CollisionComp;
+		
+			// Projectile Movement Component
+			UPROPERTY(VisibleAnywhere, BlueprintReadOnly, Category = Movement)
+			UProjectileMovementComponent* ProjectileMovement;
+		};
+
+1.	Locate the `FPSProjectile` class implementation file in the **Solution Explorer** and open `FPSProjectile.cpp`.
+
+1.	Add the following lines of code to the `FPSProjectile` constructor in `FPSProjectile.cpp`:
+
+		// Initiallizes velocity of the projectile in the shoot direction
+		void AFPSProjectile::InitVelocity(const FVector& ShootDirection)
+		{
+			if (ProjectileMovement != NULL)
+			{
+				// set the projectile's velocity to the desired direction
+				ProjectileMovement->Velocity = ShootDirection * ProjectileMovement->InitialSpeed;
+			}
+		}
+
+	[REGION:note]
+	Note that you only needed to supply a launch direction because the projectile's speed is defined by the `ProjectileMovementComponent`.
+	[/REGION]
+
+1.	`FPSProjectile.cpp` should now look like the following:
+
+		// Fill out your copyright notice in the Description page of Project Settings.
+		
+		#include "FPSProject.h"
+		#include "FPSProjectile.h"
+		
+		// Sets default values
+		AFPSProjectile::AFPSProjectile()
+		{
+		 	// Set this actor to call Tick() every frame.  You can turn this off to improve performance if you don't need it.
+			PrimaryActorTick.bCanEverTick = true;
+		
+		}
+		
+		// Constructor
+		AFPSProjectile::AFPSProjectile(const FObjectInitializer& ObjectInitializer) : Super(ObjectInitializer)
+		{
+			// Use a sphere as a simple collision representation
+			CollisionComp = ObjectInitializer.CreateDefaultSubobject<USphereComponent>(this, TEXT("SphereComp"));
+		
+			if (CollisionComp != NULL)
+			{
+				CollisionComp->InitSphereRadius(15.f);
+				RootComponent = CollisionComp;
+		
+				// Use a ProjectileMovementComponent to govern this projectile's movement
+				ProjectileMovement = ObjectInitializer.CreateDefaultSubobject<UProjectileMovementComponent>(this, TEXT("ProjectileComp"));
+		
+				if (ProjectileMovement != NULL)
+				{
+					ProjectileMovement->UpdatedComponent = CollisionComp;
+					ProjectileMovement->InitialSpeed = 3000.f;
+					ProjectileMovement->MaxSpeed = 3000.f;
+					ProjectileMovement->bRotationFollowsVelocity = true;
+					ProjectileMovement->bShouldBounce = true;
+					ProjectileMovement->Bounciness = 0.3f;
+				}
+			}
+		}
+		
+		// Called when the game starts or when spawned
+		void AFPSProjectile::BeginPlay()
+		{
+			Super::BeginPlay();
+			
+		}
+		
+		// Called every frame
+		void AFPSProjectile::Tick( float DeltaTime )
+		{
+			Super::Tick( DeltaTime );
+		
+		}
+		
+		// Initiallizes velocity of the projectile in the shoot direction
+		void AFPSProjectile::InitVelocity(const FVector& ShootDirection)
+		{
+			if (ProjectileMovement != NULL)
+			{
+				// set the projectile's velocity to the desired direction
+				ProjectileMovement->Velocity = ShootDirection * ProjectileMovement->InitialSpeed;
+			}
+		}
+
+## Binding the Fire Input Action
+
+1.	Locate the `FPSCharacter` class interface file in the **Solution Explorer** and open `FPSCharacter.h`.
+
+1.	Add the following function declaration in `FPSCharacter.h`:
+
+		// handles firing
+		UFUNCTION()
+		void OnFire();
+
+1.	`FPSCharacter.h` should now look like the following:
+
+		// Fill out your copyright notice in the Description page of Project Settings.
+		
+		#pragma once
+		
+		#include "GameFramework/Character.h"
+		#include "FPSCharacter.generated.h"
+		
+		UCLASS()
+		class FPSPROJECT_API AFPSCharacter : public ACharacter
+		{
+			GENERATED_BODY()
+		
+		public:
+			// Sets default values for this character's properties
+			AFPSCharacter();
+		
+			// Constructor for the AFPSCharacter
+			AFPSCharacter(const FObjectInitializer& ObjectInitializer);
+		
+			// Called when the game starts or when spawned
+			virtual void BeginPlay() override;
+			
+			// Called every frame
+			virtual void Tick( float DeltaSeconds ) override;
+		
+			// Called to bind functionality to input
+			virtual void SetupPlayerInputComponent(class UInputComponent* InputComponent) override;
+		
+			// handles moving forward and backward
+			UFUNCTION()
+			void MoveForward(float Value);
+		
+			// handles moving right and left
+			UFUNCTION()
+			void MoveRight(float Value);
+		
+			// sets jump flag when key is pressed
+			UFUNCTION()
+			void OnStartJump();
+			
+			// clears jump flag when key is released
+			UFUNCTION()
+			void OnStopJump();
+		
+			// handles firing
+			UFUNCTION()
+			void OnFire();
+		
+			// FPS camera
+			UPROPERTY(VisibleAnywhere, BlueprintReadOnly, Category = Camera)
+			UCameraComponent* FirstPersonCameraComponent;
+		
+			// Pawn Mesh: First person view (arms that are only seen by self)
+			UPROPERTY(VisibleDefaultsOnly, Category = Mesh)
+			USkeletalMeshComponent* FirstPersonMesh;
+		};
+
+1.	Locate the `FPSCharacter` class implementation file in the **Solution Explorer** and open `FPSCharacter.cpp`.
+
+1.	In `FPSCharacter.cpp`, add the following to `SetupPlayerInputComponent` to bind the `OnFire` function:
+
+		// set-up gun firing action binding
+		InputComponent->BindAction("Fire", IE_Pressed, this, &AFPSCharacter::OnFire);
+
+1.	`FPSCharacter.cpp` should now look like the following:
+
+		// Fill out your copyright notice in the Description page of Project Settings.
+		
+		#include "FPSProject.h"
+		#include "FPSCharacter.h"
+		
+		// Sets default values
+		AFPSCharacter::AFPSCharacter()
+		{
+		 	// Set this character to call Tick() every frame.  You can turn this off to improve performance if you don't need it.
+			PrimaryActorTick.bCanEverTick = true;
+		
+		}
+		
+		// Constructor that creates the camera component
+		AFPSCharacter::AFPSCharacter(const FObjectInitializer& ObjectInitializer) : Super(ObjectInitializer)
+		{
+			// Create a camera Component
+			FirstPersonCameraComponent = ObjectInitializer.CreateDefaultSubobject<UCameraComponent>(this, TEXT("FirstPersonCamera"));
+		
+			if (FirstPersonCameraComponent != NULL)
+			{
+				FirstPersonCameraComponent->AttachParent = GetCapsuleComponent();
+				// Position the camera a bit above the eyes
+				FirstPersonCameraComponent->RelativeLocation = FVector(0.f, 0.f, 50.f + BaseEyeHeight);
+				// Allow the pawn to control rotation
+				FirstPersonCameraComponent->bUsePawnControlRotation = true;
+			}
+		
+			// Create a mesh component that will be used when being viewed from a first person view when controlling this pawn
+			FirstPersonMesh = ObjectInitializer.CreateDefaultSubobject<USkeletalMeshComponent>(this, TEXT("FirstPersonMesh"));
+		
+			if (FirstPersonMesh != NULL)
+			{
+				FirstPersonMesh->SetOnlyOwnerSee(true);  // only the owning player will see this mesh
+				FirstPersonMesh->AttachParent = FirstPersonCameraComponent;
+				FirstPersonMesh->bCastDynamicShadow = false;
+				FirstPersonMesh->CastShadow = false;
+			}
+		
+			// everyone but the owner can see the regular body mesh
+			GetMesh()->SetOwnerNoSee(true);
+		}
+		
+		// Called when the game starts or when spawned
+		void AFPSCharacter::BeginPlay()
+		{
+			Super::BeginPlay();
+			
+			if (GEngine != nullptr)
+			{
+				GEngine->AddOnScreenDebugMessage(-1, 5.f, FColor::White, TEXT("We are using FPSCharacter."));
+			}
+		}
+		
+		// Called every frame
+		void AFPSCharacter::Tick( float DeltaTime )
+		{
+			Super::Tick( DeltaTime );
+		
+		}
+		
+		// Called to bind functionality to input
+		void AFPSCharacter::SetupPlayerInputComponent(class UInputComponent* InputComponent)
+		{
+			Super::SetupPlayerInputComponent(InputComponent);
+			
+			// set-up gameplay key bindings
+			InputComponent->BindAxis("MoveForward", this, &AFPSCharacter::MoveForward);
+			InputComponent->BindAxis("MoveRight", this, &AFPSCharacter::MoveRight);
+			// set-up gameplay mouse bindings
+			InputComponent->BindAxis("Turn", this, &AFPSCharacter::AddControllerYawInput);
+			InputComponent->BindAxis("LookUp", this, &AFPSCharacter::AddControllerPitchInput);
+			// set-up jump action bindings
+			InputComponent->BindAction("Jump", IE_Pressed, this, &AFPSCharacter::OnStartJump);
+			InputComponent->BindAction("Jump", IE_Released, this, &AFPSCharacter::OnStopJump);
+			// set-up gun firing action binding
+			InputComponent->BindAction("Fire", IE_Pressed, this, &AFPSCharacter::OnFire);
+		}
+		
+		// Forward movement function definition
+		void AFPSCharacter::MoveForward(float Value)
+		{
+			if (Controller != nullptr && Value != 0.f)
+			{
+				//find out which way is forward
+				FRotator Rotation = Controller->GetControlRotation();
+				//limit pitch when walking for falling
+				if (GetCharacterMovement()->IsMovingOnGround() || GetCharacterMovement()->IsFalling())
+				{
+					Rotation.Pitch = 0.f;
+				}
+				//add movement in that direction
+				const FVector Direction = FRotationMatrix(Rotation).GetScaledAxis(EAxis::X);
+				AddMovementInput(Direction, Value);
+			}
+		}
+		
+		//Right movement function definition
+		void AFPSCharacter::MoveRight(float Value)
+		{
+			if (Controller != nullptr && Value != 0.f)
+			{
+				//find our which way is right
+				const FRotator Rotation = Controller->GetControlRotation();
+				const FVector Direction = FRotationMatrix(Rotation).GetScaledAxis(EAxis::Y);
+				//add movement in that direction
+				AddMovementInput(Direction, Value);
+			}
+		}
+		
+		//Function setting the jump flag
+		void AFPSCharacter::OnStartJump()
+		{
+			bPressedJump = true;
+		}
+		
+		//Function clearning the jump flag
+		void AFPSCharacter::OnStopJump()
+		{
+			bPressedJump = false;
+		}
+
+## Defining the Projectile's Spawn Location
+
+1.	When spawning the `FPSProjectile` actor, there are two points to consider when implementing the `OnFire` function, namely:
+
+*	Where to spawn the projectile
+*	The projectile class (this is so the `FPSCharacter` and its derived Blueprint know what projectile to spawn)
+
+[REGION:note]
+You'll use a camera-space offset vector to determine a spawn location. You're going to make this parameter editable so that you can set and tweak it in your BP_FPSCharacter Blueprint. You'll then be able to calculate an initial location for the projectile based on this data.
+[/REGION]
+
+1.	Locate the `FPSCharacter` class interface file in the **Solution Explorer** and open `FPSCharacter.h`.
+
+1.	Add the following code to `FPSCharacter.h`:
+
+		// Gun muzzle's offset from the camera location
+		UPROPERTY(EditAnywhere, BlueprintReadWrite, Category = Gameplay)
+		FVector MuzzleOffset;
+
+	[REGION:note]
+	The EditAnywhere specifier allow you to change the value of the muzzle offset within the Defaults mode of the Blueprint Editor or within the Details tab for any instance of the character. The BlueprintReadWrite specifier allows you to get and set the value of the muzzle offset within a Blueprint.
+	[/REGION]
+
+1.	Add the following code to `FPSCharacter.h`:
+
+		// Projectile class to spawn
+		UPROPERTY(EditDefaultsOnly, Category = Projectile)
+		TSubclassOf<class AFPSProjectile> ProjectileClass;
+
+	[REGION:note]
+	The EditDefaultsOnly specifier means that you will only be able to set the projectile class as a default on the Blueprint, not on each instance of the Blueprint.
+	[/REGION]
+
+1.	`FPSCharacter.h` should look like the following:
+
+		// Fill out your copyright notice in the Description page of Project Settings.
+		
+		#pragma once
+		
+		#include "GameFramework/Character.h"
+		#include "FPSCharacter.generated.h"
+		
+		UCLASS()
+		class FPSPROJECT_API AFPSCharacter : public ACharacter
+		{
+			GENERATED_BODY()
+		
+		public:
+			// Sets default values for this character's properties
+			AFPSCharacter();
+		
+			// Constructor for the AFPSCharacter
+			AFPSCharacter(const FObjectInitializer& ObjectInitializer);
+		
+			// Called when the game starts or when spawned
+			virtual void BeginPlay() override;
+			
+			// Called every frame
+			virtual void Tick( float DeltaSeconds ) override;
+		
+			// Called to bind functionality to input
+			virtual void SetupPlayerInputComponent(class UInputComponent* InputComponent) override;
+		
+			// handles moving forward and backward
+			UFUNCTION()
+			void MoveForward(float Value);
+		
+			// handles moving right and left
+			UFUNCTION()
+			void MoveRight(float Value);
+		
+			// sets jump flag when key is pressed
+			UFUNCTION()
+			void OnStartJump();
+			
+			// clears jump flag when key is released
+			UFUNCTION()
+			void OnStopJump();
+		
+			// handles firing
+			UFUNCTION()
+			void OnFire();
+		
+			// FPS camera
+			UPROPERTY(VisibleAnywhere, BlueprintReadOnly, Category = Camera)
+			UCameraComponent* FirstPersonCameraComponent;
+		
+			// Pawn Mesh: First person view (arms that are only seen by self)
+			UPROPERTY(VisibleDefaultsOnly, Category = Mesh)
+			USkeletalMeshComponent* FirstPersonMesh;
+		
+			// Gun muzzle's offset from the camera location
+			UPROPERTY(EditAnywhere, BlueprintReadWrite, Category = Gameplay)
+			FVector MuzzleOffset;
+		
+			// Projectile class to spawn
+			UPROPERTY(EditDefaultsOnly, Category = Projectile)
+			TSubclassOf<class AFPSProjectile> ProjectileClass;
+		};
+
+## Compiling and Checking Your Code
+
+It's now time to compile and check your newly implemented projectile code.
+
+1.	Save all of your interface and implementation files in Visual Studio.
+
+1.	Locate **FPSProject** in the **Solution Explorer**.
+
+1.	Right-click on **FPSProject** and select **Build** to compile your project.
+
+	![](BuildFPSProject.png)
+
+	[REGION:note]
+	The purpose of this step is to catch any build errors before moving onto the next step.
+	[/REGION]
 
 %Steps%
\ No newline at end of file
