#UE4doc #docs
Wrote first draft of 'Implementing Shooting' section.
Availability: Docs

diff --git a/Engine/Documentation/Source/Programming/Tutorials/FirstPersonShooter/3/2/FirstPersonShooter_3_2.INT.udn b/Engine/Documentation/Source/Programming/Tutorials/FirstPersonShooter/3/2/FirstPersonShooter_3_2.INT.udn
--- a/Engine/Documentation/Source/Programming/Tutorials/FirstPersonShooter/3/2/FirstPersonShooter_3_2.INT.udn
+++ b/Engine/Documentation/Source/Programming/Tutorials/FirstPersonShooter/3/2/FirstPersonShooter_3_2.INT.udn
@@ -23,6 +23,310 @@ Version: 4.9
 
 %Steps%
 
-<!--TODO: Add tutorial steps here-->
+During this step, you're going to implement the `OnFire` function so that your character can start firing projectiles.
+
+1.	Locate the `FPSCharacter` class implementation file in the **Solution Explorer** and open `FPSCharacter.cpp`.
+
+1.	Add the following `#include` at the top of `FPSCharacter.cpp`.
+
+		#include "FPSProjectile.h"
+
+1.	Add the following `OnFire` function definition to `FPSCharacter.cpp`:
+
+		void AFPSCharacter::OnFire()
+		{
+			// attempt to fire a projectile
+			if (ProjectileClass != NULL)
+			{
+				// get the camera transform
+				FVector CameraLocation;
+				FRotator CameraRotation;
+				GetActorEyesViewPoint(CameraLocation, CameraRotation);
+				
+				// MuzzleOffset is in camera space which means that you need to transform it to world space before offsetting from the camera to find the final muzzle position
+				FVector const MuzzleLocation = CameraLocation + FTransform(CameraRotation).TransformVector(MuzzleOffset);
+				FRotator MuzzleRotation = CameraRotation;
+				MuzzleRotation.Pitch += 10.f; // slightly skew the aim upwards
+				UWorld* const World = GetWorld();
+				if (World != NULL)
+				{
+					FActorSpawnParameters SpawnParams;
+					SpawnParams.Owner = this;
+					SpawnParams.Instigator = Instigator;
+					// spawn the projectile at the muzzle
+					AFPSProjectile* const Projectile = World->SpawnActor<AFPSProjectile>(ProjectileClass, MuzzleLocation, MuzzleRotation, SpawnParams);
+					if (Projectile != NULL)
+					{
+						// find the launch direction
+						FVector const LaunchDirection = MuzzleRotation.Vector();
+						Projectile->InitVelocity(LaunchDirection);
+					}
+				}
+			}
+		}
+
+1.	`FPSCharacter.cpp` should now look like the following:
+
+		// Fill out your copyright notice in the Description page of Project Settings.
+		
+		#include "FPSProject.h"
+		#include "FPSCharacter.h"
+		#include "FPSProjectile.h"
+		
+		// Sets default values
+		AFPSCharacter::AFPSCharacter()
+		{
+		 	// Set this character to call Tick() every frame.  You can turn this off to improve performance if you don't need it.
+			PrimaryActorTick.bCanEverTick = true;
+		
+		}
+		
+		// Constructor that creates the camera component
+		AFPSCharacter::AFPSCharacter(const FObjectInitializer& ObjectInitializer) : Super(ObjectInitializer)
+		{
+			// Create a camera Component
+			FirstPersonCameraComponent = ObjectInitializer.CreateDefaultSubobject<UCameraComponent>(this, TEXT("FirstPersonCamera"));
+		
+			if (FirstPersonCameraComponent != NULL)
+			{
+				FirstPersonCameraComponent->AttachParent = GetCapsuleComponent();
+				// Position the camera a bit above the eyes
+				FirstPersonCameraComponent->RelativeLocation = FVector(0.f, 0.f, 50.f + BaseEyeHeight);
+				// Allow the pawn to control rotation
+				FirstPersonCameraComponent->bUsePawnControlRotation = true;
+			}
+		
+			// Create a mesh component that will be used when being viewed from a first person view when controlling this pawn
+			FirstPersonMesh = ObjectInitializer.CreateDefaultSubobject<USkeletalMeshComponent>(this, TEXT("FirstPersonMesh"));
+		
+			if (FirstPersonMesh != NULL)
+			{
+				FirstPersonMesh->SetOnlyOwnerSee(true);  // only the owning player will see this mesh
+				FirstPersonMesh->AttachParent = FirstPersonCameraComponent;
+				FirstPersonMesh->bCastDynamicShadow = false;
+				FirstPersonMesh->CastShadow = false;
+			}
+		
+			// everyone but the owner can see the regular body mesh
+			GetMesh()->SetOwnerNoSee(true);
+		}
+		
+		// Called when the game starts or when spawned
+		void AFPSCharacter::BeginPlay()
+		{
+			Super::BeginPlay();
+			
+			if (GEngine != nullptr)
+			{
+				GEngine->AddOnScreenDebugMessage(-1, 5.f, FColor::White, TEXT("We are using FPSCharacter."));
+			}
+		}
+		
+		// Called every frame
+		void AFPSCharacter::Tick( float DeltaTime )
+		{
+			Super::Tick( DeltaTime );
+		
+		}
+		
+		// Called to bind functionality to input
+		void AFPSCharacter::SetupPlayerInputComponent(class UInputComponent* InputComponent)
+		{
+			Super::SetupPlayerInputComponent(InputComponent);
+			
+			// set-up gameplay key bindings
+			InputComponent->BindAxis("MoveForward", this, &AFPSCharacter::MoveForward);
+			InputComponent->BindAxis("MoveRight", this, &AFPSCharacter::MoveRight);
+			// set-up gameplay mouse bindings
+			InputComponent->BindAxis("Turn", this, &AFPSCharacter::AddControllerYawInput);
+			InputComponent->BindAxis("LookUp", this, &AFPSCharacter::AddControllerPitchInput);
+			// set-up jump action bindings
+			InputComponent->BindAction("Jump", IE_Pressed, this, &AFPSCharacter::OnStartJump);
+			InputComponent->BindAction("Jump", IE_Released, this, &AFPSCharacter::OnStopJump);
+			// set-up gun firing action binding
+			InputComponent->BindAction("Fire", IE_Pressed, this, &AFPSCharacter::OnFire);
+		}
+		
+		// Forward movement function definition
+		void AFPSCharacter::MoveForward(float Value)
+		{
+			if (Controller != nullptr && Value != 0.f)
+			{
+				//find out which way is forward
+				FRotator Rotation = Controller->GetControlRotation();
+				//limit pitch when walking for falling
+				if (GetCharacterMovement()->IsMovingOnGround() || GetCharacterMovement()->IsFalling())
+				{
+					Rotation.Pitch = 0.f;
+				}
+				//add movement in that direction
+				const FVector Direction = FRotationMatrix(Rotation).GetScaledAxis(EAxis::X);
+				AddMovementInput(Direction, Value);
+			}
+		}
+		
+		//Right movement function definition
+		void AFPSCharacter::MoveRight(float Value)
+		{
+			if (Controller != nullptr && Value != 0.f)
+			{
+				//find our which way is right
+				const FRotator Rotation = Controller->GetControlRotation();
+				const FVector Direction = FRotationMatrix(Rotation).GetScaledAxis(EAxis::Y);
+				//add movement in that direction
+				AddMovementInput(Direction, Value);
+			}
+		}
+		
+		//Function setting the jump flag
+		void AFPSCharacter::OnStartJump()
+		{
+			bPressedJump = true;
+		}
+		
+		//Function clearning the jump flag
+		void AFPSCharacter::OnStopJump()
+		{
+			bPressedJump = false;
+		}
+		
+		void AFPSCharacter::OnFire()
+		{
+			// attempt to fire a projectile
+			if (ProjectileClass != NULL)
+			{
+				// get the camera transform
+				FVector CameraLocation;
+				FRotator CameraRotation;
+				GetActorEyesViewPoint(CameraLocation, CameraRotation);
+				
+				// MuzzleOffset is in camera space which means that you need to transform it to world space before offsetting from the camera to find the final muzzle position
+				FVector const MuzzleLocation = CameraLocation + FTransform(CameraRotation).TransformVector(MuzzleOffset);
+				FRotator MuzzleRotation = CameraRotation;
+				MuzzleRotation.Pitch += 10.f; // slightly skew the aim upwards
+				UWorld* const World = GetWorld();
+				if (World != NULL)
+				{
+					FActorSpawnParameters SpawnParams;
+					SpawnParams.Owner = this;
+					SpawnParams.Instigator = Instigator;
+					// spawn the projectile at the muzzle
+					AFPSProjectile* const Projectile = World->SpawnActor<AFPSProjectile>(ProjectileClass, MuzzleLocation, MuzzleRotation, SpawnParams);
+					if (Projectile != NULL)
+					{
+						// find the launch direction
+						FVector const LaunchDirection = MuzzleRotation.Vector();
+						Projectile->InitVelocity(LaunchDirection);
+					}
+				}
+			}
+		}
+
+1.	Save `FPSCharacter.cpp` in Visual Studio.
+
+1.	Locate **FPSProject** in the **Solution Explorer**.
+
+1.	Right-click on **FPSProject** and select **Build** to compile your project.
+
+	![](BuildFPSProject.png)
+
+## Building the Projectile Blueprint
+
+[REGION:note]
+Before continuing, download and extract our sample mesh from the following link:
+["Projectile Mesh"](Sphere.zip "Projectile Mesh")
+[/REGION]
+
+1.	**Right-click** inside the Content Browser's file box to open the **Import Asset** dialog box
+
+1.	**Click 'Import to /Game...'** to open the **Import** dialog box.
+
+	![](RightClickImport.png)
+
+1.	Locate and select the **Sphere.fbx** mesh file.
+
+1.	**Click Open** to begin importing the mesh to your project.
+
+1.	In the **Content Browser**, the **FBX Import Options** dialog box appears. **Clicking Import** adds your mesh to the Project.
+
+	[REGION:warning]
+	Disregard the following error regarding smoothing groups:
+	![](FBXWarning.png)  
+	This mesh still illustrates the first person mesh set-up and it will work with the animations you'll set-up in a later section.
+	[/REGION]
+
+1.	**Click** the **Save** button to save your imported static mesh.
+
+1.	Navigate back to the **Blueprints** folder in the **Content Browser**.
+
+1.	**Click** the **Add New** button and select **Blueprint Class**.
+
+	![](AddNewBPClass.png)
+
+1.	Expand the **All Classes** dropdown menu and type "FPSProjectile" into the search box.
+
+	![](AddNewBPClass.png)
+
+1.	**Click** on the **FPSCharacter** and **click** the **Select** button.
+
+1.	Name the new Blueprint "BP_FPSProjectile".
+
+	![](NameBP_FPSProjectile.png)
+
+1.	 **Double-click** the **BP_FPSCharacter** icon to open it in the **Blueprint Editor**.
+
+1.	**Click** on **CollisionComp** in the **Components** tab.
+
+1.	**Click** on the **Add Component** dropdown list and select **Static Mesh**.
+
+	![](AddStaticMeshComponent.png)
+
+1.	Name your new component "ProjectileMesh".
+
+	![](NameProjectileMeshComponent.png)
+
+1.	Scroll down to the **Mesh** section of the **Details** tab and **click** on the dropdown menu that says "None".
+
+1.	Select the **Sphere** static mesh.
+
+	![](SelectSphereMesh.png)
+
+	[REGION:note]
+	Note that if you're making a multiplayer game, you'll need to also uncheck "Initial Velocity in Local Space" in the "MovementComp" Component in order for this projectile to replicate correctly over a server.
+	[/REGION]
+
+1.	Set the X, Y, and Z scale values to 0.09.
+
+	[REGION:note]
+	Clicking on the lock icon locks all three axes so that they preserve their relative scale.
+	[/REGION]
+
+1.	Set the **ProjectileMesh's Collision Presets** value to **NoCollision** (since you're using the `SphereComponent` for collision and not this **Static Mesh**.	
+
+	![](SetCollisionPresets.png)
+	
+1.	**Compile** and **Save** your Blueprint before closing the **Blueprint Editor**.
+
+1.	**Double-click** on **BP_FPSCharacter** to open the character blueprint for editing.
+
+1.	Open **Defaults Mode**.
+
+1.	Locate the **Projectile Class** property and set it to **BP_FPSProjectile**.
+
+1.	Set the **MuzzleOffset** property to {100, 0, 00} in order to spawn the projectile slightly in front of the camera.
+
+	![](CharacterProjectileClass.png)
+
+1.	**Compile** and **Save** your Blueprint before closing the **Blueprint Editor**.
+
+## Shooting Projectiles In-Game
+
+1.	**Click** the **Play In** button in the **Level Editor Toolbar** to shoot projectiles in-game.
+
+1.	**Left-click** your mouse button to fire projectiles into the world.
+
+	![](FiringProjectiles.png)
+
+1.	Click the **Stop** button in the Level Editor to exit Play in Editor (PIE) mode.
 
 %Steps%
\ No newline at end of file
