#UE4doc #docs
Wrote first draft of 'Setting up Collision Lifetime and Collision' section.
Availability: Docs

diff --git a/Engine/Documentation/Source/Programming/Tutorials/FirstPersonShooter/3/3/FirstPersonShooter_3_3.INT.udn b/Engine/Documentation/Source/Programming/Tutorials/FirstPersonShooter/3/3/FirstPersonShooter_3_3.INT.udn
--- a/Engine/Documentation/Source/Programming/Tutorials/FirstPersonShooter/3/3/FirstPersonShooter_3_3.INT.udn
+++ b/Engine/Documentation/Source/Programming/Tutorials/FirstPersonShooter/3/3/FirstPersonShooter_3_3.INT.udn
@@ -23,6 +23,122 @@ Version: 4.9
 
 %Steps%
 
-<!--TODO: Add tutorial steps here-->
+Currently, your projectiles:
+
+*	Live forever (evidenced by the fact that they never disappear in the Scene Outliner)
+*	Don't collide with other objects in the world
+
+During this step, we're going to set-up projectile collision and lifetime.
+
+## Limiting the Projectile's Life Span
+
+1.	Locate the `FPSProjectile` class implementation file in the **Solution Explorer** and open `FPSProjectile.cpp`.
+
+1.	Add the following code to the `FPSProjectile` constructor to set the projectile's lifespan:
+
+		// Die after 3 seconds (by default)
+		InitialLifeSpan = 3.f;
+
+## Editing the Projectile's Collision Settings
+
+UE4 comes with several useful collision chanels; however, the engine also provides several customizable channels that game projects can use.
+
+1. Open **Project Settings** and select **Collision** to customize a collision channel.
+
+	![](SettingCollisionChannels.png)
+
+1.	Select **New Object Channel...** to create a new collision channel. Name your new collision channel "Projectile" and make sure that the **Default Response** is set to **Block** before clicking **Accept**.
+
+1.	Select **New...** under **Preset** and name your new profile "Projectile" as well. Refer to the following image to set your collision presets.
+
+	![](NewProfile.png)
+
+	[REGION:note]
+	The profile shown above specifies that the projectile will be blocked by `Static Actors`, `Pawns`, `Dynamic Actors`, `Actors` simulating `Physics`, `Vehicles`, and `Destructible Actors`.
+	[/REGION]
+
+## Using the New Collision Channel's Settings
+
+1.	Locate the `FPSProjectile` class implementation file in the **Solution Explorer** and open `FPSProjectile.cpp`.
+
+1.	In the `FPSProjectile` constructor, add the following line after the creation of `CollisionComp`:
+
+		CollisionComp->BodyInstance.SetCollisionProfileName("Projectile");
+
+1.	`FPSProjectile.cpp` should now look like the following:
+
+		// Fill out your copyright notice in the Description page of Project Settings.
+		
+		#include "FPSProject.h"
+		#include "FPSProjectile.h"
+		
+		// Sets default values
+		AFPSProjectile::AFPSProjectile()
+		{
+		 	// Set this actor to call Tick() every frame.  You can turn this off to improve performance if you don't need it.
+			PrimaryActorTick.bCanEverTick = true;
+			
+			// Die after 3 seconds (by default)
+			InitialLifeSpan = 3.f;
+		}
+		
+		// Constructor
+		AFPSProjectile::AFPSProjectile(const FObjectInitializer& ObjectInitializer) : Super(ObjectInitializer)
+		{
+			// Use a sphere as a simple collision representation
+			CollisionComp = ObjectInitializer.CreateDefaultSubobject<USphereComponent>(this, TEXT("SphereComp"));
+		
+			if (CollisionComp != NULL)
+			{
+				CollisionComp->BodyInstance.SetCollisionProfileName("Projectile");
+				CollisionComp->InitSphereRadius(15.f);
+				RootComponent = CollisionComp;
+		
+				// Use a ProjectileMovementComponent to govern this projectile's movement
+				ProjectileMovement = ObjectInitializer.CreateDefaultSubobject<UProjectileMovementComponent>(this, TEXT("ProjectileComp"));
+		
+				if (ProjectileMovement != NULL)
+				{
+					ProjectileMovement->UpdatedComponent = CollisionComp;
+					ProjectileMovement->InitialSpeed = 3000.f;
+					ProjectileMovement->MaxSpeed = 3000.f;
+					ProjectileMovement->bRotationFollowsVelocity = true;
+					ProjectileMovement->bShouldBounce = true;
+					ProjectileMovement->Bounciness = 0.3f;
+				}
+			}
+		}
+		
+		// Called when the game starts or when spawned
+		void AFPSProjectile::BeginPlay()
+		{
+			Super::BeginPlay();
+			
+		}
+		
+		// Called every frame
+		void AFPSProjectile::Tick( float DeltaTime )
+		{
+			Super::Tick( DeltaTime );
+		
+		}
+		
+		// Initiallizes velocity of the projectile in the shoot direction
+		void AFPSProjectile::InitVelocity(const FVector& ShootDirection)
+		{
+			if (ProjectileMovement != NULL)
+			{
+				// set the projectile's velocity to the desired direction
+				ProjectileMovement->Velocity = ShootDirection * ProjectileMovement->InitialSpeed;
+			}
+		}
+
+1.	Save `FPSProjectile.cpp` in Visual Studio.
+
+1.	Locate **FPSProject** in the **Solution Explorer**.
+
+1.	Right-click on **FPSProject** and select **Build** to compile your project.
+
+	![](BuildFPSProject.png)
 
 %Steps%
\ No newline at end of file
