#UE4doc #docs
Wrote first draft of 'ProjectilesInteractingWithWorld' section.
Availability: Docs

diff --git a/Engine/Documentation/Source/Programming/Tutorials/FirstPersonShooter/3/4/FirstPersonShooter_3_4.INT.udn b/Engine/Documentation/Source/Programming/Tutorials/FirstPersonShooter/3/4/FirstPersonShooter_3_4.INT.udn
--- a/Engine/Documentation/Source/Programming/Tutorials/FirstPersonShooter/3/4/FirstPersonShooter_3_4.INT.udn
+++ b/Engine/Documentation/Source/Programming/Tutorials/FirstPersonShooter/3/4/FirstPersonShooter_3_4.INT.udn
@@ -23,6 +23,206 @@ Version: 4.9
 
 %Steps%
 
-<!--TODO: Add tutorial steps here-->
+Now that you can detect your projectile's collision interactions, you can determine how to respond to those collisions. During this step, you'll add an `OnHit` function to `FPSProjectile` that will respond to collision events.
+
+## Getting Projectiles to React to Collisions
+
+1.	Locate the `FPSProjectile` class interface file in the **Solution Explorer** and open `FPSProjectile.h`.
+
+1.	Add the following code to the `FPSProjectile` class declaration:
+
+		// Called when projectile hits something
+		UFUNCTION()
+		void OnHit(class AActor* OtherActor, class UPrimitiveComponent* OtherComp, FVector NormalImpulse, const FHitResult& Hit);
+
+1.	`FPSProjectile.h` should now look like the following:
+
+		// Fill out your copyright notice in the Description page of Project Settings.
+		
+		#pragma once
+		
+		#include "GameFramework/Actor.h"
+		#include "FPSProjectile.generated.h"
+		
+		UCLASS()
+		class FPSPROJECT_API AFPSProjectile : public AActor
+		{
+			GENERATED_BODY()
+			
+		public:	
+			// Sets default values for this actor's properties
+			AFPSProjectile();
+		
+			// Constructor for adding a sphere as a simple collision representation
+			AFPSProjectile(const FObjectInitializer& ObjectInitializer);
+		
+			// Called when the game starts or when spawned
+			virtual void BeginPlay() override;
+			
+			// Called every frame
+			virtual void Tick( float DeltaSeconds ) override;
+		
+			// Initiallizes velocity of the projectile in the shoot direction
+			void InitVelocity(const FVector& ShootDirection);
+		
+			// Sphere collision component
+			UPROPERTY(VisibleDefaultsOnly, Category = Projectile)
+			USphereComponent* CollisionComp;
+		
+			// Projectile Movement Component
+			UPROPERTY(VisibleAnywhere, BlueprintReadOnly, Category = Movement)
+			UProjectileMovementComponent* ProjectileMovement;
+		
+			// Called when projectile hits something
+			UFUNCTION()
+			void OnHit(class AActor* OtherActor, class UPrimitiveComponent* OtherComp, FVector NormalImpulse, const FHitResult& Hit);
+		};
+
+1.	Locate the `FPSProjectile` class implementation file in the **Solution Explorer** and open `FPSProjectile.cpp`.
+
+1.	Add the following code to the `OnHit` function definition:
+
+		void AFPSProjectile::OnHit(class AActor* OtherActor, class UPrimitiveComponent* OtherComp, FVector NormalImpulse, const FHitResult& Hit)
+		{
+			if (OtherActor != NULL && OtherActor != this && OtherComp != NULL)
+			{
+				OtherComp->AddImpulseAtLocation(ProjectileMovement->Velocity * 100.f, Hit.ImpactPoint);
+			}
+		}
+
+1.	In the `FPSProjectile` constructor, add the following line after the creation of `CollisionComp`:
+
+		CollisionComp->OnComponentHit.AddDynamic(this, &AFPSProjectile::OnHit);
+
+1.	`FPSProjectile.cpp` should now look like the following:
+
+		// Fill out your copyright notice in the Description page of Project Settings.
+		
+		#include "FPSProject.h"
+		#include "FPSProjectile.h"
+		
+		// Sets default values
+		AFPSProjectile::AFPSProjectile()
+		{
+		 	// Set this actor to call Tick() every frame.  You can turn this off to improve performance if you don't need it.
+			PrimaryActorTick.bCanEverTick = true;
+			
+			// Die after 3 seconds (by default)
+			InitialLifeSpan = 3.f;
+		}
+		
+		// Constructor
+		AFPSProjectile::AFPSProjectile(const FObjectInitializer& ObjectInitializer) : Super(ObjectInitializer)
+		{
+			// Use a sphere as a simple collision representation
+			CollisionComp = ObjectInitializer.CreateDefaultSubobject<USphereComponent>(this, TEXT("SphereComp"));
+		
+			if (CollisionComp != NULL)
+			{
+				CollisionComp->BodyInstance.SetCollisionProfileName("Projectile");
+				CollisionComp->OnComponentHit.AddDynamic(this, &AFPSProjectile::OnHit);
+				CollisionComp->InitSphereRadius(15.f);
+				RootComponent = CollisionComp;
+		
+				// Use a ProjectileMovementComponent to govern this projectile's movement
+				ProjectileMovement = ObjectInitializer.CreateDefaultSubobject<UProjectileMovementComponent>(this, TEXT("ProjectileComp"));
+		
+				if (ProjectileMovement != NULL)
+				{
+					ProjectileMovement->UpdatedComponent = CollisionComp;
+					ProjectileMovement->InitialSpeed = 3000.f;
+					ProjectileMovement->MaxSpeed = 3000.f;
+					ProjectileMovement->bRotationFollowsVelocity = true;
+					ProjectileMovement->bShouldBounce = true;
+					ProjectileMovement->Bounciness = 0.3f;
+				}
+			}
+		}
+		
+		// Called when the game starts or when spawned
+		void AFPSProjectile::BeginPlay()
+		{
+			Super::BeginPlay();
+			
+		}
+		
+		// Called every frame
+		void AFPSProjectile::Tick( float DeltaTime )
+		{
+			Super::Tick( DeltaTime );
+		
+		}
+		
+		// Initiallizes velocity of the projectile in the shoot direction
+		void AFPSProjectile::InitVelocity(const FVector& ShootDirection)
+		{
+			if (ProjectileMovement != NULL)
+			{
+				// set the projectile's velocity to the desired direction
+				ProjectileMovement->Velocity = ShootDirection * ProjectileMovement->InitialSpeed;
+			}
+		}
+		
+		void AFPSProjectile::OnHit(class AActor* OtherActor, class UPrimitiveComponent* OtherComp, FVector NormalImpulse, const FHitResult& Hit)
+		{
+			if (OtherActor != NULL && OtherActor != this && OtherComp != NULL)
+			{
+				OtherComp->AddImpulseAtLocation(ProjectileMovement->Velocity * 100.f, Hit.ImpactPoint);
+			}
+		}
+
+1.	Save `FPSProjectile.h` and `FPSProjectile.cpp` in Visual Studio.
+
+1.	Locate **FPSProject** in the **Solution Explorer**.
+
+1.	Right-click on **FPSProject** and select **Build** to compile your project.
+
+	![](BuildFPSProject.png)
+
+## Testing Projectile Collision
+
+1.	After the build finishes, go back to **Unreal Editor** and open **FPSProject**.
+
+1.	Select the **Floor** `StaticMesh`.
+
+1.	**Copy and paste** the floor mesh.
+
+1.	Making sure that the ratio lock is unlocked, scale the floor mesh copy (named "Floor2") to {0.2, 0.2, 0.3}.
+
+1.	Position the floor mesh copy at {320, 0, 160}.
+
+	[REGION:lightbox]
+		[![](Floor2_CubeTransform.png)(w:750)](Floor2_CubeTransform.png)
+	[/REGION]
+	_Click on the image to zoom in._
+
+1.	Scroll down to the **Physics** section and check the **Simulate Physics** box.
+
+	[REGION:lightbox]
+		[![](Floor2_CubePhysics.png)(w:750)](Floor2_CubePhysics.png)
+	[/REGION]
+	_Click on the image to zoom in._
+
+1.	**Save** your map.
+
+1.	**Double-click** on **BP_FPSProjectile** to open the projectile blueprint for editing.
+
+1.	Open **Defaults Mode**.
+
+1.	**Click** on the **ProjectileMesh** in the **Components** tab.
+
+1.	Locate the **Collision Presets** property under **Collision** and set it to **Projectile**.
+
+	![](SetCollisionPresets.png)
+
+1.	**Compile** and **Save** your Blueprint before closing the **Blueprint Editor**.
+
+1.	**Click** the **Play In** button in the **Level Editor Toolbar**.
+
+1.	**Left-click** your mouse button to fire projectiles and move the cube around your level. Congratulations, your projectiles are complete!
+
+	![](ProjectileComplete.png)
+
+1.	Click the **Stop** button in the Level Editor to exit Play in Editor (PIE) mode.
 
 %Steps%
\ No newline at end of file
