#UE4doc #docs
Completed first draft of `Adding Crosshairs` section.
Availability: Docs

diff --git a/Engine/Documentation/Source/Programming/Tutorials/FirstPersonShooter/3/5/FirstPersonShooter_3_5.INT.udn b/Engine/Documentation/Source/Programming/Tutorials/FirstPersonShooter/3/5/FirstPersonShooter_3_5.INT.udn
--- a/Engine/Documentation/Source/Programming/Tutorials/FirstPersonShooter/3/5/FirstPersonShooter_3_5.INT.udn
+++ b/Engine/Documentation/Source/Programming/Tutorials/FirstPersonShooter/3/5/FirstPersonShooter_3_5.INT.udn
@@ -23,6 +23,229 @@ Version: 4.9
 
 %Steps%
 
-<!--TODO: Add tutorial steps here-->
+[REGION:note]
+Before getting started, download and extract our sample mesh from the following link:
+*	[Sample Crosshair Image](Crosshair_fps_tutorial.zip "Sample Crosshair Image")
+[/REGION]
+
+During this step, you'll add a crosshair HUD element to your game so that you can aim your projectiles.
+
+## Importing a Crosshair Asset
+
+1.	**Right-click** inside the Content Browser's file box to open the **Import Asset** dialog box
+
+1.	**Click 'Import to /Game...'** to open the **Import** dialog box.
+
+	![](RightClickImport.png)
+
+1.	Locate and select the **crosshair.TGA** image file.
+
+1.	**Click Open** to begin importing the image file to your project.
+
+1.	**Click** the **Save** button to save your imported mesh.
+
+## Adding a New HUD Class
+
+1.	In the File menu, select **New C++ Class...** to choose your new parent class.
+
+1.	The **Choose Parent Class** menu will open. Scroll down, select **HUD** as the parent class and click **Next**.
+
+	![](ChooseHUDClass.png)
+
+1.	Name the new class "FPSHUD," then click **Create**.
+
+	![](NameHUDClass.png)
+
+1.	Locate the `FPSHUD` class interface file in the **Solution Explorer** and open `FPSHUD.h`.
+
+1.	Place the following constructor declaration in `FPSHUD.h` under `GENERATED_BODY()`:
+
+		AFPSHUD(const FObjectInitializer& ObjectInitializer);
+
+1.	Locate the `FPSHUD` class implementation file in the **Solution Explorer** and open `FPSHUD.cpp`.
+
+1.	Write the following constructor definition in `FPSHUD.h` below `GENERATED_BODY()`:
+
+		AFPSHUD::AFPSHUD(const FObjectInitializer& ObjectInitializer) : Super(ObjectInitializer)
+		{
+		
+		}
+
+1.	Going back to `FPSHUD.h`, add the following private variable:
+
+		private:
+		
+			// Crosshair asset pointer
+			UTexture2D* CrosshairTex;
+
+1.	Now add the following code inside the constructor that's located in `FPSHUD.cpp`:
+
+		// Set the crosshair texture
+		static ConstructorHelpers::FObjectFinder<UTexture2D> CrosshairTexObj(TEXT("Texture2D'/Game/crosshair.crosshair'"));
+		
+		if (CrosshairTexObj.Object != NULL)
+		{
+			CrosshairTex = CrosshairTexObj.Object;
+		}
+
+	[REGION:note]
+	As a reminder, you can obtain the asset path to your texture in the editor by right-clicking the asset in the Content Browser and selecting Copy Reference.
+	![](CopyReference.png)
+	[/REGION]
+
+1.	Going back to `FPSHUD.h`, add the following function declaration below `GENERATED_BODY()`:
+
+		// Primary draw call for the HUD
+		virtual void DrawHUD() override;
+
+1.	Now implement the `DrawHUD` override in `FPSHUD.cpp` with the following function definition:
+
+		void AFPSHUD::DrawHUD()
+		{
+			Super::DrawHUD();
+			// Draw simple crosshair
+			// Find the center of our canvas
+			const FVector2D Center(Canvas->ClipX * 0.5f, Canvas->ClipY * 0.5f);
+			// Offset by half of the texture's dimensions so that the center of the texture aligns with the center of the Canvas
+			const FVector2D CrossHairDrawPosition((Center.X - (CrosshairTex->GetSurfaceWidth() * 0.5f)), (Center.Y - (CrosshairTex->GetSurfaceHeight() * 0.5f)));
+			// Draw the crosshair
+			FCanvasTileItem TileItem(CrossHairDrawPosition, CrosshairTex->Resource, FLinearColor::White);
+			TileItem.BlendMode = SE_BLEND_Translucent;
+			Canvas->DrawItem(TileItem);
+		}
+
+1.	`FPSHUD.h` should look like the following:
+
+		// Fill out your copyright notice in the Description page of Project Settings.
+		
+		#pragma once
+		
+		#include "GameFramework/HUD.h"
+		#include "FPSHUD.generated.h"
+		
+		/**
+		 * 
+		 */
+		UCLASS()
+		class FPSPROJECT_API AFPSHUD : public AHUD
+		{
+			GENERATED_BODY()
+		
+			AFPSHUD(const FObjectInitializer& ObjectInitializer);
+		
+			// Primary draw call for the HUD
+			virtual void DrawHUD() override;
+		
+		private:
+		
+			// Crosshair asset pointer
+			UTexture2D* CrosshairTex;
+			
+		};
+
+1.	`FPSHUD.cpp` should now look like the following:
+
+		// Fill out your copyright notice in the Description page of Project Settings.
+		
+		#include "FPSProject.h"
+		#include "FPSHUD.h"
+		
+		AFPSHUD::AFPSHUD(const FObjectInitializer& ObjectInitializer) : Super(ObjectInitializer)
+		{
+			// Set the crosshair texture
+			static ConstructorHelpers::FObjectFinder<UTexture2D> CrosshairTexObj(TEXT("Texture2D'/Game/crosshair.crosshair'"));
+			
+			if (CrosshairTexObj.Object != NULL)
+			{
+				CrosshairTex = CrosshairTexObj.Object;
+			}
+		}
+		
+		void AFPSHUD::DrawHUD()
+		{
+			Super::DrawHUD();
+			// Draw simple crosshair
+			// Find the center of our canvas
+			const FVector2D Center(Canvas->ClipX * 0.5f, Canvas->ClipY * 0.5f);
+			// Offset by half of the texture's dimensions so that the center of the texture aligns with the center of the Canvas
+			const FVector2D CrossHairDrawPosition((Center.X - (CrosshairTex->GetSurfaceWidth() * 0.5f)), (Center.Y - (CrosshairTex->GetSurfaceHeight() * 0.5f)));
+			// Draw the crosshair
+			FCanvasTileItem TileItem(CrossHairDrawPosition, CrosshairTex->Resource, FLinearColor::White);
+			TileItem.BlendMode = SE_BLEND_Translucent;
+			Canvas->DrawItem(TileItem);
+		}
+
+1.	Save `FPSHUD.h` and `FPSHUD.cpp` in Visual Studio.
+
+1.	Locate **FPSProject** in the **Solution Explorer**.
+
+1.	Right-click on **FPSProject** and select **Build** to compile your project.
+
+	![](BuildFPSProject.png)
+
+## Updating GameMode to Use the New HUD Class
+
+1.	Add the following include statement at the top of `FPSGameMode.cpp`:
+
+		#include "FPSHUD.h"
+
+1.	In the `FPSGameMode` constructor, located inside `FPSGameMode.cpp`, add the following line of code:
+
+		HUDClass = AFPSHUD::StaticClass();
+
+1.	`FPSGameMode.cpp` should now look like the following:
+
+		// Fill out your copyright notice in the Description page of Project Settings.
+		
+		#include "FPSProject.h"
+		#include "FPSGameMode.h"
+		#include "FPSHUD.h"
+		
+		AFPSGameMode::AFPSGameMode(const FObjectInitializer& ObjectInitializer) : Super(ObjectInitializer)
+		{
+			// set the default pawn class to the Blueprinted character
+			static ConstructorHelpers::FClassFinder<APawn> PlayerPawnObject(TEXT("Pawn'/Game/Blueprints/BP_FPSCharacter.BP_FPSCharacter_C'"));
+			if (PlayerPawnObject.Class != NULL)
+			{
+				DefaultPawnClass = PlayerPawnObject.Class;
+			}
+			HUDClass = AFPSHUD::StaticClass();
+		}//AFPSGameMode Constructor
+		
+		void AFPSGameMode::StartPlay()
+		{
+			Super::StartPlay();
+		
+			StartMatch();
+		
+			if (GEngine != nullptr)
+			{
+				GEngine->AddOnScreenDebugMessage(-1, 5.f, FColor::Yellow, TEXT("We are using FPSGameMode."));
+			}
+		}//StartPlay()
+
+1.	Save `FPSGameMode.cpp` in Visual Studio.
+
+1.	Locate **FPSProject** in the **Solution Explorer**.
+
+1.	Right-click on **FPSProject** and select **Build** to compile your project.
+
+	![](BuildFPSProject.png)
+
+1.	After the build finishes, open your **FPSProject** in Unreal Editor.
+
+1.	Click the **Play** button in the Level Editor Toolbar. You should be able to aim the projectiles with your newly created crosshairs.
+
+	![](CrosshairsInGame.png)
+
+1.	Click the **Stop** button in the Level Editor to exit Play in Editor (PIE) mode.
+
+Congratulations, you can now:
+
+*	Move around your level
+*	Jump around your level
+*	Control your camera
+*	Shoot projectiles that interact with physics objects
+*	Aim with crosshairs
 
 %Steps%
\ No newline at end of file
