#UEDOC Added additional informaiton to this page as per StephenE

diff --git a/Engine/Documentation/Source/Platforms/Oculus/Oculus.INT.udn b/Engine/Documentation/Source/Platforms/Oculus/Oculus.INT.udn
--- a/Engine/Documentation/Source/Platforms/Oculus/Oculus.INT.udn
+++ b/Engine/Documentation/Source/Platforms/Oculus/Oculus.INT.udn
@@ -7,7 +7,7 @@ Platform:Oculus
 
 [EXCERPT:warning]
 [REGION:warning]
-Developing for Oculus on Mac currently requires Unreal Engine 4.3
+Developing for Oculus on Mac currently requires Unreal Engine 4.3 or higher
 [/REGION]
 [/EXCERPT:warning]
 
