#UE4doc #docs:
Changed font settings for multiple terms
Replaced a few images due to the updated UI
Changed some words to reflect how things are displayed in the editor now.

diff --git a/Engine/Documentation/Source/Engine/Animation/Sequences/Notifies/AnimNotifies.INT.udn b/Engine/Documentation/Source/Engine/Animation/Sequences/Notifies/AnimNotifies.INT.udn
--- a/Engine/Documentation/Source/Engine/Animation/Sequences/Notifies/AnimNotifies.INT.udn
+++ b/Engine/Documentation/Source/Engine/Animation/Sequences/Notifies/AnimNotifies.INT.udn
@@ -132,11 +132,11 @@ To use an AnimNotifyState:
 
 1. Finally from the **Content Browser**, **Double Left-click** your AnimNotifyState to bring up its Blueprint so you can start defining what if does.
 	
-	[REGION:lightbox]
-		[![](animNotifyStateImplementFunction.png)(w:860)](animNotifyStateImplementFunction.png)
-	[/REGION]
 
-At this point you can start implementing the 3 functions of the **AnimNotifyState**. It is not necessary to implement every function, but at the bare minimum you should implement `Recieved_NotifyTick` or both `Recieved_NotifyStart` and `Recieved_NotifyEnd`. If implementing just `Recieved_NotifyStart` or `Recieved_NotifyEnd`, you could simply use a standard Animation Notify.
+![](animNotifyStateImplementFunction.png)
+
+
+At this point you can start implementing the 3 functions of the **AnimNotifyState**. It is not necessary to implement every function, but at the bare minimum you should implement **Recieved Notify Tick** or both **Recieved Notify Start** and **Recieved Notify End**. If implementing just **Recieved Notify Start** or **Recieved Notify End**, you could simply use a standard Animation Notify.
 
 ![](Recieved_NotifyStart.png)
 
@@ -164,36 +164,36 @@ There is currently no way to retrieve the length of the AnimNotifyState within a
 
 ### Order and Logic
 
-When the **AnimNotifyState** is triggered, first the **Start** Blueprint will be called, and then the **Tick** Blueprint is called on the same frame update. Every subsequent `tick`, only the **Tick** Blueprint will be called. On the final `tick` update, the **Tick** Blueprint will be called followed by the **End** Blueprint.
+When the **AnimNotifyState** is triggered, first the **Start** Blueprint will be called, and then the **Tick** Blueprint is called on the same frame update. Every subsequent tick, only the Tick Blueprint will be called. On the final tick update, the Tick Blueprint will be called followed by the **End** Blueprint.
 
-If you are modifying something every frame, it should go in the **Tick** Blueprint. If you need to set variables, flags, or properties, to some value before the `ticks` hit, it should go in the **Start** Blueprint. And finally, the **End** Blueprint can be used to change variables, flags, or properties after the final `tick` hits.
+If you are modifying something every frame, it should go in the Tick Blueprint. If you need to set variables, flags, or properties, to some value before the ticks hit, it should go in the Start Blueprint. And finally, the End Blueprint can be used to change variables, flags, or properties after the final tick hits.
 
 
 ## Native Notify States 
-When choosing to add a **Notify State**, you can also assign one of two **Native Notify States** in the form of a **Timed Particle Effect** or an Animation **Trail**. 
+When choosing to add a Notify State, you can also assign one of two **Native Notify States** in the form of a **Timed Particle Effect** or an **Animation Trail**. 
 
 ![](AddingNotifyStates.png)
 
 ### Timed Particle Effects
 
-The **Timed Particle Effect** is very similar to the **Play Particle Effect** notify in that they both allow you to play a **Particle Effect** at any given point in an animation, the difference however is that with a **Timed Particle Effect**, you are given the ability to define the length of the effect inside the notify itself. 
+The Timed Particle Effect is very similar to the **Play Particle Effect** notify in that they both allow you to play a Particle Effect at any given point in an animation, the difference however is that with a Timed Particle Effect, you are given the ability to define the length of the effect inside the notify itself. 
 
 ![](TimedParticleEffect.png)
 
-In the image above, both a **Play Particle Effect** and **Timed Particle Effect** are assigned to an animation. Inside the yellow box is the **Timed Particle Effect** and you will notice that this notify has two handle bars; one at the start and one at the end of the event. By **Left-Clicking** on the handle bar at the end (denoted by the green arrow), you can drag to increase or decrease the length of the effect.
+In the image above, both a Play Particle Effect and Timed Particle Effect are assigned to an animation. Inside the yellow box is the Timed Particle Effect and you will notice that this notify has two handle bars; one at the start and one at the end of the event. By **Left-Clicking** on the handle bar at the end (denoted by the green arrow), you can drag to increase or decrease the length of the effect.
 
 ### Animation Trails
-**Animation Trails** (or **AnimTrail** for short) are emitters that allow you to leave trails behind moving objects in a scene. 
-They do this by sampling the positions of two sockets on a **Skeletal Mesh** each tick and stitching triangles between the sockets, down the length of the trail. 
+Animation Trails (or **AnimTrail** for short) are emitters that allow you to leave trails behind moving objects in a scene. 
+They do this by sampling the positions of two sockets on a Skeletal Mesh each tick and stitching triangles between the sockets, down the length of the trail. 
   
-To use an **AnimTrail** emitter you must do the following; 
-1.       Create a **Particle System** that contains an emitter with **AnimTrail** type data. 
-2.       In the **Animation Editor**, add an **AnimNotifyState** of type **Trail**. 
+To use an AnimTrail emitter you must do the following; 
+1.       Create a Particle System that contains an emitter with AnimTrail type data. 
+2.       In the **Animation Editor**, add an AnimNotifyState of type **Trail**. 
 
 Setting Up the Particle System
 
-Create a **Particle System** as usual and add a new emitter with type data **AnimTrail** to it. 
-Several properties of the type data module will affect how the trail behaves. Below is what you should see in **Cascade** when setting up the system. Note, you will see nothing in the preview window as trails must be used in conjunction with a **Skeletal Mesh** component as explained later. 
+Create a Particle System as usual and add a new emitter with type data AnimTrail to it. 
+Several properties of the type data module will affect how the trail behaves. Below is what you should see in **Cascade** when setting up the system. Note, you will see nothing in the preview window as trails must be used in conjunction with a Skeletal Mesh component as explained later. 
 
 [REGION:lightbox]  
 [![](TrailExample.png)(w:480 h:480)](TrailExample.png)
@@ -201,7 +201,7 @@ Several properties of the type data module will affect how the trail behaves. Be
 _click to view fullsize_
 
 
-Below are the properties inside the **AnimTrail** type data section of the particle system. 
+Below are the properties inside the AnimTrail type data section of the particle system. 
 
 | **Property** | **Description** |
 | --- | --- |
@@ -223,12 +223,12 @@ However, the more tessellated a trail is, the more expensive it is.
  
 Setting Up the AnimNotify
 
-Now that we have a particle system, we must use it with an **AnimNotify**. 
-Open the **Animation Editor**, add or use an existing **Timeline** under the **Notifies** section and **Right-Click** on it to bring up the menu. The **Trails** notify can be found under the **Add Notify State** section. 
+Now that we have a particle system, we must use it with an AnimNotify. 
+Open the **Animation Editor**, add or use an existing **Timeline** under the **Notifies** section and **Right-click** on it to bring up the menu. The Trails notify can be found under the **Add Notify State** section. 
 
 ![](SetUpTrail.png)
 
-Below is the **Details** panel and **Properties** for the **Trail** notify. 
+Below is the **Details** panel and **Properties** for the Trail notify. 
 
 ![](TrailDetails.png)
 
@@ -248,7 +248,7 @@ This method is based on providing a width scale value. A width scale of 1.0 will
 
 ![](WidthScaleMode.png)
 
-_The left image depicts **From Centre** which scales outwards from center, the middle image shows **From First Socket** which scales away from the first socket and the right image demonstrates **From Second Socket** which scales away from the second socket._ 
+_The left image depicts From Centre which scales outwards from center, the middle image shows From First Socket which scales away from the first socket and the right image demonstrates From Second Socket which scales away from the second socket._ 
 
 There are also some debug rendering options available:
 
@@ -260,10 +260,10 @@ There are also some debug rendering options available:
 | **Render Tessellation** | This toggles debug rendering of all vertices on the curve. This shows how the curve is tessellated between the spawn points. |
 
 ### Controlling Animation Trails in Blueprint
-While the above method is the standard approach to controlling **AnimTrails**, you can also control them from inside a **Blueprint** without the need of an **Anim Notify**. 
-The only restriction with this is that the **Particle System Component** must be a child of the **Skeletal Mesh Component** from which you are sampling socket locations. 
+While the above method is the standard approach to controlling AnimTrails, you can also control them from inside a Blueprint without the need of an Anim Notify. 
+The only restriction with this is that the Particle System Component must be a child of the Skeletal Mesh Component from which you are sampling socket locations. 
 
-1. Create an **Emitter** component (or use an existing one). 
+1. Create an Emitter component (or use an existing one). 
 1. Call the **BeginTrails** function on the component. This will start all trail emitters in the component. You also need to provide the trail's other properties here. 
 
     ![](TrailBlueprint2.png)
