#UE4doc #docs:
Added engine version field
Removed italics from a word that didn't need it
Replaced a picture to reflect changes made to the Anim Graph context menu.

diff --git a/Engine/Documentation/Source/Engine/Animation/Sequences/UserGuide/AnimSequenceUserGuide.INT.udn b/Engine/Documentation/Source/Engine/Animation/Sequences/UserGuide/AnimSequenceUserGuide.INT.udn
--- a/Engine/Documentation/Source/Engine/Animation/Sequences/UserGuide/AnimSequenceUserGuide.INT.udn
+++ b/Engine/Documentation/Source/Engine/Animation/Sequences/UserGuide/AnimSequenceUserGuide.INT.udn
@@ -1,7 +1,8 @@
 Availability:Public
 Title:Animation Sequence User Guide
 Crumbs: %ROOT%, Engine, Engine/Animation, Engine/Animation/Sequences
-Description:Guide to placing and sampling AnimationSequences in VimBlueprints
+Description:Guide to placing and sampling AnimationSequences in AnimBlueprints
+version: 4.9
 
 [EXCERPT:Intro]
 Playback of **Animation Sequences** can be performed through code, but is generally handled within an instance 
@@ -16,7 +17,7 @@ placed as sequence player nodes whose output is a pose generated by sampling the
 
 	![Using Animation Sequences - List](CE_anim_list.png)
 
-1. Select the _AnimationSequence_ to be sampled from the list. A node is placed in the graph for the sequence.
+1. Select the AnimationSequence to be sampled from the list. A node is placed in the graph for the sequence.
 
 	![Using Animation Sequences - Node](CE_anim_node.png)
 
