#UE4doc #docs:
Added engine version field
Removed Overview header
Changed font settings for multiple terms
Replaced an image due to a UI discrepency

diff --git a/Engine/Documentation/Source/Engine/Animation/Skeleton/Skeleton.INT.udn b/Engine/Documentation/Source/Engine/Animation/Skeleton/Skeleton.INT.udn
--- a/Engine/Documentation/Source/Engine/Animation/Skeleton/Skeleton.INT.udn
+++ b/Engine/Documentation/Source/Engine/Animation/Skeleton/Skeleton.INT.udn
@@ -2,12 +2,13 @@ Availability:Public
 Title: Skeleton Assets 
 Crumbs: %ROOT%, Engine, Engine/Animation
 Description:An overview of the Skeleton asset and how it applies to animation in UE4.
+version: 4.9
 
 [REGION:fullwidth]
 ![](SkeletonImage.png)
 [/REGION]
 
-## Overview
+
 
 In Unreal Engine 4, Skeletons are assets that are separate from Skeletal Meshes. Each Skeletal Mesh requires a Skeleton to be associated with it, and all animation assets (such as Animation Sequences, Animation Montages, and Blend Spaces) are associated with the Skeleton, not with the Skeletal Mesh.
 
@@ -21,7 +22,7 @@ A typical character skeleton as you might see in Maya.
 [/REGION]
 
 
-An important distinction is that a Skeleton asset is _not_ the same thing as the bone hierarchy found within a Skeletal Mesh. In Unreal Engine 4, the Skeleton asset is less like a 3D representation of an actual skeleton and more like a list which holds bone and hierarchy information for a specific type of character or Skeletal Mesh. It is through this list that animation associations are made. 
+An important distinction is that a Skeleton asset is not the same thing as the bone hierarchy found within a Skeletal Mesh. In Unreal Engine 4, the Skeleton asset is less like a 3D representation of an actual skeleton and more like a list which holds bone and hierarchy information for a specific type of character or Skeletal Mesh. It is through this list that animation associations are made. 
 
 In terms of its place in the animation system, it is probably best to think of a Skeleton asset as an intermediary between Skeletal Meshes and animations. Multiple Skeletal Meshes can use the same Skeleton asset, and any Skeletal Meshes using the same Skeleton asset will be able to share animations.
 
@@ -66,16 +67,16 @@ However, you could not re-order the hierarchy. For instance, this is not valid:
 
 ### Additional Limbs
 
-Because adding bones that are peripheral to the hierarchy is valid, it is easy to reuse Skeleton assets between Skeletal Meshes that have additional parts. Say, for example, you wanted to share animation between 3 different but similar character types, all having essentially the same bone structure (_and same bone names_) except for some key differences, such as:
+Because adding bones that are peripheral to the hierarchy is valid, it is easy to reuse Skeleton assets between Skeletal Meshes that have additional parts. Say, for example, you wanted to share animation between 3 different but similar character types, all having essentially the same bone structure (and same bone names) except for some key differences, such as:
 
 * A human male (for this example, this is like your basic character)
 * A human female with bones for a ponytail
 * A humanoid character with four arms
 
-You could use the same Skeleton asset with all of these characters, and the result would be that the Skeleton asset would contain joint hierarchy information for _all three_ types. The animation will use the Skeleton asset to determine what rotation to apply to each bone. Since the Skeletal Meshes are designed so that their base bone hierarchies are the same, the animations between them will work perfectly. The only difference is that you may need some special animations to drive the extra arms or ponytail on those characters.
+You could use the same Skeleton asset with all of these characters, and the result would be that the Skeleton asset would contain joint hierarchy information for all three types. The animation will use the Skeleton asset to determine what rotation to apply to each bone. Since the Skeletal Meshes are designed so that their base bone hierarchies are the same, the animations between them will work perfectly. The only difference is that you may need some special animations to drive the extra arms or ponytail on those characters.
 
 Any time animation data is received for a bone that is not included in the Skeletal Mesh (say, if ponytail animation was used on a character that used the same Skeleton asset but did not have a ponytail), then that animation data is ignored.
 
-There is very little performance overhead with this system, though in terms of production, it is generally a good idea to only share Skeleton assets between Skeletal Meshes with _similar_ bone hierarchies.
+There is very little performance overhead with this system, though in terms of production, it is generally a good idea to only share Skeleton assets between Skeletal Meshes with similar bone hierarchies.
 
 
