#UE4doc #docs:
Added engine version field
Changed font settings for a few terms

diff --git a/Engine/Documentation/Source/Engine/Animation/StateMachines/StateMachines.INT.udn b/Engine/Documentation/Source/Engine/Animation/StateMachines/StateMachines.INT.udn
--- a/Engine/Documentation/Source/Engine/Animation/StateMachines/StateMachines.INT.udn
+++ b/Engine/Documentation/Source/Engine/Animation/StateMachines/StateMachines.INT.udn
@@ -2,6 +2,7 @@ Availability:Public
 Title: State Machines
 Crumbs: %ROOT%, Engine, Engine/Animation
 Description:State Machines allow Skeletal Animations to be broken up into various states, with full control over how blends occur from one state to another.
+version: 4.9
 
 <!--
 [VAR:Topic]
@@ -28,7 +29,7 @@ Description:State Machines allow Skeletal Animations to be broken up into variou
 [TOC (start:2)]
 
 [EXCERPT:Intro]
-**State Machines** provide a graphical way to break the animation of a Skeletal Mesh into a series of _States_. These states are then governed by _Transition Rules_ that control how to blend from one state to another. As a tool, they greatly simplify the design process for Skeletal Mesh animation, in that you can create a graph that easily controls how your characters can flow between types of animation without having to create a complex Blueprint network.
+**State Machines** provide a graphical way to break the animation of a Skeletal Mesh into a series of States. These states are then governed by Transition Rules that control how to blend from one state to another. As a tool, they greatly simplify the design process for Skeletal Mesh animation, in that you can create a graph that easily controls how your characters can flow between types of animation without having to create a complex Blueprint network.
 
 They also simplify conceptualization of your animation graph. You can start the design of your animation by thinking of the necessary states you will need and how a character may flow from one state to another. A great way to think of them is as a way to break all possible animations for your character into an easy-to-read flow chart, in which you can watch the Skeletal Mesh's state move throughout the chart.
 
@@ -94,7 +95,7 @@ Here is such an Idle state seen from the State Machine.
 
 ![](State.png)
 
-When **double-clicked**, here is the simple graph found within that Idle state.
+When **Double-clicked**, here is the simple graph found within that Idle state.
 
 ![](IdleStateInside.png)
 
@@ -118,7 +119,7 @@ Once you have your states defined, you need to be able to control how your Skele
 
 ![](TransitionRule.png)
 
-Transition rules can perform any number of checks and tests against variable values, all with the purpose of outputting a _True/False_ value. This Boolean output determines whether the animation can move through the transition or not. 
+Transition rules can perform any number of checks and tests against variable values, all with the purpose of outputting a True/False value. This Boolean output determines whether the animation can move through the transition or not. 
 
 ![](TransitionRuleInside.png)
 
@@ -126,7 +127,7 @@ In this example, we are testing the value of a character's Speed. Once it reache
 
 ![](TransitionRuleReturn.png)
 
-Inside, this simply checks to see if the same value we used to transition _into_ the Movement state is no longer true. 
+Inside, this simply checks to see if the same value we used to transition into the Movement state is no longer true. 
 
 [REGION:fullwidth]
 ![](TransitionRuleReturnInside.png)
