#UE4doc #docs:
Added engine version field
Changed font settings for multiple terms
Fixed capitalization
Replaced a couple images due to the updated UI
Added information about the warning on a new State Machine

diff --git a/Engine/Documentation/Source/Engine/Animation/StateMachines/CreatingStateMachines/CreatingStateMachines.INT.udn b/Engine/Documentation/Source/Engine/Animation/StateMachines/CreatingStateMachines/CreatingStateMachines.INT.udn
--- a/Engine/Documentation/Source/Engine/Animation/StateMachines/CreatingStateMachines/CreatingStateMachines.INT.udn
+++ b/Engine/Documentation/Source/Engine/Animation/StateMachines/CreatingStateMachines/CreatingStateMachines.INT.udn
@@ -2,7 +2,7 @@ Availability:Public
 Title: Creating State Machines
 Crumbs: %ROOT%, Engine, Engine/Animation, Engine/Animation/StateMachines
 Description:Guide to creating new State Machines for use in AnimGraphs of animation assets.
-
+version: 4.9
 
 [TOC(start:2)]
 
@@ -10,36 +10,36 @@ Description:Guide to creating new State Machines for use in AnimGraphs of animat
 
 Before actually creating your **State Machine**, it is a good idea to take a little time and plan out exactly which states you will need, as well as the circumstances that will drive transitions between states. This will give you a solid plan of attack for not only the State Machine, but also the variables you will need to drive it and which animations will be a part of each state.
 
-Creating State Machines is done within the AnimGraph of an Animation Blueprint. Simply **right-click** within the AnimGraph and choose **State Machines > New State Machine** from the context menu.
+Creating State Machines is done within the AnimGraph of an Animation Blueprint. Simply **Right-click** within the AnimGraph and choose **State Machines > New State Machine** from the context menu.
 
 ![](CreateNewStateMachine.png)
 
-Generally, it is a good idea to name your State Machine immediately after creation. You can **right-click** the new State Machine in the _My Blueprint_ window and enter a new name in the name field. In this case, the State Machine has been named "Locomotion."
+Generally, it is a good idea to name your State Machine immediately after creation. You can **Right-click** the new State Machine in the **My Blueprint** window and enter a new name in the name field. In this case, the State Machine has been named "Locomotion."
 
 ![](RenameStateMachine.png)
 
 ![](LocomotionSMNode.png)
 
-When first created, a State Machine is blank. You can open it to define its states and transitions by **double-clicking** it. This will open a new graph tab. This graph begins with nothing but an Entry point, which essentially kicks off the animation.
+When first created, a State Machine is blank and will give a warning about this. You can open it to define its states and transitions by **Double-clicking** it. This will open a new graph tab. This graph begins with nothing but an Entry point, which essentially kicks off the animation.
 
 
 ![](StateMachineEntryPoint.png)
 
 ## Creating States and Conduits
 
-Creating a new State is similar to creating new nodes within a Blueprint. You can either **right-click** and create one out of the context menu, or you can drag a wire into an empty portion of the graph and release the mouse, which will also expose the context menu.
+Creating a new State is similar to creating new nodes within a Blueprint. You can either **Right-click** and create one out of the context menu, or you can drag a wire into an empty portion of the graph and release the mouse, which will also expose the context menu.
 
 ![](AddState_CM.png)
 
 ![](AddStateDragWire.png)
 
-Once a new State has been created, it should immediately be named. This can be done immediately after creation by clicking on the name of the new State, or by **right-clicking** the node in the graph, as well as within the **My Blueprint** panel.
+Once a new State has been created, it should immediately be named. This can be done immediately after creation by clicking on the name of the new State, or by **Right-clicking** the node in the graph, as well as within the **My Blueprint** panel.
 
 ![](StateDefault.png)
 
 ![](StateRenamed.png)
 
-Conduits are created in a similar manner, by either **right-clicking** in the graph or dragging a new wire into an empty space on the graph.
+Conduits are created in a similar manner, by either **Right-clicking** in the graph or dragging a new wire into an empty space on the graph.
 
 ![](CreateConduit.png)
 
