UEDOC-1835

Changed the arrow color from green to yellow

diff --git a/Engine/Documentation/Source/Engine/Blueprints/Editor/UIComponents/CompilerResults/CompilerResults.INT.udn b/Engine/Documentation/Source/Engine/Blueprints/Editor/UIComponents/CompilerResults/CompilerResults.INT.udn
--- a/Engine/Documentation/Source/Engine/Blueprints/Editor/UIComponents/CompilerResults/CompilerResults.INT.udn
+++ b/Engine/Documentation/Source/Engine/Blueprints/Editor/UIComponents/CompilerResults/CompilerResults.INT.udn
@@ -42,7 +42,7 @@ Whenever your compile process results in a warning or error, the **Compiler Resu
 * Click on a message to jump the graph view directly to the problematic node or part of the visual script.
 
 	[REGION:raw]
-	![](WarningFocus.png)(w:750)
+	![](WarningFocus.jpg)
 	[/REGION]
 	
 * Click the **Clear** button at the bottom right corner of the panel to clear out any existing messages.
