#UE4doc #docs:
Added engine version field
Changed font settings for multiple terms
Replaced a few images due to updated UI.

diff --git a/Engine/Documentation/Source/Engine/Animation/StateMachines/EditingStateMachines/EditingStateMachines.INT.udn b/Engine/Documentation/Source/Engine/Animation/StateMachines/EditingStateMachines/EditingStateMachines.INT.udn
--- a/Engine/Documentation/Source/Engine/Animation/StateMachines/EditingStateMachines/EditingStateMachines.INT.udn
+++ b/Engine/Documentation/Source/Engine/Animation/StateMachines/EditingStateMachines/EditingStateMachines.INT.udn
@@ -2,6 +2,7 @@ Availability:Public
 Title: Editing State Machines
 Crumbs: %ROOT%, Engine, Engine/Animation, Engine/Animation/StateMachines
 Description:Guide to editing animation State Machines.
+version: 4.9
 
 
 [TOC(start:2)]
