#UE4doc #docs:
Added engine version field
Changed font settings for multiple terms

diff --git a/Engine/Documentation/Source/Engine/Animation/StateMachines/TransitionRules/TransitionRules.INT.udn b/Engine/Documentation/Source/Engine/Animation/StateMachines/TransitionRules/TransitionRules.INT.udn
--- a/Engine/Documentation/Source/Engine/Animation/StateMachines/TransitionRules/TransitionRules.INT.udn
+++ b/Engine/Documentation/Source/Engine/Animation/StateMachines/TransitionRules/TransitionRules.INT.udn
@@ -2,13 +2,14 @@ Availability:Public
 Title: Transition Rules
 Crumbs: %ROOT%, Engine, Engine/Animation, Engine/Animation/StateMachines
 Description:Guide to rules that govern State Machine Transitions
+version: 4.9
 
 
 **Transition Rules** are method of controlling flow inside a **State Machine**. 
-**Transition Rules** are created automatically along the wires that connect **States** and **Conduits** and can perform any number of checks and tests against variable values, all with the purpose of outputting a _True_ or _False_ value. 
-This **Boolean** output determines whether the animation can move through the transition or not. 
+Transition Rules are created automatically along the wires that connect **States** and **Conduits** and can perform any number of checks and tests against variable values, all with the purpose of outputting a true or false value. 
+This Boolean output determines whether the animation can move through the transition or not. 
 
-Visually, a **Transition Rule** appears as a small circular directional icon inside the **AnimationGraph**, like so: 
+Visually, a Transition Rule appears as a small circular directional icon inside the **AnimationGraph**, like so: 
 
 ![](TransitionRule.png)
 
@@ -24,7 +25,7 @@ Rules can also be set to determine flow for a transition in both directions as w
 
 This means that a rule can be set to exit one state to move to another and a rule can be applied to transition back to that state. 
 
-When **left-clicking** on the **Transition Rule**, the **Details** panel will update and additional options will become available. 
+When **Left-clicking** on the Transition Rule, the **Details** panel will update and additional options will become available. 
 
 ![](TransitionDetails.png)
 
@@ -34,16 +35,16 @@ When **left-clicking** on the **Transition Rule**, the **Details** panel will up
 | **Bidirectional** | Determines if the transition can go both directions. |
 | **Blend Logic** | Specifies what transition logic to use: Standard Blend or Custom Blend. Selecting Custom will enable you to enter the Blend Graph to set a custom Blend Transition. |
 | **Transition Rule Sharing** | This option allows you to set whether or not the rule can be shared with other transitions or to use rule from another transition.  |
-| **Transition Interrupt** | Set the **Trigger Weight Threshold** and specify a **Notify State Class**.  |
-| **Transition Crossfade Sharing** | Use an existing **Transition Crossfade Blend** or promote the existing **Transition Crossfade** to shared.  |
+| **Transition Interrupt** | Set the Trigger Weight Threshold and specify a Notify State Class.  |
+| **Transition Crossfade Sharing** | Use an existing Transition Crossfade Blend or promote the existing Transition Crossfade to shared.  |
 | **Duration** | The duration to apply to the cross-fade. |
 | **Mode** | The type of blending applied to the cross-fade. |
-| **Start Transition Event / Custom Blueprint Event** | Call and execute an **Event** from within the notified **Blueprint** at the start of the transition. |
-| **Blueprint Notify** | **Blueprint** that is notified when the transition starts.  |
-| **End Transition Event / Custom Blueprint Event** | Call and execute an **Event** from within the notified **Blueprint** at the end of the transition. |
-| **Blueprint Notify** | **Blueprint** that is notified when the transition ends.  |
-| **Interrupt Transition Event / Custom Blueprint Event** | Call and execute an **Event** from within the notified **Blueprint** when the transition is interrupted. |
-| **Blueprint Notify** |  **Blueprint** that is notified when the transition is interrupted. |
+| **Start Transition Event / Custom Blueprint Event** | Call and execute an Event from within the notified Blueprint at the start of the transition. |
+| **Blueprint Notify** | Blueprint that is notified when the transition starts.  |
+| **End Transition Event / Custom Blueprint Event** | Call and execute an Event from within the notified Blueprint at the end of the transition. |
+| **Blueprint Notify** | Blueprint that is notified when the transition ends.  |
+| **Interrupt Transition Event / Custom Blueprint Event** | Call and execute an Event from within the notified Blueprint when the transition is interrupted. |
+| **Blueprint Notify** |  Blueprint that is notified when the transition is interrupted. |
 
 When **Double-Clicking** on the **Transition Rule**, a window will open where you can define the variables that must be met in order for the transition to occur. 
 
@@ -53,8 +54,8 @@ When **Double-Clicking** on the **Transition Rule**, a window will open where yo
 
 _click for full view_
 
-Similar to **Blueprints**, you can specify variables in the **My Blueprint** window and drag them into the graph. By connecting the variables to various conditions which are then connected to the **Result** node, you can define what must be _true_ in order for the transition to occur. 
+Similar to Blueprints, you can specify variables in the **My Blueprint** window and drag them into the graph. By connecting the variables to various conditions which are then connected to the **Result** node, you can define what must be _true_ in order for the transition to occur. 
 
-For information on getting started with **Blueprints**, refer to the [Introduction to Blueprints](Engine/Blueprints/GettingStarted) documentation. 
+For information on getting started with Blueprints, refer to the [Introduction to Blueprints](Engine/Blueprints/GettingStarted) documentation. 
 
 
