#UE4doc #docs:
Added engine version field
Removed Overview header
Changed font settings for multiple terms

diff --git a/Engine/Documentation/Source/Engine/Animation/SyncGroups/SyncGroups.INT.udn b/Engine/Documentation/Source/Engine/Animation/SyncGroups/SyncGroups.INT.udn
--- a/Engine/Documentation/Source/Engine/Animation/SyncGroups/SyncGroups.INT.udn
+++ b/Engine/Documentation/Source/Engine/Animation/SyncGroups/SyncGroups.INT.udn
@@ -2,20 +2,20 @@ Availability:Public
 Title: Sync Groups
 Crumbs: %ROOT%, Engine, Engine/Animation
 Description:Sync Groups allow you to maintain the synchronization of animations of different lengths.
+version: 4.9
 
 
-## Overview
 [EXCERPT:Intro]
 **Sync Groups** are to keep related animations that may have different lengths synchronized. For example, you may have a walk cycle and a run cycle that you would like to blend together so that the character can smoothly accelerate or decelerate. But what if these animations were different lengths, such as the walk being significantly longer than the run? In such a case, directly blending from one to the other would have unnatural results, an unsightly "beat" as the foot animation switches.
 
-**Sync Groups** solve this problem by allowing one primary animation node to serve as the Leader, and all other related animations will simply scale their time length to match. Typically, the leader is  the node with the greatest blend weight. As the weight blends and the follower's weight exceeds the leader, the follower becomes the leader. In this way, the two animations can work smoothly together offering a seamless transition from one motion to the next. 
+Sync Groups solve this problem by allowing one primary animation node to serve as the Leader, and all other related animations will simply scale their time length to match. Typically, the leader is  the node with the greatest blend weight. As the weight blends and the follower's weight exceeds the leader, the follower becomes the leader. In this way, the two animations can work smoothly together offering a seamless transition from one motion to the next. 
 
 It should be noted, however, that since the time of the animations is shifting, certain animation considerations become apparent. For instance, in the example of blending between different walk and run cycles, you would want to make sure your animations all started and ended on the same foot. Establishing these standards early will help everything blend much more smoothly!
 [/EXCERPT:Intro]
 
 ### Group Roles: Leaders and Followers
 
-Sync Groups work off the concept of group roles, in which one animation is the _leader_ while all other animations are _followers_. 
+Sync Groups work off the concept of group roles, in which one animation is the "leader" while all other animations are "followers". 
 
 | | |
 |---|---|
@@ -26,8 +26,8 @@ Leaders and followers are designated simply by which animation has the greater b
 
 | Image | Description |
 |---|---|
-|[INCLUDE:#runanimation]|Run animation, original size, set to a Sync Group named _WalkRun_.  |
-|[INCLUDE:#walkanimation]| Walk animation, original size, also set to a Sync Group named _WalkRun_. |
+|[INCLUDE:#runanimation]|Run animation, original size, set to a Sync Group named WalkRun.  |
+|[INCLUDE:#walkanimation]| Walk animation, original size, also set to a Sync Group named WalkRun. |
 |[INCLUDE:#walkanimationscaled] |In this instance, the blend weight for the run animation would be greater than the walk animation. Since they are in the same Sync Group, the walk scales down to match the length of the run animation.|
 |[INCLUDE:#runanimationscaled]|Here, the weighting has switched such that the walking animation is receiving a greater blend weight. This causes the run animation to scale up to match the length of the walk. | 
 
