UEDOC-1644

Adding in Compatibility for the iPhone6 and 6 Plus

diff --git a/Engine/Documentation/Source/Platforms/iOS/DeviceCompatibility/iOSDeviceCompatibility.INT.udn b/Engine/Documentation/Source/Platforms/iOS/DeviceCompatibility/iOSDeviceCompatibility.INT.udn
--- a/Engine/Documentation/Source/Platforms/iOS/DeviceCompatibility/iOSDeviceCompatibility.INT.udn
+++ b/Engine/Documentation/Source/Platforms/iOS/DeviceCompatibility/iOSDeviceCompatibility.INT.udn
@@ -34,6 +34,8 @@ This table lists mobile phone devices:
 | **iPhone4S** | %Retina% Retina | %NonRetina% Non-Retina | %Unsupported% Unsupported | %Unsupported% Unsupported |
 | **iPhone5** | %Retina% Retina | %Retina% Retina |%NonRetina%  Non-Retina | %Unsupported% Unsupported |
 | **iPhone5S** | %Retina% Retina | %Retina% Retina | %Retina% Retina | %Retina% Retina |
+| **iPhone6** | %Retina% Retina | %Retina% Retina | %Retina% Retina | %Retina% Retina |
+| **iPhone6 Plus** | %Retina% Retina | %Retina% Retina | %Retina% Retina | %Retina% Retina |
 [/REGION]
 
 This table lists tablet devices:
