#UE4doc #docs - Updated code in page for Blueprint setup. User Almo2001 had reported this on the AnswerHub.

diff --git a/Engine/Documentation/Source/Engine/AI/BehaviorTrees/QuickStart/11/11.INT.udn b/Engine/Documentation/Source/Engine/AI/BehaviorTrees/QuickStart/11/11.INT.udn
--- a/Engine/Documentation/Source/Engine/AI/BehaviorTrees/QuickStart/11/11.INT.udn
+++ b/Engine/Documentation/Source/Engine/AI/BehaviorTrees/QuickStart/11/11.INT.udn
@@ -2,6 +2,7 @@ Availability:Public
 Title: 11. Make a Service
 Crumbs:%ROOT%, Engine, Gameplay, Engine/AI/BehaviorTrees, Engine/AI/BehaviorTrees/QuickStart
 Description:Step 11. Make a Service
+Version: 4.9
 
 [VAR:Steps]
 [OBJECT:Navigation]
@@ -88,210 +89,205 @@ Here is what is going on in this graph:
 %Steps%
 
 [VAR:BPCode]
-	Begin Object Class=K2Node_Event Name="K2Node_Event_44"
-	   Begin Object Class=EdGraphPin Name="EdGraphPin_1590"
+	Begin Object Class=K2Node_CustomEvent Name="K2Node_CustomEvent_6"
+	   Begin Object Class=EdGraphPin Name="EdGraphPin_19767"
 	   End Object
-	   Begin Object Class=EdGraphPin Name="EdGraphPin_1591"
+	   Begin Object Class=EdGraphPin Name="EdGraphPin_19768"
 	   End Object
-	   Begin Object Class=EdGraphPin Name="EdGraphPin_1592"
+	   Begin Object Class=EdGraphPin Name="EdGraphPin_19769"
 	   End Object
-	   Begin Object Class=EdGraphPin Name="EdGraphPin_1593"
+	   Begin Object Class=EdGraphPin Name="EdGraphPin_19770"
 	   End Object
-	   Begin Object Name="EdGraphPin_1590"
+	   Begin Object Name="EdGraphPin_19767"
 		  PinName="OutputDelegate"
 		  Direction=EGPD_Output
-		  PinType=(PinCategory="delegate",PinSubCategoryMemberReference=(MemberParentClass=Class'/Script/AIModule.BTService_BlueprintBase',MemberName="ReceiveTick"))
+		  PinType=(PinCategory="delegate",PinSubCategoryMemberReference=(MemberParent=BlueprintGeneratedClass'/Game/StarterContent/Blueprints/Blueprint_Effect_Explosion.Blueprint_Effect_Explosion_C',MemberName="ReceiveTick_Copy_Copy",MemberGuid=6D1107C948D1F4E6B9BA9FBC654B2264))
 	   End Object
-	   Begin Object Name="EdGraphPin_1591"
+	   Begin Object Name="EdGraphPin_19768"
 		  PinName="then"
 		  Direction=EGPD_Output
 		  PinType=(PinCategory="exec")
-		  LinkedTo(0)=EdGraphPin'K2Node_IfThenElse_251.EdGraphPin_1802'
+		  LinkedTo(0)=EdGraphPin'K2Node_IfThenElse_28.EdGraphPin_20178'
 	   End Object
-	   Begin Object Name="EdGraphPin_1592"
+	   Begin Object Name="EdGraphPin_19769"
 		  PinName="OwnerActor"
-		  PinToolTip="Owner Actor\nActor Reference"
 		  Direction=EGPD_Output
 		  PinType=(PinCategory="object",PinSubCategoryObject=Class'/Script/Engine.Actor')
-		  LinkedTo(0)=EdGraphPin'K2Node_DynamicCast_140.EdGraphPin_6206'
+		  LinkedTo(0)=EdGraphPin'K2Node_DynamicCast_437.EdGraphPin_19782'
 	   End Object
-	   Begin Object Name="EdGraphPin_1593"
+	   Begin Object Name="EdGraphPin_19770"
 		  PinName="DeltaSeconds"
-		  PinToolTip="Delta Seconds\nFloat"
 		  Direction=EGPD_Output
 		  PinType=(PinCategory="float")
-		  DefaultValue="0.0"
-		  AutogeneratedDefaultValue="0.0"
 	   End Object
-	   EventSignatureName="ReceiveTick"
-	   EventSignatureClass=Class'/Script/AIModule.BTService_BlueprintBase'
-	   bOverrideFunction=True
-	   Pins(0)=EdGraphPin'EdGraphPin_1590'
-	   Pins(1)=EdGraphPin'EdGraphPin_1591'
-	   Pins(2)=EdGraphPin'EdGraphPin_1592'
-	   Pins(3)=EdGraphPin'EdGraphPin_1593'
-	   NodePosX=-480
-	   NodePosY=384
-	   NodeGuid=DBA2A7724B1F8C08DFA5958A86A71B98
+	   CustomFunctionName="ReceiveTick_Copy_Copy"
+	   Pins(0)=EdGraphPin'EdGraphPin_19767'
+	   Pins(1)=EdGraphPin'EdGraphPin_19768'
+	   Pins(2)=EdGraphPin'EdGraphPin_19769'
+	   Pins(3)=EdGraphPin'EdGraphPin_19770'
+	   NodePosX=-4592
+	   NodePosY=2304
+	   bHasCompilerMessage=True
+	   ErrorType=3
+	   ErrorMsg="Warning Can\'t connect pins  Owner Actor  and  Object : This cast has an invalid target type (was the class deleted without a redirect?)."
+	   NodeGuid=6D1107C948D1F4E6B9BA9FBC654B2264
+	   CustomProperties UserDefinedPin Name="OwnerActor" IsArray=0 IsReference=0 Category=object SubCategoryObject=/Script/Engine.Actor 
+	   CustomProperties UserDefinedPin Name="DeltaSeconds" IsArray=0 IsReference=0 Category=float 
 	End Object
-	Begin Object Class=K2Node_DynamicCast Name="K2Node_DynamicCast_140"
-	   Begin Object Class=EdGraphPin Name="EdGraphPin_6203"
+	Begin Object Class=K2Node_DynamicCast Name="K2Node_DynamicCast_437"
+	   Begin Object Class=EdGraphPin Name="EdGraphPin_19779"
 	   End Object
-	   Begin Object Class=EdGraphPin Name="EdGraphPin_6204"
+	   Begin Object Class=EdGraphPin Name="EdGraphPin_19780"
 	   End Object
-	   Begin Object Class=EdGraphPin Name="EdGraphPin_6205"
+	   Begin Object Class=EdGraphPin Name="EdGraphPin_19781"
 	   End Object
-	   Begin Object Class=EdGraphPin Name="EdGraphPin_6206"
+	   Begin Object Class=EdGraphPin Name="EdGraphPin_19782"
 	   End Object
-	   Begin Object Class=EdGraphPin Name="EdGraphPin_6207"
+	   Begin Object Class=EdGraphPin Name="EdGraphPin_19783"
 	   End Object
-	   Begin Object Class=EdGraphPin Name="EdGraphPin_6208"
-	   End Object
-	   Begin Object Name="EdGraphPin_6203"
+	   Begin Object Name="EdGraphPin_19779"
 		  PinName="execute"
 		  PinType=(PinCategory="exec")
-		  LinkedTo(0)=EdGraphPin'K2Node_IfThenElse_251.EdGraphPin_1805'
+		  LinkedTo(0)=EdGraphPin'K2Node_IfThenElse_28.EdGraphPin_20181'
 	   End Object
-	   Begin Object Name="EdGraphPin_6204"
+	   Begin Object Name="EdGraphPin_19780"
 		  PinName="then"
 		  Direction=EGPD_Output
 		  PinType=(PinCategory="exec")
-		  LinkedTo(0)=EdGraphPin'K2Node_VariableSet_149.EdGraphPin_1509'
+		  LinkedTo(0)=EdGraphPin'K2Node_VariableSet_11.EdGraphPin_20028'
 	   End Object
-	   Begin Object Name="EdGraphPin_6205"
+	   Begin Object Name="EdGraphPin_19781"
 		  PinName="CastFailed"
 		  Direction=EGPD_Output
 		  PinType=(PinCategory="exec")
 	   End Object
-	   Begin Object Name="EdGraphPin_6206"
+	   Begin Object Name="EdGraphPin_19782"
 		  PinName="Object"
 		  PinType=(PinCategory="object",PinSubCategoryObject=Class'/Script/CoreUObject.Object')
-		  LinkedTo(0)=EdGraphPin'K2Node_Event_44.EdGraphPin_1592'
-	   End Object
-	   Begin Object Name="EdGraphPin_6207"
-		  PinName="AsFollower AI CON"
-		  Direction=EGPD_Output
-		  PinType=(PinCategory="object",PinSubCategoryObject=BlueprintGeneratedClass'/Game/Blueprints/AI/Follower_AI_Con.Follower_AI_CON_C')
-		  LinkedTo(0)=EdGraphPin'K2Node_VariableSet_149.EdGraphPin_1511'
+		  LinkedTo(0)=EdGraphPin'K2Node_CustomEvent_6.EdGraphPin_19769'
 	   End Object
-	   Begin Object Name="EdGraphPin_6208"
+	   Begin Object Name="EdGraphPin_19783"
 		  PinName="bSuccess"
 		  Direction=EGPD_Output
 		  PinType=(PinCategory="bool")
 		  bHidden=True
 	   End Object
-	   TargetType=BlueprintGeneratedClass'/Game/Blueprints/AI/Follower_AI_Con.Follower_AI_CON_C'
-	   Pins(0)=EdGraphPin'EdGraphPin_6203'
-	   Pins(1)=EdGraphPin'EdGraphPin_6204'
-	   Pins(2)=EdGraphPin'EdGraphPin_6205'
-	   Pins(3)=EdGraphPin'EdGraphPin_6206'
-	   Pins(4)=EdGraphPin'EdGraphPin_6207'
-	   Pins(5)=EdGraphPin'EdGraphPin_6208'
-	   NodePosX=96
-	   NodePosY=384
-	   NodeGuid=432A90474A29E0D5304FE691968D7ACC
+	   Pins(0)=EdGraphPin'EdGraphPin_19779'
+	   Pins(1)=EdGraphPin'EdGraphPin_19780'
+	   Pins(2)=EdGraphPin'EdGraphPin_19781'
+	   Pins(3)=EdGraphPin'EdGraphPin_19782'
+	   Pins(4)=EdGraphPin'EdGraphPin_19783'
+	   NodePosX=-4016
+	   NodePosY=2304
+	   bHasCompilerMessage=True
+	   ErrorType=3
+	   ErrorMsg="Warning Can\'t connect pins  Object  and  Owner Actor : This cast has an invalid target type (was the class deleted without a redirect?)."
+	   NodeGuid=5791C348428BE0E23579BCBACAAA5EED
 	End Object
-	Begin Object Class=K2Node_CallFunction Name="K2Node_CallFunction_33307"
-	   Begin Object Class=EdGraphPin Name="EdGraphPin_1257"
+	Begin Object Class=K2Node_CallFunction Name="K2Node_CallFunction_2180"
+	   Begin Object Class=EdGraphPin Name="EdGraphPin_19794"
 	   End Object
-	   Begin Object Class=EdGraphPin Name="EdGraphPin_1258"
+	   Begin Object Class=EdGraphPin Name="EdGraphPin_19795"
 	   End Object
-	   Begin Object Name="EdGraphPin_1257"
+	   Begin Object Name="EdGraphPin_19794"
 		  PinName="self"
 		  PinFriendlyName="Target"
 		  PinToolTip="Target\nController Reference"
 		  PinType=(PinCategory="object",PinSubCategoryObject=Class'/Script/Engine.Controller')
-		  LinkedTo(0)=EdGraphPin'K2Node_VariableGet_340.EdGraphPin_1814'
+		  LinkedTo(0)=EdGraphPin'K2Node_VariableGet_41.EdGraphPin_20190'
 	   End Object
-	   Begin Object Name="EdGraphPin_1258"
+	   Begin Object Name="EdGraphPin_19795"
 		  PinName="ReturnValue"
 		  PinToolTip="Return Value\nPawn Reference"
 		  Direction=EGPD_Output
 		  PinType=(PinCategory="object",PinSubCategoryObject=Class'/Script/Engine.Pawn')
-		  LinkedTo(0)=EdGraphPin'K2Node_CallFunction_33401.EdGraphPin_1263'
+		  LinkedTo(0)=EdGraphPin'K2Node_CallFunction_2181.EdGraphPin_19800'
 	   End Object
 	   bIsPureFunc=True
 	   bIsConstFunc=True
-	   FunctionReference=(MemberParentClass=Class'/Script/Engine.Controller',MemberName="K2_GetPawn")
-	   Pins(0)=EdGraphPin'EdGraphPin_1257'
-	   Pins(1)=EdGraphPin'EdGraphPin_1258'
-	   NodePosX=400
-	   NodePosY=560
-	   NodeGuid=2873731E40941990681872A3D6B411AB
+	   FunctionReference=(MemberParent=Class'/Script/Engine.Controller',MemberName="K2_GetPawn")
+	   Pins(0)=EdGraphPin'EdGraphPin_19794'
+	   Pins(1)=EdGraphPin'EdGraphPin_19795'
+	   NodePosX=-3712
+	   NodePosY=2480
+	   bHasCompilerMessage=True
+	   ErrorType=3
+	   ErrorMsg="Warning Can\'t connect pins  Target  and  AI CON Ref : Controller Reference is not compatible with Reference."
+	   NodeGuid=F069F59C4F8E0476DF1ACE998083AA38
 	End Object
-	Begin Object Class=K2Node_CallFunction Name="K2Node_CallFunction_33401"
-	   Begin Object Class=EdGraphPin Name="EdGraphPin_1263"
+	Begin Object Class=K2Node_CallFunction Name="K2Node_CallFunction_2181"
+	   Begin Object Class=EdGraphPin Name="EdGraphPin_19800"
 	   End Object
-	   Begin Object Class=EdGraphPin Name="EdGraphPin_1264"
+	   Begin Object Class=EdGraphPin Name="EdGraphPin_19801"
 	   End Object
-	   Begin Object Name="EdGraphPin_1263"
+	   Begin Object Name="EdGraphPin_19800"
 		  PinName="self"
 		  PinFriendlyName="Target"
 		  PinToolTip="Target\nActor Reference"
 		  PinType=(PinCategory="object",PinSubCategoryObject=Class'/Script/Engine.Actor')
-		  LinkedTo(0)=EdGraphPin'K2Node_CallFunction_33307.EdGraphPin_1258'
+		  LinkedTo(0)=EdGraphPin'K2Node_CallFunction_2180.EdGraphPin_19795'
 	   End Object
-	   Begin Object Name="EdGraphPin_1264"
+	   Begin Object Name="EdGraphPin_19801"
 		  PinName="ReturnValue"
 		  PinToolTip="Return Value\nVector "
 		  Direction=EGPD_Output
-		  PinType=(PinCategory="struct",PinSubCategoryObject=ScriptStruct'/Script/CoreUObject.Object:Vector')
+		  PinType=(PinCategory="struct",PinSubCategoryObject=ScriptStruct'/Script/CoreUObject.Vector')
 		  DefaultValue="0, 0, 0"
 		  AutogeneratedDefaultValue="0, 0, 0"
-		  LinkedTo(0)=EdGraphPin'K2Node_CommutativeAssociativeBinaryOperator_518.EdGraphPin_1324'
-		  LinkedTo(1)=EdGraphPin'K2Node_CallFunction_39224.EdGraphPin_1285'
-		  LinkedTo(2)=EdGraphPin'K2Node_VariableSet_86.EdGraphPin_1439'
+		  LinkedTo(0)=EdGraphPin'K2Node_CallFunction_2182.EdGraphPin_19810'
+		  LinkedTo(1)=EdGraphPin'K2Node_CommutativeAssociativeBinaryOperator_24.EdGraphPin_19849'
+		  LinkedTo(2)=EdGraphPin'K2Node_VariableSet_10.EdGraphPin_19955'
 	   End Object
 	   bIsPureFunc=True
 	   bIsConstFunc=True
-	   FunctionReference=(MemberParentClass=Class'/Script/Engine.Actor',MemberName="K2_GetActorLocation")
-	   Pins(0)=EdGraphPin'EdGraphPin_1263'
-	   Pins(1)=EdGraphPin'EdGraphPin_1264'
-	   NodePosX=624
-	   NodePosY=560
-	   NodeGuid=17B04C114955ACDB6B5B12A2FAD54419
+	   FunctionReference=(MemberParent=Class'/Script/Engine.Actor',MemberName="K2_GetActorLocation")
+	   Pins(0)=EdGraphPin'EdGraphPin_19800'
+	   Pins(1)=EdGraphPin'EdGraphPin_19801'
+	   NodePosX=-3488
+	   NodePosY=2480
+	   NodeGuid=27DE28464468122F953517A540AD9308
 	End Object
-	Begin Object Class=K2Node_CallFunction Name="K2Node_CallFunction_39224"
-	   Begin Object Class=EdGraphPin Name="EdGraphPin_1281"
+	Begin Object Class=K2Node_CallFunction Name="K2Node_CallFunction_2182"
+	   Begin Object Class=EdGraphPin Name="EdGraphPin_19806"
 	   End Object
-	   Begin Object Class=EdGraphPin Name="EdGraphPin_1282"
+	   Begin Object Class=EdGraphPin Name="EdGraphPin_19807"
 	   End Object
-	   Begin Object Class=EdGraphPin Name="EdGraphPin_1283"
+	   Begin Object Class=EdGraphPin Name="EdGraphPin_19808"
 	   End Object
-	   Begin Object Class=EdGraphPin Name="EdGraphPin_1284"
+	   Begin Object Class=EdGraphPin Name="EdGraphPin_19809"
 	   End Object
-	   Begin Object Class=EdGraphPin Name="EdGraphPin_1285"
+	   Begin Object Class=EdGraphPin Name="EdGraphPin_19810"
 	   End Object
-	   Begin Object Class=EdGraphPin Name="EdGraphPin_1286"
+	   Begin Object Class=EdGraphPin Name="EdGraphPin_19811"
 	   End Object
-	   Begin Object Class=EdGraphPin Name="EdGraphPin_1287"
+	   Begin Object Class=EdGraphPin Name="EdGraphPin_19812"
 	   End Object
-	   Begin Object Class=EdGraphPin Name="EdGraphPin_1288"
+	   Begin Object Class=EdGraphPin Name="EdGraphPin_19813"
 	   End Object
-	   Begin Object Class=EdGraphPin Name="EdGraphPin_1289"
+	   Begin Object Class=EdGraphPin Name="EdGraphPin_19814"
 	   End Object
-	   Begin Object Class=EdGraphPin Name="EdGraphPin_1290"
+	   Begin Object Class=EdGraphPin Name="EdGraphPin_19815"
 	   End Object
-	   Begin Object Class=EdGraphPin Name="EdGraphPin_1291"
+	   Begin Object Class=EdGraphPin Name="EdGraphPin_19816"
 	   End Object
-	   Begin Object Class=EdGraphPin Name="EdGraphPin_1292"
+	   Begin Object Class=EdGraphPin Name="EdGraphPin_19817"
 	   End Object
-	   Begin Object Class=EdGraphPin Name="EdGraphPin_1293"
+	   Begin Object Class=EdGraphPin Name="EdGraphPin_19818"
 	   End Object
-	   Begin Object Class=EdGraphPin Name="EdGraphPin_1294"
+	   Begin Object Class=EdGraphPin Name="EdGraphPin_19819"
 	   End Object
-	   Begin Object Name="EdGraphPin_1281"
+	   Begin Object Name="EdGraphPin_19806"
 		  PinName="execute"
 		  PinType=(PinCategory="exec")
-		  LinkedTo(0)=EdGraphPin'K2Node_VariableSet_86.EdGraphPin_1438'
+		  LinkedTo(0)=EdGraphPin'K2Node_VariableSet_10.EdGraphPin_19954'
 	   End Object
-	   Begin Object Name="EdGraphPin_1282"
+	   Begin Object Name="EdGraphPin_19807"
 		  PinName="then"
 		  Direction=EGPD_Output
 		  PinType=(PinCategory="exec")
-		  LinkedTo(0)=EdGraphPin'K2Node_MacroInstance_762.EdGraphPin_6221'
+		  LinkedTo(0)=EdGraphPin'K2Node_MacroInstance_27.EdGraphPin_20037'
 	   End Object
-	   Begin Object Name="EdGraphPin_1283"
+	   Begin Object Name="EdGraphPin_19808"
 		  PinName="self"
 		  PinFriendlyName="Target"
 		  PinToolTip="Target\nKismet System Library Reference"
@@ -299,77 +295,77 @@ Here is what is going on in this graph:
 		  DefaultObject=Default__KismetSystemLibrary
 		  bHidden=True
 	   End Object
-	   Begin Object Name="EdGraphPin_1284"
+	   Begin Object Name="EdGraphPin_19809"
 		  PinName="WorldContextObject"
 		  PinToolTip="World Context Object\nObject Reference"
 		  PinType=(PinCategory="object",PinSubCategoryObject=Class'/Script/CoreUObject.Object')
 		  bHidden=True
 	   End Object
-	   Begin Object Name="EdGraphPin_1285"
+	   Begin Object Name="EdGraphPin_19810"
 		  PinName="Start"
 		  PinToolTip="Start\nVector \n\nStart of line segment."
-		  PinType=(PinCategory="struct",PinSubCategoryObject=ScriptStruct'/Script/CoreUObject.Object:Vector',bIsConst=True)
+		  PinType=(PinCategory="struct",PinSubCategoryObject=ScriptStruct'/Script/CoreUObject.Vector',bIsConst=True)
 		  DefaultValue="0, 0, 0"
 		  AutogeneratedDefaultValue="0, 0, 0"
-		  LinkedTo(0)=EdGraphPin'K2Node_CallFunction_33401.EdGraphPin_1264'
+		  LinkedTo(0)=EdGraphPin'K2Node_CallFunction_2181.EdGraphPin_19801'
 	   End Object
-	   Begin Object Name="EdGraphPin_1286"
+	   Begin Object Name="EdGraphPin_19811"
 		  PinName="End"
 		  PinToolTip="End\nVector \n\nEnd of line segment."
-		  PinType=(PinCategory="struct",PinSubCategoryObject=ScriptStruct'/Script/CoreUObject.Object:Vector',bIsConst=True)
+		  PinType=(PinCategory="struct",PinSubCategoryObject=ScriptStruct'/Script/CoreUObject.Vector',bIsConst=True)
 		  DefaultValue="0, 0, 0"
 		  AutogeneratedDefaultValue="0, 0, 0"
-		  LinkedTo(0)=EdGraphPin'K2Node_CommutativeAssociativeBinaryOperator_518.EdGraphPin_1326'
+		  LinkedTo(0)=EdGraphPin'K2Node_CommutativeAssociativeBinaryOperator_24.EdGraphPin_19851'
 	   End Object
-	   Begin Object Name="EdGraphPin_1287"
+	   Begin Object Name="EdGraphPin_19812"
 		  PinName="Radius"
 		  PinToolTip="Radius\nFloat\n\nRadius of the sphere to sweep"
 		  PinType=(PinCategory="float")
 		  DefaultValue="1500.0"
 		  AutogeneratedDefaultValue="0.0"
 	   End Object
-	   Begin Object Name="EdGraphPin_1288"
+	   Begin Object Name="EdGraphPin_19813"
 		  PinName="ObjectTypes"
 		  PinToolTip="Object Types\nArray of EObjectTypeQuery Enums\n\nArray of Object Types to trace"
-		  PinType=(PinCategory="byte",PinSubCategoryObject=Enum'/Script/Engine.EngineTypes:EObjectTypeQuery',bIsArray=True,bIsReference=True,bIsConst=True)
-		  LinkedTo(0)=EdGraphPin'K2Node_VariableGet_81.EdGraphPin_1335'
+		  PinType=(PinCategory="byte",PinSubCategoryObject=Enum'/Script/Engine.EObjectTypeQuery',bIsArray=True,bIsReference=True,bIsConst=True)
+		  LinkedTo(0)=EdGraphPin'K2Node_VariableGet_35.EdGraphPin_19860'
 		  bDefaultValueIsIgnored=True
 	   End Object
-	   Begin Object Name="EdGraphPin_1289"
+	   Begin Object Name="EdGraphPin_19814"
 		  PinName="bTraceComplex"
 		  PinToolTip="Trace Complex\nBoolean\n\nTrue to test against complex collision, false to test against simplified collision."
 		  PinType=(PinCategory="bool")
 		  DefaultValue="false"
 		  AutogeneratedDefaultValue="false"
 	   End Object
-	   Begin Object Name="EdGraphPin_1290"
+	   Begin Object Name="EdGraphPin_19815"
 		  PinName="ActorsToIgnore"
 		  PinToolTip="Actors to Ignore\nArray of Actor References"
 		  PinType=(PinCategory="object",PinSubCategoryObject=Class'/Script/Engine.Actor',bIsArray=True,bIsReference=True,bIsConst=True)
-		  LinkedTo(0)=EdGraphPin'K2Node_Knot_1.EdGraphPin_20244'
+		  LinkedTo(0)=EdGraphPin'K2Node_Knot_43.EdGraphPin_20209'
 		  bDefaultValueIsIgnored=True
 	   End Object
-	   Begin Object Name="EdGraphPin_1291"
+	   Begin Object Name="EdGraphPin_19816"
 		  PinName="DrawDebugType"
 		  PinToolTip="Draw Debug Type\nEDrawDebugTrace Enum"
-		  PinType=(PinCategory="byte",PinSubCategoryObject=Enum'/Script/Engine.KismetSystemLibrary:EDrawDebugTrace')
+		  PinType=(PinCategory="byte",PinSubCategoryObject=Enum'/Script/Engine.EDrawDebugTrace')
 		  DefaultValue="None"
 	   End Object
-	   Begin Object Name="EdGraphPin_1292"
+	   Begin Object Name="EdGraphPin_19817"
 		  PinName="OutHits"
 		  PinToolTip="Out Hits\nArray of Hit Result Structures\n\nA list of hits, sorted along the trace from start to finish.  The blocking hit will be the last hit, if there was one."
 		  Direction=EGPD_Output
-		  PinType=(PinCategory="struct",PinSubCategoryObject=ScriptStruct'/Script/Engine.EngineTypes:HitResult',bIsArray=True)
-		  LinkedTo(0)=EdGraphPin'K2Node_MacroInstance_762.EdGraphPin_6222'
+		  PinType=(PinCategory="struct",PinSubCategoryObject=ScriptStruct'/Script/Engine.HitResult',bIsArray=True)
+		  LinkedTo(0)=EdGraphPin'K2Node_MacroInstance_27.EdGraphPin_20038'
 	   End Object
-	   Begin Object Name="EdGraphPin_1293"
+	   Begin Object Name="EdGraphPin_19818"
 		  PinName="bIgnoreSelf"
 		  PinToolTip="Ignore Self\nBoolean"
 		  PinType=(PinCategory="bool")
 		  DefaultValue="true"
 		  AutogeneratedDefaultValue="true"
 	   End Object
-	   Begin Object Name="EdGraphPin_1294"
+	   Begin Object Name="EdGraphPin_19819"
 		  PinName="ReturnValue"
 		  PinToolTip="Return Value\nBoolean\n\nTrue if there was a hit, false otherwise."
 		  Direction=EGPD_Output
@@ -377,35 +373,35 @@ Here is what is going on in this graph:
 		  DefaultValue="false"
 		  AutogeneratedDefaultValue="false"
 	   End Object
-	   FunctionReference=(MemberParentClass=Class'/Script/Engine.KismetSystemLibrary',MemberName="SphereTraceMultiForObjects")
-	   Pins(0)=EdGraphPin'EdGraphPin_1281'
-	   Pins(1)=EdGraphPin'EdGraphPin_1282'
-	   Pins(2)=EdGraphPin'EdGraphPin_1283'
-	   Pins(3)=EdGraphPin'EdGraphPin_1284'
-	   Pins(4)=EdGraphPin'EdGraphPin_1285'
-	   Pins(5)=EdGraphPin'EdGraphPin_1286'
-	   Pins(6)=EdGraphPin'EdGraphPin_1287'
-	   Pins(7)=EdGraphPin'EdGraphPin_1288'
-	   Pins(8)=EdGraphPin'EdGraphPin_1289'
-	   Pins(9)=EdGraphPin'EdGraphPin_1290'
-	   Pins(10)=EdGraphPin'EdGraphPin_1291'
-	   Pins(11)=EdGraphPin'EdGraphPin_1292'
-	   Pins(12)=EdGraphPin'EdGraphPin_1293'
-	   Pins(13)=EdGraphPin'EdGraphPin_1294'
-	   NodePosX=1456
-	   NodePosY=544
-	   NodeGuid=080B4B384BF07BC92BC67AA962F0E3B8
+	   FunctionReference=(MemberParent=Class'/Script/Engine.KismetSystemLibrary',MemberName="SphereTraceMultiForObjects")
+	   Pins(0)=EdGraphPin'EdGraphPin_19806'
+	   Pins(1)=EdGraphPin'EdGraphPin_19807'
+	   Pins(2)=EdGraphPin'EdGraphPin_19808'
+	   Pins(3)=EdGraphPin'EdGraphPin_19809'
+	   Pins(4)=EdGraphPin'EdGraphPin_19810'
+	   Pins(5)=EdGraphPin'EdGraphPin_19811'
+	   Pins(6)=EdGraphPin'EdGraphPin_19812'
+	   Pins(7)=EdGraphPin'EdGraphPin_19813'
+	   Pins(8)=EdGraphPin'EdGraphPin_19814'
+	   Pins(9)=EdGraphPin'EdGraphPin_19815'
+	   Pins(10)=EdGraphPin'EdGraphPin_19816'
+	   Pins(11)=EdGraphPin'EdGraphPin_19817'
+	   Pins(12)=EdGraphPin'EdGraphPin_19818'
+	   Pins(13)=EdGraphPin'EdGraphPin_19819'
+	   NodePosX=-2656
+	   NodePosY=2464
+	   NodeGuid=58CC8F764DEBE8121F8309837449CC18
 	End Object
-	Begin Object Class=K2Node_CommutativeAssociativeBinaryOperator Name="K2Node_CommutativeAssociativeBinaryOperator_518"
-	   Begin Object Class=EdGraphPin Name="EdGraphPin_1323"
+	Begin Object Class=K2Node_CommutativeAssociativeBinaryOperator Name="K2Node_CommutativeAssociativeBinaryOperator_24"
+	   Begin Object Class=EdGraphPin Name="EdGraphPin_19848"
 	   End Object
-	   Begin Object Class=EdGraphPin Name="EdGraphPin_1324"
+	   Begin Object Class=EdGraphPin Name="EdGraphPin_19849"
 	   End Object
-	   Begin Object Class=EdGraphPin Name="EdGraphPin_1325"
+	   Begin Object Class=EdGraphPin Name="EdGraphPin_19850"
 	   End Object
-	   Begin Object Class=EdGraphPin Name="EdGraphPin_1326"
+	   Begin Object Class=EdGraphPin Name="EdGraphPin_19851"
 	   End Object
-	   Begin Object Name="EdGraphPin_1323"
+	   Begin Object Name="EdGraphPin_19848"
 		  PinName="self"
 		  PinFriendlyName="Target"
 		  PinToolTip="Target\nKismet Math Library Reference"
@@ -413,92 +409,101 @@ Here is what is going on in this graph:
 		  DefaultObject=Default__KismetMathLibrary
 		  bHidden=True
 	   End Object
-	   Begin Object Name="EdGraphPin_1324"
+	   Begin Object Name="EdGraphPin_19849"
 		  PinName="A"
 		  PinToolTip="A\nVector "
-		  PinType=(PinCategory="struct",PinSubCategoryObject=ScriptStruct'/Script/CoreUObject.Object:Vector')
+		  PinType=(PinCategory="struct",PinSubCategoryObject=ScriptStruct'/Script/CoreUObject.Vector')
 		  DefaultValue="0, 0, 0"
 		  AutogeneratedDefaultValue="0, 0, 0"
-		  LinkedTo(0)=EdGraphPin'K2Node_CallFunction_33401.EdGraphPin_1264'
+		  LinkedTo(0)=EdGraphPin'K2Node_CallFunction_2181.EdGraphPin_19801'
 	   End Object
-	   Begin Object Name="EdGraphPin_1325"
+	   Begin Object Name="EdGraphPin_19850"
 		  PinName="B"
 		  PinToolTip="B\nVector "
-		  PinType=(PinCategory="struct",PinSubCategoryObject=ScriptStruct'/Script/CoreUObject.Object:Vector')
+		  PinType=(PinCategory="struct",PinSubCategoryObject=ScriptStruct'/Script/CoreUObject.Vector')
 		  DefaultValue="0, 0,15.000000"
 		  AutogeneratedDefaultValue="0, 0, 0"
 	   End Object
-	   Begin Object Name="EdGraphPin_1326"
+	   Begin Object Name="EdGraphPin_19851"
 		  PinName="ReturnValue"
 		  PinToolTip="Return Value\nVector "
 		  Direction=EGPD_Output
-		  PinType=(PinCategory="struct",PinSubCategoryObject=ScriptStruct'/Script/CoreUObject.Object:Vector')
+		  PinType=(PinCategory="struct",PinSubCategoryObject=ScriptStruct'/Script/CoreUObject.Vector')
 		  DefaultValue="0, 0, 0"
 		  AutogeneratedDefaultValue="0, 0, 0"
-		  LinkedTo(0)=EdGraphPin'K2Node_CallFunction_39224.EdGraphPin_1286'
+		  LinkedTo(0)=EdGraphPin'K2Node_CallFunction_2182.EdGraphPin_19811'
 	   End Object
 	   bIsPureFunc=True
-	   FunctionReference=(MemberParentClass=Class'/Script/Engine.KismetMathLibrary',MemberName="Add_VectorVector")
-	   Pins(0)=EdGraphPin'EdGraphPin_1323'
-	   Pins(1)=EdGraphPin'EdGraphPin_1324'
-	   Pins(2)=EdGraphPin'EdGraphPin_1325'
-	   Pins(3)=EdGraphPin'EdGraphPin_1326'
-	   NodePosX=992
-	   NodePosY=656
-	   NodeGuid=76DB5E5B4876D3487BEDC8A98E8B7E00
+	   FunctionReference=(MemberParent=Class'/Script/Engine.KismetMathLibrary',MemberName="Add_VectorVector")
+	   Pins(0)=EdGraphPin'EdGraphPin_19848'
+	   Pins(1)=EdGraphPin'EdGraphPin_19849'
+	   Pins(2)=EdGraphPin'EdGraphPin_19850'
+	   Pins(3)=EdGraphPin'EdGraphPin_19851'
+	   NodePosX=-3120
+	   NodePosY=2576
+	   NodeGuid=88589C3B4FF2BB5BD7AF8C967445D287
 	End Object
-	Begin Object Class=K2Node_VariableGet Name="K2Node_VariableGet_81"
-	   Begin Object Class=EdGraphPin Name="EdGraphPin_1335"
+	Begin Object Class=K2Node_VariableGet Name="K2Node_VariableGet_35"
+	   Begin Object Class=EdGraphPin Name="EdGraphPin_19860"
 	   End Object
-	   Begin Object Class=EdGraphPin Name="EdGraphPin_1336"
+	   Begin Object Class=EdGraphPin Name="EdGraphPin_19861"
 	   End Object
-	   Begin Object Name="EdGraphPin_1335"
+	   Begin Object Name="EdGraphPin_19860"
 		  PinName="DesiredObjectTypes"
 		  Direction=EGPD_Output
-		  PinType=(PinCategory="byte",PinSubCategoryObject=Enum'/Script/Engine.EngineTypes:EObjectTypeQuery',bIsArray=True)
-		  LinkedTo(0)=EdGraphPin'K2Node_CallFunction_39224.EdGraphPin_1288'
+		  PinType=(PinCategory="byte",PinSubCategoryObject=Enum'/Script/Engine.EObjectTypeQuery',bIsArray=True)
+		  LinkedTo(0)=EdGraphPin'K2Node_CallFunction_2182.EdGraphPin_19813'
 	   End Object
-	   Begin Object Name="EdGraphPin_1336"
+	   Begin Object Name="EdGraphPin_19861"
 		  PinName="self"
 		  PinFriendlyName="Target"
-		  PinType=(PinCategory="object",PinSubCategoryObject=BlueprintGeneratedClass'/Game/Blueprints/AI/AgroCheck.AgroCheck_C')
+		  PinType=(PinCategory="object",PinSubCategoryObject=BlueprintGeneratedClass'/Game/StarterContent/Blueprints/Blueprint_Effect_Explosion.Blueprint_Effect_Explosion_C')
 		  bHidden=True
 	   End Object
 	   VariableReference=(MemberName="DesiredObjectTypes",bSelfContext=True)
-	   Pins(0)=EdGraphPin'EdGraphPin_1335'
-	   Pins(1)=EdGraphPin'EdGraphPin_1336'
-	   NodePosX=1232
-	   NodePosY=800
-	   NodeGuid=C9389E8E47AEB7FA8BFC229124C2783E
+	   Pins(0)=EdGraphPin'EdGraphPin_19860'
+	   Pins(1)=EdGraphPin'EdGraphPin_19861'
+	   NodePosX=-2880
+	   NodePosY=2720
+	   bHasCompilerMessage=True
+	   ErrorType=1
+	   ErrorMsg="Warning Could not find a variable named \"DesiredObjectTypes\" in \'Blueprint_Effect_Explosion_C\'.\nMake sure \'Blueprint_Effect_Explosion_C\' has been compiled for  Get DesiredObjectTypes \nError The property associated with  Desired Object Types  could not be found"
+	   NodeGuid=96942C864E7424B50A7235A931497C03
 	End Object
-	Begin Object Class=K2Node_CallFunction Name="K2Node_CallFunction_41995"
-	   Begin Object Class=EdGraphPin Name="EdGraphPin_1359"
+	Begin Object Class=K2Node_CallFunction Name="K2Node_CallFunction_2183"
+	   Begin Object Class=EdGraphPin Name="EdGraphPin_19866"
+	   End Object
+	   Begin Object Class=EdGraphPin Name="EdGraphPin_19867"
 	   End Object
-	   Begin Object Class=EdGraphPin Name="EdGraphPin_1360"
+	   Begin Object Class=EdGraphPin Name="EdGraphPin_19868"
 	   End Object
-	   Begin Object Class=EdGraphPin Name="EdGraphPin_1361"
+	   Begin Object Class=EdGraphPin Name="EdGraphPin_19869"
 	   End Object
-	   Begin Object Class=EdGraphPin Name="EdGraphPin_1362"
+	   Begin Object Class=EdGraphPin Name="EdGraphPin_19870"
 	   End Object
-	   Begin Object Class=EdGraphPin Name="EdGraphPin_1363"
+	   Begin Object Class=EdGraphPin Name="EdGraphPin_19871"
 	   End Object
-	   Begin Object Class=EdGraphPin Name="EdGraphPin_1364"
+	   Begin Object Class=EdGraphPin Name="EdGraphPin_19872"
 	   End Object
-	   Begin Object Class=EdGraphPin Name="EdGraphPin_1365"
+	   Begin Object Class=EdGraphPin Name="EdGraphPin_19873"
 	   End Object
-	   Begin Object Class=EdGraphPin Name="EdGraphPin_1366"
+	   Begin Object Class=EdGraphPin Name="EdGraphPin_19874"
 	   End Object
-	   Begin Object Class=EdGraphPin Name="EdGraphPin_1367"
+	   Begin Object Class=EdGraphPin Name="EdGraphPin_19875"
 	   End Object
-	   Begin Object Class=EdGraphPin Name="EdGraphPin_1368"
+	   Begin Object Class=EdGraphPin Name="EdGraphPin_19876"
 	   End Object
-	   Begin Object Class=EdGraphPin Name="EdGraphPin_1369"
+	   Begin Object Class=EdGraphPin Name="EdGraphPin_19877"
 	   End Object
-	   Begin Object Class=EdGraphPin Name="EdGraphPin_1370"
+	   Begin Object Class=EdGraphPin Name="EdGraphPin_19878"
 	   End Object
-	   Begin Object Class=EdGraphPin Name="EdGraphPin_1371"
+	   Begin Object Class=EdGraphPin Name="EdGraphPin_19879"
 	   End Object
-	   Begin Object Name="EdGraphPin_1359"
+	   Begin Object Class=EdGraphPin Name="EdGraphPin_19880"
+	   End Object
+	   Begin Object Class=EdGraphPin Name="EdGraphPin_19881"
+	   End Object
+	   Begin Object Name="EdGraphPin_19866"
 		  PinName="self"
 		  PinFriendlyName="Target"
 		  PinToolTip="Target\nGameplay Statics Reference"
@@ -506,14 +511,14 @@ Here is what is going on in this graph:
 		  DefaultObject=Default__GameplayStatics
 		  bHidden=True
 	   End Object
-	   Begin Object Name="EdGraphPin_1360"
+	   Begin Object Name="EdGraphPin_19867"
 		  PinName="Hit"
 		  PinToolTip="Hit\nHit Result Structure (by ref)\n\nThe source HitResult."
-		  PinType=(PinCategory="struct",PinSubCategoryObject=ScriptStruct'/Script/Engine.EngineTypes:HitResult',bIsReference=True,bIsConst=True)
-		  LinkedTo(0)=EdGraphPin'K2Node_MacroInstance_762.EdGraphPin_6225'
+		  PinType=(PinCategory="struct",PinSubCategoryObject=ScriptStruct'/Script/Engine.HitResult',bIsReference=True,bIsConst=True)
+		  LinkedTo(0)=EdGraphPin'K2Node_MacroInstance_27.EdGraphPin_20041'
 		  bDefaultValueIsIgnored=True
 	   End Object
-	   Begin Object Name="EdGraphPin_1361"
+	   Begin Object Name="EdGraphPin_19868"
 		  PinName="bBlockingHit"
 		  PinToolTip="Blocking Hit\nBoolean\n\nTrue if there was a blocking hit, false otherwise."
 		  Direction=EGPD_Output
@@ -521,7 +526,15 @@ Here is what is going on in this graph:
 		  DefaultValue="false"
 		  AutogeneratedDefaultValue="false"
 	   End Object
-	   Begin Object Name="EdGraphPin_1362"
+	   Begin Object Name="EdGraphPin_19869"
+		  PinName="bInitialOverlap"
+		  PinToolTip="Initial Overlap\nBoolean\n\nTrue if the hit started in an initial overlap. In this case some other values should be interpreted differently. Time will be 0, ImpactPoint will equal Location, and normals will be equal and indicate a depenetration vector."
+		  Direction=EGPD_Output
+		  PinType=(PinCategory="bool")
+		  DefaultValue="false"
+		  AutogeneratedDefaultValue="false"
+	   End Object
+	   Begin Object Name="EdGraphPin_19870"
 		  PinName="Time"
 		  PinToolTip="Time\nFloat\n\n\'Time\' of impact along trace direction ranging from [0.0 to 1.0) if there is a hit, indicating time between start and end. Equals 1.0 if there is no hit."
 		  Direction=EGPD_Output
@@ -529,59 +542,59 @@ Here is what is going on in this graph:
 		  DefaultValue="0.0"
 		  AutogeneratedDefaultValue="0.0"
 	   End Object
-	   Begin Object Name="EdGraphPin_1363"
+	   Begin Object Name="EdGraphPin_19871"
 		  PinName="Location"
 		  PinToolTip="Location\nVector \n\nLocation of the hit in world space. If this was a swept shape test, this is the location where we can place the shape in the world where it will not penetrate."
 		  Direction=EGPD_Output
-		  PinType=(PinCategory="struct",PinSubCategoryObject=ScriptStruct'/Script/CoreUObject.Object:Vector')
+		  PinType=(PinCategory="struct",PinSubCategoryObject=ScriptStruct'/Script/CoreUObject.Vector')
 		  DefaultValue="0, 0, 0"
 		  AutogeneratedDefaultValue="0, 0, 0"
 	   End Object
-	   Begin Object Name="EdGraphPin_1364"
-		  PinName="Normal"
-		  PinToolTip="Normal\nVector \n\nNormal of the hit in world space, for the object that was swept (e.g. for a sphere trace this points towards the sphere\'s center). Equal to ImpactNormal for line tests."
+	   Begin Object Name="EdGraphPin_19872"
+		  PinName="ImpactPoint"
+		  PinToolTip="Impact Point\nVector \n\nLocation of the actual contact point of the trace shape with the surface of the hit object. Equal to Location in the case of an initial overlap."
 		  Direction=EGPD_Output
-		  PinType=(PinCategory="struct",PinSubCategoryObject=ScriptStruct'/Script/CoreUObject.Object:Vector')
+		  PinType=(PinCategory="struct",PinSubCategoryObject=ScriptStruct'/Script/CoreUObject.Vector')
 		  DefaultValue="0, 0, 0"
 		  AutogeneratedDefaultValue="0, 0, 0"
 	   End Object
-	   Begin Object Name="EdGraphPin_1365"
-		  PinName="ImpactPoint"
-		  PinToolTip="Impact Point\nVector \n\nLocation of the actual contact point of the trace shape with the surface of the hit object."
+	   Begin Object Name="EdGraphPin_19873"
+		  PinName="Normal"
+		  PinToolTip="Normal\nVector \n\nNormal of the hit in world space, for the object that was swept (e.g. for a sphere trace this points towards the sphere\'s center). Equal to ImpactNormal for line tests."
 		  Direction=EGPD_Output
-		  PinType=(PinCategory="struct",PinSubCategoryObject=ScriptStruct'/Script/CoreUObject.Object:Vector')
+		  PinType=(PinCategory="struct",PinSubCategoryObject=ScriptStruct'/Script/CoreUObject.Vector')
 		  DefaultValue="0, 0, 0"
 		  AutogeneratedDefaultValue="0, 0, 0"
 	   End Object
-	   Begin Object Name="EdGraphPin_1366"
+	   Begin Object Name="EdGraphPin_19874"
 		  PinName="ImpactNormal"
 		  PinToolTip="Impact Normal\nVector \n\nNormal of the hit in world space, for the object that was hit by the sweep."
 		  Direction=EGPD_Output
-		  PinType=(PinCategory="struct",PinSubCategoryObject=ScriptStruct'/Script/CoreUObject.Object:Vector')
+		  PinType=(PinCategory="struct",PinSubCategoryObject=ScriptStruct'/Script/CoreUObject.Vector')
 		  DefaultValue="0, 0, 0"
 		  AutogeneratedDefaultValue="0, 0, 0"
 	   End Object
-	   Begin Object Name="EdGraphPin_1367"
+	   Begin Object Name="EdGraphPin_19875"
 		  PinName="PhysMat"
 		  PinToolTip="Phys Mat\nPhysical Material Reference\n\nPhysical material that was hit. Must set bReturnPhysicalMaterial to true in the query params for this to be returned."
 		  Direction=EGPD_Output
 		  PinType=(PinCategory="object",PinSubCategoryObject=Class'/Script/Engine.PhysicalMaterial')
 	   End Object
-	   Begin Object Name="EdGraphPin_1368"
+	   Begin Object Name="EdGraphPin_19876"
 		  PinName="HitActor"
 		  PinToolTip="Hit Actor\nActor Reference\n\nActor hit by the trace."
 		  Direction=EGPD_Output
 		  PinType=(PinCategory="object",PinSubCategoryObject=Class'/Script/Engine.Actor')
-		  LinkedTo(0)=EdGraphPin'K2Node_CallFunction_92334.EdGraphPin_1542'
-		  LinkedTo(1)=EdGraphPin'K2Node_Knot_3.EdGraphPin_20247'
+		  LinkedTo(0)=EdGraphPin'K2Node_CallFunction_2187.EdGraphPin_20058'
+		  LinkedTo(1)=EdGraphPin'K2Node_Knot_45.EdGraphPin_20220'
 	   End Object
-	   Begin Object Name="EdGraphPin_1369"
+	   Begin Object Name="EdGraphPin_19877"
 		  PinName="HitComponent"
 		  PinToolTip="Hit Component\nPrimitive Component Reference\n\nPrimitiveComponent hit by the trace."
 		  Direction=EGPD_Output
 		  PinType=(PinCategory="object",PinSubCategoryObject=Class'/Script/Engine.PrimitiveComponent')
 	   End Object
-	   Begin Object Name="EdGraphPin_1370"
+	   Begin Object Name="EdGraphPin_19878"
 		  PinName="HitBoneName"
 		  PinToolTip="Hit Bone Name\nName\n\nName of the bone hit (valid only if we hit a skeletal mesh)."
 		  Direction=EGPD_Output
@@ -589,7 +602,7 @@ Here is what is going on in this graph:
 		  DefaultValue="None"
 		  AutogeneratedDefaultValue="None"
 	   End Object
-	   Begin Object Name="EdGraphPin_1371"
+	   Begin Object Name="EdGraphPin_19879"
 		  PinName="HitItem"
 		  PinToolTip="Hit Item\nInteger\n\nPrimitive-specific data recording which item in the primitive was hit"
 		  Direction=EGPD_Output
@@ -597,64 +610,83 @@ Here is what is going on in this graph:
 		  DefaultValue="0"
 		  AutogeneratedDefaultValue="0"
 	   End Object
+	   Begin Object Name="EdGraphPin_19880"
+		  PinName="TraceStart"
+		  PinToolTip="Trace Start\nVector "
+		  Direction=EGPD_Output
+		  PinType=(PinCategory="struct",PinSubCategoryObject=ScriptStruct'/Script/CoreUObject.Vector')
+		  DefaultValue="0, 0, 0"
+		  AutogeneratedDefaultValue="0, 0, 0"
+	   End Object
+	   Begin Object Name="EdGraphPin_19881"
+		  PinName="TraceEnd"
+		  PinToolTip="Trace End\nVector "
+		  Direction=EGPD_Output
+		  PinType=(PinCategory="struct",PinSubCategoryObject=ScriptStruct'/Script/CoreUObject.Vector')
+		  DefaultValue="0, 0, 0"
+		  AutogeneratedDefaultValue="0, 0, 0"
+	   End Object
 	   bIsPureFunc=True
-	   FunctionReference=(MemberParentClass=Class'/Script/Engine.GameplayStatics',MemberName="BreakHitResult")
-	   Pins(0)=EdGraphPin'EdGraphPin_1359'
-	   Pins(1)=EdGraphPin'EdGraphPin_1360'
-	   Pins(2)=EdGraphPin'EdGraphPin_1361'
-	   Pins(3)=EdGraphPin'EdGraphPin_1362'
-	   Pins(4)=EdGraphPin'EdGraphPin_1363'
-	   Pins(5)=EdGraphPin'EdGraphPin_1364'
-	   Pins(6)=EdGraphPin'EdGraphPin_1365'
-	   Pins(7)=EdGraphPin'EdGraphPin_1366'
-	   Pins(8)=EdGraphPin'EdGraphPin_1367'
-	   Pins(9)=EdGraphPin'EdGraphPin_1368'
-	   Pins(10)=EdGraphPin'EdGraphPin_1369'
-	   Pins(11)=EdGraphPin'EdGraphPin_1370'
-	   Pins(12)=EdGraphPin'EdGraphPin_1371'
-	   NodePosX=2128
-	   NodePosY=720
-	   NodeGuid=401E84BA466A9DBEF26665B77743FB9A
+	   FunctionReference=(MemberParent=Class'/Script/Engine.GameplayStatics',MemberName="BreakHitResult")
+	   Pins(0)=EdGraphPin'EdGraphPin_19866'
+	   Pins(1)=EdGraphPin'EdGraphPin_19867'
+	   Pins(2)=EdGraphPin'EdGraphPin_19868'
+	   Pins(3)=EdGraphPin'EdGraphPin_19869'
+	   Pins(4)=EdGraphPin'EdGraphPin_19870'
+	   Pins(5)=EdGraphPin'EdGraphPin_19871'
+	   Pins(6)=EdGraphPin'EdGraphPin_19872'
+	   Pins(7)=EdGraphPin'EdGraphPin_19873'
+	   Pins(8)=EdGraphPin'EdGraphPin_19874'
+	   Pins(9)=EdGraphPin'EdGraphPin_19875'
+	   Pins(10)=EdGraphPin'EdGraphPin_19876'
+	   Pins(11)=EdGraphPin'EdGraphPin_19877'
+	   Pins(12)=EdGraphPin'EdGraphPin_19878'
+	   Pins(13)=EdGraphPin'EdGraphPin_19879'
+	   Pins(14)=EdGraphPin'EdGraphPin_19880'
+	   Pins(15)=EdGraphPin'EdGraphPin_19881'
+	   NodePosX=-1984
+	   NodePosY=2640
+	   NodeGuid=43206A0549F55A1474CA3F92727686F9
 	End Object
-	Begin Object Class=K2Node_CallFunction Name="K2Node_CallFunction_169714"
-	   Begin Object Class=EdGraphPin Name="EdGraphPin_1398"
+	Begin Object Class=K2Node_CallFunction Name="K2Node_CallFunction_2184"
+	   Begin Object Class=EdGraphPin Name="EdGraphPin_19914"
 	   End Object
-	   Begin Object Class=EdGraphPin Name="EdGraphPin_1399"
+	   Begin Object Class=EdGraphPin Name="EdGraphPin_19915"
 	   End Object
-	   Begin Object Class=EdGraphPin Name="EdGraphPin_1400"
+	   Begin Object Class=EdGraphPin Name="EdGraphPin_19916"
 	   End Object
-	   Begin Object Class=EdGraphPin Name="EdGraphPin_1401"
+	   Begin Object Class=EdGraphPin Name="EdGraphPin_19917"
 	   End Object
-	   Begin Object Class=EdGraphPin Name="EdGraphPin_1402"
+	   Begin Object Class=EdGraphPin Name="EdGraphPin_19918"
 	   End Object
-	   Begin Object Class=EdGraphPin Name="EdGraphPin_1403"
+	   Begin Object Class=EdGraphPin Name="EdGraphPin_19919"
 	   End Object
-	   Begin Object Class=EdGraphPin Name="EdGraphPin_1404"
+	   Begin Object Class=EdGraphPin Name="EdGraphPin_19920"
 	   End Object
-	   Begin Object Class=EdGraphPin Name="EdGraphPin_1405"
+	   Begin Object Class=EdGraphPin Name="EdGraphPin_19921"
 	   End Object
-	   Begin Object Class=EdGraphPin Name="EdGraphPin_1406"
+	   Begin Object Class=EdGraphPin Name="EdGraphPin_19922"
 	   End Object
-	   Begin Object Class=EdGraphPin Name="EdGraphPin_1407"
+	   Begin Object Class=EdGraphPin Name="EdGraphPin_19923"
 	   End Object
-	   Begin Object Class=EdGraphPin Name="EdGraphPin_1408"
+	   Begin Object Class=EdGraphPin Name="EdGraphPin_19924"
 	   End Object
-	   Begin Object Class=EdGraphPin Name="EdGraphPin_1409"
+	   Begin Object Class=EdGraphPin Name="EdGraphPin_19925"
 	   End Object
-	   Begin Object Class=EdGraphPin Name="EdGraphPin_1410"
+	   Begin Object Class=EdGraphPin Name="EdGraphPin_19926"
 	   End Object
-	   Begin Object Name="EdGraphPin_1398"
+	   Begin Object Name="EdGraphPin_19914"
 		  PinName="execute"
 		  PinType=(PinCategory="exec")
-		  LinkedTo(0)=EdGraphPin'K2Node_MacroInstance_762.EdGraphPin_6224'
+		  LinkedTo(0)=EdGraphPin'K2Node_MacroInstance_27.EdGraphPin_20040'
 	   End Object
-	   Begin Object Name="EdGraphPin_1399"
+	   Begin Object Name="EdGraphPin_19915"
 		  PinName="then"
 		  Direction=EGPD_Output
 		  PinType=(PinCategory="exec")
-		  LinkedTo(0)=EdGraphPin'K2Node_IfThenElse_229.EdGraphPin_1622'
+		  LinkedTo(0)=EdGraphPin'K2Node_IfThenElse_27.EdGraphPin_20142'
 	   End Object
-	   Begin Object Name="EdGraphPin_1400"
+	   Begin Object Name="EdGraphPin_19916"
 		  PinName="self"
 		  PinFriendlyName="Target"
 		  PinToolTip="Target\nKismet System Library Reference"
@@ -662,69 +694,69 @@ Here is what is going on in this graph:
 		  DefaultObject=Default__KismetSystemLibrary
 		  bHidden=True
 	   End Object
-	   Begin Object Name="EdGraphPin_1401"
+	   Begin Object Name="EdGraphPin_19917"
 		  PinName="WorldContextObject"
 		  PinToolTip="World Context Object\nObject Reference"
 		  PinType=(PinCategory="object",PinSubCategoryObject=Class'/Script/CoreUObject.Object')
 		  bHidden=True
 	   End Object
-	   Begin Object Name="EdGraphPin_1402"
+	   Begin Object Name="EdGraphPin_19918"
 		  PinName="Start"
 		  PinToolTip="Start\nVector \n\nStart of line segment."
-		  PinType=(PinCategory="struct",PinSubCategoryObject=ScriptStruct'/Script/CoreUObject.Object:Vector',bIsConst=True)
+		  PinType=(PinCategory="struct",PinSubCategoryObject=ScriptStruct'/Script/CoreUObject.Vector',bIsConst=True)
 		  DefaultValue="0, 0, 0"
 		  AutogeneratedDefaultValue="0, 0, 0"
-		  LinkedTo(0)=EdGraphPin'K2Node_VariableGet_89.EdGraphPin_1452'
+		  LinkedTo(0)=EdGraphPin'K2Node_VariableGet_36.EdGraphPin_19962'
 	   End Object
-	   Begin Object Name="EdGraphPin_1403"
+	   Begin Object Name="EdGraphPin_19919"
 		  PinName="End"
 		  PinToolTip="End\nVector \n\nEnd of line segment."
-		  PinType=(PinCategory="struct",PinSubCategoryObject=ScriptStruct'/Script/CoreUObject.Object:Vector',bIsConst=True)
+		  PinType=(PinCategory="struct",PinSubCategoryObject=ScriptStruct'/Script/CoreUObject.Vector',bIsConst=True)
 		  DefaultValue="0, 0, 0"
 		  AutogeneratedDefaultValue="0, 0, 0"
-		  LinkedTo(0)=EdGraphPin'K2Node_CallFunction_92334.EdGraphPin_1543'
+		  LinkedTo(0)=EdGraphPin'K2Node_CallFunction_2187.EdGraphPin_20059'
 	   End Object
-	   Begin Object Name="EdGraphPin_1404"
+	   Begin Object Name="EdGraphPin_19920"
 		  PinName="TraceChannel"
 		  PinToolTip="Trace Channel\nETraceTypeQuery Enum"
-		  PinType=(PinCategory="byte",PinSubCategoryObject=Enum'/Script/Engine.EngineTypes:ETraceTypeQuery')
+		  PinType=(PinCategory="byte",PinSubCategoryObject=Enum'/Script/Engine.ETraceTypeQuery')
 		  DefaultValue="TraceTypeQuery2"
 	   End Object
-	   Begin Object Name="EdGraphPin_1405"
+	   Begin Object Name="EdGraphPin_19921"
 		  PinName="bTraceComplex"
 		  PinToolTip="Trace Complex\nBoolean\n\nTrue to test against complex collision, false to test against simplified collision."
 		  PinType=(PinCategory="bool")
 		  DefaultValue="false"
 		  AutogeneratedDefaultValue="false"
 	   End Object
-	   Begin Object Name="EdGraphPin_1406"
+	   Begin Object Name="EdGraphPin_19922"
 		  PinName="ActorsToIgnore"
 		  PinToolTip="Actors to Ignore\nArray of Actor References"
 		  PinType=(PinCategory="object",PinSubCategoryObject=Class'/Script/Engine.Actor',bIsArray=True,bIsReference=True,bIsConst=True)
-		  LinkedTo(0)=EdGraphPin'K2Node_Knot_2.EdGraphPin_20246'
+		  LinkedTo(0)=EdGraphPin'K2Node_Knot_44.EdGraphPin_20215'
 		  bDefaultValueIsIgnored=True
 	   End Object
-	   Begin Object Name="EdGraphPin_1407"
+	   Begin Object Name="EdGraphPin_19923"
 		  PinName="DrawDebugType"
 		  PinToolTip="Draw Debug Type\nEDrawDebugTrace Enum"
-		  PinType=(PinCategory="byte",PinSubCategoryObject=Enum'/Script/Engine.KismetSystemLibrary:EDrawDebugTrace')
+		  PinType=(PinCategory="byte",PinSubCategoryObject=Enum'/Script/Engine.EDrawDebugTrace')
 		  DefaultValue="None"
 	   End Object
-	   Begin Object Name="EdGraphPin_1408"
+	   Begin Object Name="EdGraphPin_19924"
 		  PinName="OutHit"
 		  PinToolTip="Out Hit\nHit Result Structure\n\nProperties of the trace hit."
 		  Direction=EGPD_Output
-		  PinType=(PinCategory="struct",PinSubCategoryObject=ScriptStruct'/Script/Engine.EngineTypes:HitResult')
-		  LinkedTo(0)=EdGraphPin'K2Node_CallFunction_42262.EdGraphPin_1459'
+		  PinType=(PinCategory="struct",PinSubCategoryObject=ScriptStruct'/Script/Engine.HitResult')
+		  LinkedTo(0)=EdGraphPin'K2Node_CallFunction_2185.EdGraphPin_19969'
 	   End Object
-	   Begin Object Name="EdGraphPin_1409"
+	   Begin Object Name="EdGraphPin_19925"
 		  PinName="bIgnoreSelf"
 		  PinToolTip="Ignore Self\nBoolean"
 		  PinType=(PinCategory="bool")
 		  DefaultValue="true"
 		  AutogeneratedDefaultValue="true"
 	   End Object
-	   Begin Object Name="EdGraphPin_1410"
+	   Begin Object Name="EdGraphPin_19926"
 		  PinName="ReturnValue"
 		  PinToolTip="Return Value\nBoolean\n\nTrue if there was a hit, false otherwise."
 		  Direction=EGPD_Output
@@ -732,131 +764,123 @@ Here is what is going on in this graph:
 		  DefaultValue="false"
 		  AutogeneratedDefaultValue="false"
 	   End Object
-	   FunctionReference=(MemberParentClass=Class'/Script/Engine.KismetSystemLibrary',MemberName="LineTraceSingle_NEW")
-	   Pins(0)=EdGraphPin'EdGraphPin_1398'
-	   Pins(1)=EdGraphPin'EdGraphPin_1399'
-	   Pins(2)=EdGraphPin'EdGraphPin_1400'
-	   Pins(3)=EdGraphPin'EdGraphPin_1401'
-	   Pins(4)=EdGraphPin'EdGraphPin_1402'
-	   Pins(5)=EdGraphPin'EdGraphPin_1403'
-	   Pins(6)=EdGraphPin'EdGraphPin_1404'
-	   Pins(7)=EdGraphPin'EdGraphPin_1405'
-	   Pins(8)=EdGraphPin'EdGraphPin_1406'
-	   Pins(9)=EdGraphPin'EdGraphPin_1407'
-	   Pins(10)=EdGraphPin'EdGraphPin_1408'
-	   Pins(11)=EdGraphPin'EdGraphPin_1409'
-	   Pins(12)=EdGraphPin'EdGraphPin_1410'
-	   NodePosX=2688
-	   NodePosY=544
-	   NodeGuid=3D53187C4646E111682A33A88639171A
+	   FunctionReference=(MemberParent=Class'/Script/Engine.KismetSystemLibrary',MemberName="LineTraceSingle_NEW")
+	   Pins(0)=EdGraphPin'EdGraphPin_19914'
+	   Pins(1)=EdGraphPin'EdGraphPin_19915'
+	   Pins(2)=EdGraphPin'EdGraphPin_19916'
+	   Pins(3)=EdGraphPin'EdGraphPin_19917'
+	   Pins(4)=EdGraphPin'EdGraphPin_19918'
+	   Pins(5)=EdGraphPin'EdGraphPin_19919'
+	   Pins(6)=EdGraphPin'EdGraphPin_19920'
+	   Pins(7)=EdGraphPin'EdGraphPin_19921'
+	   Pins(8)=EdGraphPin'EdGraphPin_19922'
+	   Pins(9)=EdGraphPin'EdGraphPin_19923'
+	   Pins(10)=EdGraphPin'EdGraphPin_19924'
+	   Pins(11)=EdGraphPin'EdGraphPin_19925'
+	   Pins(12)=EdGraphPin'EdGraphPin_19926'
+	   NodePosX=-1424
+	   NodePosY=2464
+	   NodeGuid=CD0FA2E24E93A82F677826A64000AE11
 	End Object
-	Begin Object Class=K2Node_VariableSet Name="K2Node_VariableSet_86"
-	   Begin Object Class=EdGraphPin Name="EdGraphPin_1437"
-	   End Object
-	   Begin Object Class=EdGraphPin Name="EdGraphPin_1438"
+	Begin Object Class=K2Node_VariableSet Name="K2Node_VariableSet_10"
+	   Begin Object Class=EdGraphPin Name="EdGraphPin_19953"
 	   End Object
-	   Begin Object Class=EdGraphPin Name="EdGraphPin_1439"
+	   Begin Object Class=EdGraphPin Name="EdGraphPin_19954"
 	   End Object
-	   Begin Object Class=EdGraphPin Name="EdGraphPin_1440"
+	   Begin Object Class=EdGraphPin Name="EdGraphPin_19955"
 	   End Object
-	   Begin Object Class=EdGraphPin Name="EdGraphPin_1441"
-	   End Object
-	   Begin Object Name="EdGraphPin_1437"
+	   Begin Object Name="EdGraphPin_19953"
 		  PinName="execute"
 		  PinType=(PinCategory="exec")
-		  LinkedTo(0)=EdGraphPin'K2Node_CallFunction_834898.EdGraphPin_1767'
+		  LinkedTo(0)=EdGraphPin'K2Node_CallFunction_2192.EdGraphPin_20155'
 	   End Object
-	   Begin Object Name="EdGraphPin_1438"
+	   Begin Object Name="EdGraphPin_19954"
 		  PinName="then"
 		  Direction=EGPD_Output
 		  PinType=(PinCategory="exec")
-		  LinkedTo(0)=EdGraphPin'K2Node_CallFunction_39224.EdGraphPin_1281'
+		  LinkedTo(0)=EdGraphPin'K2Node_CallFunction_2182.EdGraphPin_19806'
 	   End Object
-	   Begin Object Name="EdGraphPin_1439"
+	   Begin Object Name="EdGraphPin_19955"
 		  PinName="MyLocation"
-		  PinType=(PinCategory="struct",PinSubCategoryObject=ScriptStruct'/Script/CoreUObject.Object:Vector')
+		  PinType=(PinCategory="struct",PinSubCategoryObject=ScriptStruct'/Script/CoreUObject.Vector')
 		  DefaultValue="0, 0, 0"
 		  AutogeneratedDefaultValue="0, 0, 0"
-		  LinkedTo(0)=EdGraphPin'K2Node_CallFunction_33401.EdGraphPin_1264'
-	   End Object
-	   Begin Object Name="EdGraphPin_1440"
-		  PinName="Output_Get"
-		  PinToolTip="Retrieves the value of the variable, can use instead of a separate Get node"
-		  Direction=EGPD_Output
-		  PinType=(PinCategory="struct",PinSubCategoryObject=ScriptStruct'/Script/CoreUObject.Object:Vector')
-		  DefaultValue="0, 0, 0"
-		  AutogeneratedDefaultValue="0, 0, 0"
-	   End Object
-	   Begin Object Name="EdGraphPin_1441"
-		  PinName="self"
-		  PinFriendlyName="Target"
-		  PinType=(PinCategory="object",PinSubCategoryObject=BlueprintGeneratedClass'/Game/Blueprints/AI/AgroCheck.AgroCheck_C')
-		  bHidden=True
+		  LinkedTo(0)=EdGraphPin'K2Node_CallFunction_2181.EdGraphPin_19801'
 	   End Object
 	   VariableReference=(MemberName="MyLocation",bSelfContext=True)
-	   Pins(0)=EdGraphPin'EdGraphPin_1437'
-	   Pins(1)=EdGraphPin'EdGraphPin_1438'
-	   Pins(2)=EdGraphPin'EdGraphPin_1439'
-	   Pins(3)=EdGraphPin'EdGraphPin_1440'
-	   Pins(4)=EdGraphPin'EdGraphPin_1441'
-	   NodePosX=1152
-	   NodePosY=256
-	   NodeGuid=EA2BEE1F49C81272D71C4C8C6512E205
+	   Pins(0)=EdGraphPin'EdGraphPin_19953'
+	   Pins(1)=EdGraphPin'EdGraphPin_19954'
+	   Pins(2)=EdGraphPin'EdGraphPin_19955'
+	   NodePosX=-2960
+	   NodePosY=2176
+	   bHasCompilerMessage=True
+	   ErrorType=1
+	   ErrorMsg="Warning Could not find a variable named \"MyLocation\" in \'Blueprint_Effect_Explosion_C\'.\nMake sure \'Blueprint_Effect_Explosion_C\' has been compiled for  Set MyLocation \nError The property associated with  My Location  could not be found"
+	   NodeGuid=4CF54639407C0DF97B2C889DB7A1994F
 	End Object
-	Begin Object Class=K2Node_VariableGet Name="K2Node_VariableGet_89"
-	   Begin Object Class=EdGraphPin Name="EdGraphPin_1452"
+	Begin Object Class=K2Node_VariableGet Name="K2Node_VariableGet_36"
+	   Begin Object Class=EdGraphPin Name="EdGraphPin_19962"
 	   End Object
-	   Begin Object Class=EdGraphPin Name="EdGraphPin_1453"
+	   Begin Object Class=EdGraphPin Name="EdGraphPin_19963"
 	   End Object
-	   Begin Object Name="EdGraphPin_1452"
+	   Begin Object Name="EdGraphPin_19962"
 		  PinName="MyLocation"
 		  Direction=EGPD_Output
-		  PinType=(PinCategory="struct",PinSubCategoryObject=ScriptStruct'/Script/CoreUObject.Object:Vector')
+		  PinType=(PinCategory="struct",PinSubCategoryObject=ScriptStruct'/Script/CoreUObject.Vector')
 		  DefaultValue="0, 0, 0"
 		  AutogeneratedDefaultValue="0, 0, 0"
-		  LinkedTo(0)=EdGraphPin'K2Node_CallFunction_169714.EdGraphPin_1402'
+		  LinkedTo(0)=EdGraphPin'K2Node_CallFunction_2184.EdGraphPin_19918'
 	   End Object
-	   Begin Object Name="EdGraphPin_1453"
+	   Begin Object Name="EdGraphPin_19963"
 		  PinName="self"
 		  PinFriendlyName="Target"
-		  PinType=(PinCategory="object",PinSubCategoryObject=BlueprintGeneratedClass'/Game/Blueprints/AI/AgroCheck.AgroCheck_C')
+		  PinType=(PinCategory="object",PinSubCategoryObject=BlueprintGeneratedClass'/Game/StarterContent/Blueprints/Blueprint_Effect_Explosion.Blueprint_Effect_Explosion_C')
 		  bHidden=True
 	   End Object
 	   VariableReference=(MemberName="MyLocation",bSelfContext=True)
-	   Pins(0)=EdGraphPin'EdGraphPin_1452'
-	   Pins(1)=EdGraphPin'EdGraphPin_1453'
-	   NodePosX=2448
-	   NodePosY=624
-	   NodeGuid=22FE3CB446C68C2496052083E29E1CFE
+	   Pins(0)=EdGraphPin'EdGraphPin_19962'
+	   Pins(1)=EdGraphPin'EdGraphPin_19963'
+	   NodePosX=-1664
+	   NodePosY=2544
+	   bHasCompilerMessage=True
+	   ErrorType=1
+	   ErrorMsg="Warning Could not find a variable named \"MyLocation\" in \'Blueprint_Effect_Explosion_C\'.\nMake sure \'Blueprint_Effect_Explosion_C\' has been compiled for  Get MyLocation \nError The property associated with  My Location  could not be found"
+	   NodeGuid=918C98B94449B34384258CB776D1E920
 	End Object
-	Begin Object Class=K2Node_CallFunction Name="K2Node_CallFunction_42262"
-	   Begin Object Class=EdGraphPin Name="EdGraphPin_1458"
+	Begin Object Class=K2Node_CallFunction Name="K2Node_CallFunction_2185"
+	   Begin Object Class=EdGraphPin Name="EdGraphPin_19968"
+	   End Object
+	   Begin Object Class=EdGraphPin Name="EdGraphPin_19969"
+	   End Object
+	   Begin Object Class=EdGraphPin Name="EdGraphPin_19970"
+	   End Object
+	   Begin Object Class=EdGraphPin Name="EdGraphPin_19971"
 	   End Object
-	   Begin Object Class=EdGraphPin Name="EdGraphPin_1459"
+	   Begin Object Class=EdGraphPin Name="EdGraphPin_19972"
 	   End Object
-	   Begin Object Class=EdGraphPin Name="EdGraphPin_1460"
+	   Begin Object Class=EdGraphPin Name="EdGraphPin_19973"
 	   End Object
-	   Begin Object Class=EdGraphPin Name="EdGraphPin_1461"
+	   Begin Object Class=EdGraphPin Name="EdGraphPin_19974"
 	   End Object
-	   Begin Object Class=EdGraphPin Name="EdGraphPin_1462"
+	   Begin Object Class=EdGraphPin Name="EdGraphPin_19975"
 	   End Object
-	   Begin Object Class=EdGraphPin Name="EdGraphPin_1463"
+	   Begin Object Class=EdGraphPin Name="EdGraphPin_19976"
 	   End Object
-	   Begin Object Class=EdGraphPin Name="EdGraphPin_1464"
+	   Begin Object Class=EdGraphPin Name="EdGraphPin_19977"
 	   End Object
-	   Begin Object Class=EdGraphPin Name="EdGraphPin_1465"
+	   Begin Object Class=EdGraphPin Name="EdGraphPin_19978"
 	   End Object
-	   Begin Object Class=EdGraphPin Name="EdGraphPin_1466"
+	   Begin Object Class=EdGraphPin Name="EdGraphPin_19979"
 	   End Object
-	   Begin Object Class=EdGraphPin Name="EdGraphPin_1467"
+	   Begin Object Class=EdGraphPin Name="EdGraphPin_19980"
 	   End Object
-	   Begin Object Class=EdGraphPin Name="EdGraphPin_1468"
+	   Begin Object Class=EdGraphPin Name="EdGraphPin_19981"
 	   End Object
-	   Begin Object Class=EdGraphPin Name="EdGraphPin_1469"
+	   Begin Object Class=EdGraphPin Name="EdGraphPin_19982"
 	   End Object
-	   Begin Object Class=EdGraphPin Name="EdGraphPin_1470"
+	   Begin Object Class=EdGraphPin Name="EdGraphPin_19983"
 	   End Object
-	   Begin Object Name="EdGraphPin_1458"
+	   Begin Object Name="EdGraphPin_19968"
 		  PinName="self"
 		  PinFriendlyName="Target"
 		  PinToolTip="Target\nGameplay Statics Reference"
@@ -864,14 +888,14 @@ Here is what is going on in this graph:
 		  DefaultObject=Default__GameplayStatics
 		  bHidden=True
 	   End Object
-	   Begin Object Name="EdGraphPin_1459"
+	   Begin Object Name="EdGraphPin_19969"
 		  PinName="Hit"
 		  PinToolTip="Hit\nHit Result Structure (by ref)\n\nThe source HitResult."
-		  PinType=(PinCategory="struct",PinSubCategoryObject=ScriptStruct'/Script/Engine.EngineTypes:HitResult',bIsReference=True,bIsConst=True)
-		  LinkedTo(0)=EdGraphPin'K2Node_CallFunction_169714.EdGraphPin_1408'
+		  PinType=(PinCategory="struct",PinSubCategoryObject=ScriptStruct'/Script/Engine.HitResult',bIsReference=True,bIsConst=True)
+		  LinkedTo(0)=EdGraphPin'K2Node_CallFunction_2184.EdGraphPin_19924'
 		  bDefaultValueIsIgnored=True
 	   End Object
-	   Begin Object Name="EdGraphPin_1460"
+	   Begin Object Name="EdGraphPin_19970"
 		  PinName="bBlockingHit"
 		  PinToolTip="Blocking Hit\nBoolean\n\nTrue if there was a blocking hit, false otherwise."
 		  Direction=EGPD_Output
@@ -879,7 +903,15 @@ Here is what is going on in this graph:
 		  DefaultValue="false"
 		  AutogeneratedDefaultValue="false"
 	   End Object
-	   Begin Object Name="EdGraphPin_1461"
+	   Begin Object Name="EdGraphPin_19971"
+		  PinName="bInitialOverlap"
+		  PinToolTip="Initial Overlap\nBoolean\n\nTrue if the hit started in an initial overlap. In this case some other values should be interpreted differently. Time will be 0, ImpactPoint will equal Location, and normals will be equal and indicate a depenetration vector."
+		  Direction=EGPD_Output
+		  PinType=(PinCategory="bool")
+		  DefaultValue="false"
+		  AutogeneratedDefaultValue="false"
+	   End Object
+	   Begin Object Name="EdGraphPin_19972"
 		  PinName="Time"
 		  PinToolTip="Time\nFloat\n\n\'Time\' of impact along trace direction ranging from [0.0 to 1.0) if there is a hit, indicating time between start and end. Equals 1.0 if there is no hit."
 		  Direction=EGPD_Output
@@ -887,59 +919,59 @@ Here is what is going on in this graph:
 		  DefaultValue="0.0"
 		  AutogeneratedDefaultValue="0.0"
 	   End Object
-	   Begin Object Name="EdGraphPin_1462"
+	   Begin Object Name="EdGraphPin_19973"
 		  PinName="Location"
 		  PinToolTip="Location\nVector \n\nLocation of the hit in world space. If this was a swept shape test, this is the location where we can place the shape in the world where it will not penetrate."
 		  Direction=EGPD_Output
-		  PinType=(PinCategory="struct",PinSubCategoryObject=ScriptStruct'/Script/CoreUObject.Object:Vector')
+		  PinType=(PinCategory="struct",PinSubCategoryObject=ScriptStruct'/Script/CoreUObject.Vector')
 		  DefaultValue="0, 0, 0"
 		  AutogeneratedDefaultValue="0, 0, 0"
 	   End Object
-	   Begin Object Name="EdGraphPin_1463"
-		  PinName="Normal"
-		  PinToolTip="Normal\nVector \n\nNormal of the hit in world space, for the object that was swept (e.g. for a sphere trace this points towards the sphere\'s center). Equal to ImpactNormal for line tests."
+	   Begin Object Name="EdGraphPin_19974"
+		  PinName="ImpactPoint"
+		  PinToolTip="Impact Point\nVector \n\nLocation of the actual contact point of the trace shape with the surface of the hit object. Equal to Location in the case of an initial overlap."
 		  Direction=EGPD_Output
-		  PinType=(PinCategory="struct",PinSubCategoryObject=ScriptStruct'/Script/CoreUObject.Object:Vector')
+		  PinType=(PinCategory="struct",PinSubCategoryObject=ScriptStruct'/Script/CoreUObject.Vector')
 		  DefaultValue="0, 0, 0"
 		  AutogeneratedDefaultValue="0, 0, 0"
 	   End Object
-	   Begin Object Name="EdGraphPin_1464"
-		  PinName="ImpactPoint"
-		  PinToolTip="Impact Point\nVector \n\nLocation of the actual contact point of the trace shape with the surface of the hit object."
+	   Begin Object Name="EdGraphPin_19975"
+		  PinName="Normal"
+		  PinToolTip="Normal\nVector \n\nNormal of the hit in world space, for the object that was swept (e.g. for a sphere trace this points towards the sphere\'s center). Equal to ImpactNormal for line tests."
 		  Direction=EGPD_Output
-		  PinType=(PinCategory="struct",PinSubCategoryObject=ScriptStruct'/Script/CoreUObject.Object:Vector')
+		  PinType=(PinCategory="struct",PinSubCategoryObject=ScriptStruct'/Script/CoreUObject.Vector')
 		  DefaultValue="0, 0, 0"
 		  AutogeneratedDefaultValue="0, 0, 0"
 	   End Object
-	   Begin Object Name="EdGraphPin_1465"
+	   Begin Object Name="EdGraphPin_19976"
 		  PinName="ImpactNormal"
 		  PinToolTip="Impact Normal\nVector \n\nNormal of the hit in world space, for the object that was hit by the sweep."
 		  Direction=EGPD_Output
-		  PinType=(PinCategory="struct",PinSubCategoryObject=ScriptStruct'/Script/CoreUObject.Object:Vector')
+		  PinType=(PinCategory="struct",PinSubCategoryObject=ScriptStruct'/Script/CoreUObject.Vector')
 		  DefaultValue="0, 0, 0"
 		  AutogeneratedDefaultValue="0, 0, 0"
 	   End Object
-	   Begin Object Name="EdGraphPin_1466"
+	   Begin Object Name="EdGraphPin_19977"
 		  PinName="PhysMat"
 		  PinToolTip="Phys Mat\nPhysical Material Reference\n\nPhysical material that was hit. Must set bReturnPhysicalMaterial to true in the query params for this to be returned."
 		  Direction=EGPD_Output
 		  PinType=(PinCategory="object",PinSubCategoryObject=Class'/Script/Engine.PhysicalMaterial')
 	   End Object
-	   Begin Object Name="EdGraphPin_1467"
+	   Begin Object Name="EdGraphPin_19978"
 		  PinName="HitActor"
 		  PinToolTip="Hit Actor\nActor Reference\n\nActor hit by the trace."
 		  Direction=EGPD_Output
 		  PinType=(PinCategory="object",PinSubCategoryObject=Class'/Script/Engine.Actor')
-		  LinkedTo(0)=EdGraphPin'K2Node_CallFunction_42644.EdGraphPin_1498'
-		  LinkedTo(1)=EdGraphPin'K2Node_Knot_7.EdGraphPin_20255'
+		  LinkedTo(0)=EdGraphPin'K2Node_CallFunction_2186.EdGraphPin_20017'
+		  LinkedTo(1)=EdGraphPin'K2Node_Knot_48.EdGraphPin_20238'
 	   End Object
-	   Begin Object Name="EdGraphPin_1468"
+	   Begin Object Name="EdGraphPin_19979"
 		  PinName="HitComponent"
 		  PinToolTip="Hit Component\nPrimitive Component Reference\n\nPrimitiveComponent hit by the trace."
 		  Direction=EGPD_Output
 		  PinType=(PinCategory="object",PinSubCategoryObject=Class'/Script/Engine.PrimitiveComponent')
 	   End Object
-	   Begin Object Name="EdGraphPin_1469"
+	   Begin Object Name="EdGraphPin_19980"
 		  PinName="HitBoneName"
 		  PinToolTip="Hit Bone Name\nName\n\nName of the bone hit (valid only if we hit a skeletal mesh)."
 		  Direction=EGPD_Output
@@ -947,7 +979,7 @@ Here is what is going on in this graph:
 		  DefaultValue="None"
 		  AutogeneratedDefaultValue="None"
 	   End Object
-	   Begin Object Name="EdGraphPin_1470"
+	   Begin Object Name="EdGraphPin_19981"
 		  PinName="HitItem"
 		  PinToolTip="Hit Item\nInteger\n\nPrimitive-specific data recording which item in the primitive was hit"
 		  Direction=EGPD_Output
@@ -955,35 +987,54 @@ Here is what is going on in this graph:
 		  DefaultValue="0"
 		  AutogeneratedDefaultValue="0"
 	   End Object
+	   Begin Object Name="EdGraphPin_19982"
+		  PinName="TraceStart"
+		  PinToolTip="Trace Start\nVector "
+		  Direction=EGPD_Output
+		  PinType=(PinCategory="struct",PinSubCategoryObject=ScriptStruct'/Script/CoreUObject.Vector')
+		  DefaultValue="0, 0, 0"
+		  AutogeneratedDefaultValue="0, 0, 0"
+	   End Object
+	   Begin Object Name="EdGraphPin_19983"
+		  PinName="TraceEnd"
+		  PinToolTip="Trace End\nVector "
+		  Direction=EGPD_Output
+		  PinType=(PinCategory="struct",PinSubCategoryObject=ScriptStruct'/Script/CoreUObject.Vector')
+		  DefaultValue="0, 0, 0"
+		  AutogeneratedDefaultValue="0, 0, 0"
+	   End Object
 	   bIsPureFunc=True
-	   FunctionReference=(MemberParentClass=Class'/Script/Engine.GameplayStatics',MemberName="BreakHitResult")
-	   Pins(0)=EdGraphPin'EdGraphPin_1458'
-	   Pins(1)=EdGraphPin'EdGraphPin_1459'
-	   Pins(2)=EdGraphPin'EdGraphPin_1460'
-	   Pins(3)=EdGraphPin'EdGraphPin_1461'
-	   Pins(4)=EdGraphPin'EdGraphPin_1462'
-	   Pins(5)=EdGraphPin'EdGraphPin_1463'
-	   Pins(6)=EdGraphPin'EdGraphPin_1464'
-	   Pins(7)=EdGraphPin'EdGraphPin_1465'
-	   Pins(8)=EdGraphPin'EdGraphPin_1466'
-	   Pins(9)=EdGraphPin'EdGraphPin_1467'
-	   Pins(10)=EdGraphPin'EdGraphPin_1468'
-	   Pins(11)=EdGraphPin'EdGraphPin_1469'
-	   Pins(12)=EdGraphPin'EdGraphPin_1470'
-	   NodePosX=3136
-	   NodePosY=624
-	   NodeGuid=53CCA2EF47E08514144AEC9C1042AC4C
+	   FunctionReference=(MemberParent=Class'/Script/Engine.GameplayStatics',MemberName="BreakHitResult")
+	   Pins(0)=EdGraphPin'EdGraphPin_19968'
+	   Pins(1)=EdGraphPin'EdGraphPin_19969'
+	   Pins(2)=EdGraphPin'EdGraphPin_19970'
+	   Pins(3)=EdGraphPin'EdGraphPin_19971'
+	   Pins(4)=EdGraphPin'EdGraphPin_19972'
+	   Pins(5)=EdGraphPin'EdGraphPin_19973'
+	   Pins(6)=EdGraphPin'EdGraphPin_19974'
+	   Pins(7)=EdGraphPin'EdGraphPin_19975'
+	   Pins(8)=EdGraphPin'EdGraphPin_19976'
+	   Pins(9)=EdGraphPin'EdGraphPin_19977'
+	   Pins(10)=EdGraphPin'EdGraphPin_19978'
+	   Pins(11)=EdGraphPin'EdGraphPin_19979'
+	   Pins(12)=EdGraphPin'EdGraphPin_19980'
+	   Pins(13)=EdGraphPin'EdGraphPin_19981'
+	   Pins(14)=EdGraphPin'EdGraphPin_19982'
+	   Pins(15)=EdGraphPin'EdGraphPin_19983'
+	   NodePosX=-976
+	   NodePosY=2544
+	   NodeGuid=1CDC0FC94EA121DC8807C393D14E7C95
 	End Object
-	Begin Object Class=K2Node_CallFunction Name="K2Node_CallFunction_42644"
-	   Begin Object Class=EdGraphPin Name="EdGraphPin_1497"
+	Begin Object Class=K2Node_CallFunction Name="K2Node_CallFunction_2186"
+	   Begin Object Class=EdGraphPin Name="EdGraphPin_20016"
 	   End Object
-	   Begin Object Class=EdGraphPin Name="EdGraphPin_1498"
+	   Begin Object Class=EdGraphPin Name="EdGraphPin_20017"
 	   End Object
-	   Begin Object Class=EdGraphPin Name="EdGraphPin_1499"
+	   Begin Object Class=EdGraphPin Name="EdGraphPin_20018"
 	   End Object
-	   Begin Object Class=EdGraphPin Name="EdGraphPin_1500"
+	   Begin Object Class=EdGraphPin Name="EdGraphPin_20019"
 	   End Object
-	   Begin Object Name="EdGraphPin_1497"
+	   Begin Object Name="EdGraphPin_20016"
 		  PinName="self"
 		  PinFriendlyName="Target"
 		  PinToolTip="Target\nKismet Math Library Reference"
@@ -991,236 +1042,220 @@ Here is what is going on in this graph:
 		  DefaultObject=Default__KismetMathLibrary
 		  bHidden=True
 	   End Object
-	   Begin Object Name="EdGraphPin_1498"
+	   Begin Object Name="EdGraphPin_20017"
 		  PinName="A"
 		  PinToolTip="A\nObject Reference"
 		  PinType=(PinCategory="object",PinSubCategoryObject=Class'/Script/CoreUObject.Object')
-		  LinkedTo(0)=EdGraphPin'K2Node_CallFunction_42262.EdGraphPin_1467'
+		  LinkedTo(0)=EdGraphPin'K2Node_CallFunction_2185.EdGraphPin_19978'
 	   End Object
-	   Begin Object Name="EdGraphPin_1499"
+	   Begin Object Name="EdGraphPin_20018"
 		  PinName="B"
 		  PinToolTip="B\nObject Reference"
 		  PinType=(PinCategory="object",PinSubCategoryObject=Class'/Script/CoreUObject.Object')
-		  LinkedTo(0)=EdGraphPin'K2Node_Knot_4.EdGraphPin_20250'
+		  LinkedTo(0)=EdGraphPin'K2Node_Knot_46.EdGraphPin_20227'
 	   End Object
-	   Begin Object Name="EdGraphPin_1500"
+	   Begin Object Name="EdGraphPin_20019"
 		  PinName="ReturnValue"
 		  PinToolTip="Return Value\nBoolean"
 		  Direction=EGPD_Output
 		  PinType=(PinCategory="bool")
 		  DefaultValue="false"
 		  AutogeneratedDefaultValue="false"
-		  LinkedTo(0)=EdGraphPin'K2Node_IfThenElse_229.EdGraphPin_1623'
+		  LinkedTo(0)=EdGraphPin'K2Node_IfThenElse_27.EdGraphPin_20143'
 	   End Object
 	   bIsPureFunc=True
-	   FunctionReference=(MemberParentClass=Class'/Script/Engine.KismetMathLibrary',MemberName="EqualEqual_ObjectObject")
-	   Pins(0)=EdGraphPin'EdGraphPin_1497'
-	   Pins(1)=EdGraphPin'EdGraphPin_1498'
-	   Pins(2)=EdGraphPin'EdGraphPin_1499'
-	   Pins(3)=EdGraphPin'EdGraphPin_1500'
-	   NodePosX=3376
-	   NodePosY=896
-	   NodeGuid=34EF9ADA4D017A60DF99789D7BF1A767
+	   FunctionReference=(MemberParent=Class'/Script/Engine.KismetMathLibrary',MemberName="EqualEqual_ObjectObject")
+	   Pins(0)=EdGraphPin'EdGraphPin_20016'
+	   Pins(1)=EdGraphPin'EdGraphPin_20017'
+	   Pins(2)=EdGraphPin'EdGraphPin_20018'
+	   Pins(3)=EdGraphPin'EdGraphPin_20019'
+	   NodePosX=-736
+	   NodePosY=2816
+	   NodeGuid=B98BAAFE42763E72A7C42CA3A7E8218F
 	End Object
-	Begin Object Class=K2Node_VariableSet Name="K2Node_VariableSet_149"
-	   Begin Object Class=EdGraphPin Name="EdGraphPin_1509"
+	Begin Object Class=K2Node_VariableSet Name="K2Node_VariableSet_11"
+	   Begin Object Class=EdGraphPin Name="EdGraphPin_20028"
 	   End Object
-	   Begin Object Class=EdGraphPin Name="EdGraphPin_1510"
+	   Begin Object Class=EdGraphPin Name="EdGraphPin_20029"
 	   End Object
-	   Begin Object Class=EdGraphPin Name="EdGraphPin_1511"
+	   Begin Object Class=EdGraphPin Name="EdGraphPin_20030"
 	   End Object
-	   Begin Object Class=EdGraphPin Name="EdGraphPin_1512"
-	   End Object
-	   Begin Object Class=EdGraphPin Name="EdGraphPin_1513"
-	   End Object
-	   Begin Object Name="EdGraphPin_1509"
+	   Begin Object Name="EdGraphPin_20028"
 		  PinName="execute"
 		  PinType=(PinCategory="exec")
-		  LinkedTo(0)=EdGraphPin'K2Node_DynamicCast_140.EdGraphPin_6204'
+		  LinkedTo(0)=EdGraphPin'K2Node_DynamicCast_437.EdGraphPin_19780'
 	   End Object
-	   Begin Object Name="EdGraphPin_1510"
+	   Begin Object Name="EdGraphPin_20029"
 		  PinName="then"
 		  Direction=EGPD_Output
 		  PinType=(PinCategory="exec")
-		  LinkedTo(0)=EdGraphPin'K2Node_CallFunction_834898.EdGraphPin_1766'
+		  LinkedTo(0)=EdGraphPin'K2Node_CallFunction_2192.EdGraphPin_20154'
 	   End Object
-	   Begin Object Name="EdGraphPin_1511"
+	   Begin Object Name="EdGraphPin_20030"
 		  PinName="AI_CON_Ref"
-		  PinType=(PinCategory="object",PinSubCategoryObject=BlueprintGeneratedClass'/Game/Blueprints/AI/Follower_AI_Con.Follower_AI_CON_C')
-		  LinkedTo(0)=EdGraphPin'K2Node_DynamicCast_140.EdGraphPin_6207'
-	   End Object
-	   Begin Object Name="EdGraphPin_1512"
-		  PinName="Output_Get"
-		  PinToolTip="Retrieves the value of the variable, can use instead of a separate Get node"
-		  Direction=EGPD_Output
-		  PinType=(PinCategory="object",PinSubCategoryObject=BlueprintGeneratedClass'/Game/Blueprints/AI/Follower_AI_Con.Follower_AI_CON_C')
-	   End Object
-	   Begin Object Name="EdGraphPin_1513"
-		  PinName="self"
-		  PinFriendlyName="Target"
-		  PinType=(PinCategory="object",PinSubCategoryObject=BlueprintGeneratedClass'/Game/Blueprints/AI/AgroCheck.AgroCheck_C')
-		  bHidden=True
+		  PinType=(PinCategory="object")
 	   End Object
 	   VariableReference=(MemberName="AI_CON_Ref",bSelfContext=True)
-	   Pins(0)=EdGraphPin'EdGraphPin_1509'
-	   Pins(1)=EdGraphPin'EdGraphPin_1510'
-	   Pins(2)=EdGraphPin'EdGraphPin_1511'
-	   Pins(3)=EdGraphPin'EdGraphPin_1512'
-	   Pins(4)=EdGraphPin'EdGraphPin_1513'
-	   NodePosX=384
-	   NodePosY=400
-	   NodeGuid=C40EE24F42575DEE754DBBBE157F0612
+	   Pins(0)=EdGraphPin'EdGraphPin_20028'
+	   Pins(1)=EdGraphPin'EdGraphPin_20029'
+	   Pins(2)=EdGraphPin'EdGraphPin_20030'
+	   NodePosX=-3728
+	   NodePosY=2320
+	   bHasCompilerMessage=True
+	   ErrorType=1
+	   ErrorMsg="Warning Could not find a variable named \"AI_CON_Ref\" in \'Blueprint_Effect_Explosion_C\'.\nMake sure \'Blueprint_Effect_Explosion_C\' has been compiled for  Set AI_CON_Ref \nError Default value \'\' for  AI CON Ref  is invalid: \'PinSubCategoryObject on pin \'AI_CON_Ref\' is NULL and PinSubCategory is \'\' not \'self\'\'"
+	   NodeGuid=9D269BA94474BDDEF9E15E9B0380623B
 	End Object
-	Begin Object Class=K2Node_MacroInstance Name="K2Node_MacroInstance_762"
-	   Begin Object Class=EdGraphPin Name="EdGraphPin_6221"
+	Begin Object Class=K2Node_MacroInstance Name="K2Node_MacroInstance_27"
+	   Begin Object Class=EdGraphPin Name="EdGraphPin_20037"
 	   End Object
-	   Begin Object Class=EdGraphPin Name="EdGraphPin_6222"
+	   Begin Object Class=EdGraphPin Name="EdGraphPin_20038"
 	   End Object
-	   Begin Object Class=EdGraphPin Name="EdGraphPin_6223"
+	   Begin Object Class=EdGraphPin Name="EdGraphPin_20039"
 	   End Object
-	   Begin Object Class=EdGraphPin Name="EdGraphPin_6224"
+	   Begin Object Class=EdGraphPin Name="EdGraphPin_20040"
 	   End Object
-	   Begin Object Class=EdGraphPin Name="EdGraphPin_6225"
+	   Begin Object Class=EdGraphPin Name="EdGraphPin_20041"
 	   End Object
-	   Begin Object Class=EdGraphPin Name="EdGraphPin_6226"
+	   Begin Object Class=EdGraphPin Name="EdGraphPin_20042"
 	   End Object
-	   Begin Object Class=EdGraphPin Name="EdGraphPin_6227"
+	   Begin Object Class=EdGraphPin Name="EdGraphPin_20043"
 	   End Object
-	   Begin Object Name="EdGraphPin_6221"
+	   Begin Object Name="EdGraphPin_20037"
 		  PinName="Exec"
 		  PinType=(PinCategory="exec")
-		  LinkedTo(0)=EdGraphPin'K2Node_CallFunction_39224.EdGraphPin_1282'
+		  LinkedTo(0)=EdGraphPin'K2Node_CallFunction_2182.EdGraphPin_19807'
 	   End Object
-	   Begin Object Name="EdGraphPin_6222"
+	   Begin Object Name="EdGraphPin_20038"
 		  PinName="Array"
-		  PinType=(PinCategory="struct",PinSubCategoryObject=ScriptStruct'/Script/Engine.EngineTypes:HitResult',bIsArray=True)
-		  LinkedTo(0)=EdGraphPin'K2Node_CallFunction_39224.EdGraphPin_1292'
+		  PinType=(PinCategory="struct",PinSubCategoryObject=ScriptStruct'/Script/Engine.HitResult',bIsArray=True)
+		  LinkedTo(0)=EdGraphPin'K2Node_CallFunction_2182.EdGraphPin_19817'
 	   End Object
-	   Begin Object Name="EdGraphPin_6223"
+	   Begin Object Name="EdGraphPin_20039"
 		  PinName="Break"
 		  PinType=(PinCategory="exec")
 	   End Object
-	   Begin Object Name="EdGraphPin_6224"
+	   Begin Object Name="EdGraphPin_20040"
 		  PinName="LoopBody"
 		  Direction=EGPD_Output
 		  PinType=(PinCategory="exec")
-		  LinkedTo(0)=EdGraphPin'K2Node_CallFunction_169714.EdGraphPin_1398'
+		  LinkedTo(0)=EdGraphPin'K2Node_CallFunction_2184.EdGraphPin_19914'
 	   End Object
-	   Begin Object Name="EdGraphPin_6225"
+	   Begin Object Name="EdGraphPin_20041"
 		  PinName="Array Element"
 		  Direction=EGPD_Output
-		  PinType=(PinCategory="struct",PinSubCategoryObject=ScriptStruct'/Script/Engine.EngineTypes:HitResult')
-		  LinkedTo(0)=EdGraphPin'K2Node_CallFunction_41995.EdGraphPin_1360'
+		  PinType=(PinCategory="struct",PinSubCategoryObject=ScriptStruct'/Script/Engine.HitResult')
+		  LinkedTo(0)=EdGraphPin'K2Node_CallFunction_2183.EdGraphPin_19867'
 	   End Object
-	   Begin Object Name="EdGraphPin_6226"
+	   Begin Object Name="EdGraphPin_20042"
 		  PinName="Array Index"
 		  Direction=EGPD_Output
 		  PinType=(PinCategory="int")
 	   End Object
-	   Begin Object Name="EdGraphPin_6227"
+	   Begin Object Name="EdGraphPin_20043"
 		  PinName="Completed"
 		  Direction=EGPD_Output
 		  PinType=(PinCategory="exec")
 	   End Object
 	   MacroGraphReference=(MacroGraph=EdGraph'/Engine/EditorBlueprintResources/StandardMacros.StandardMacros:ForEachLoopWithBreak',GraphBlueprint=Blueprint'/Engine/EditorBlueprintResources/StandardMacros.StandardMacros',GraphGuid=F07560274C5742E391E84B8F394CFB36)
-	   ResolvedWildcardType=(PinCategory="struct",PinSubCategoryObject=ScriptStruct'/Script/Engine.EngineTypes:HitResult',bIsArray=True)
-	   Pins(0)=EdGraphPin'EdGraphPin_6221'
-	   Pins(1)=EdGraphPin'EdGraphPin_6222'
-	   Pins(2)=EdGraphPin'EdGraphPin_6223'
-	   Pins(3)=EdGraphPin'EdGraphPin_6224'
-	   Pins(4)=EdGraphPin'EdGraphPin_6225'
-	   Pins(5)=EdGraphPin'EdGraphPin_6226'
-	   Pins(6)=EdGraphPin'EdGraphPin_6227'
-	   NodePosX=1888
-	   NodePosY=544
+	   ResolvedWildcardType=(PinCategory="struct",PinSubCategoryObject=ScriptStruct'/Script/Engine.HitResult',bIsArray=True)
+	   Pins(0)=EdGraphPin'EdGraphPin_20037'
+	   Pins(1)=EdGraphPin'EdGraphPin_20038'
+	   Pins(2)=EdGraphPin'EdGraphPin_20039'
+	   Pins(3)=EdGraphPin'EdGraphPin_20040'
+	   Pins(4)=EdGraphPin'EdGraphPin_20041'
+	   Pins(5)=EdGraphPin'EdGraphPin_20042'
+	   Pins(6)=EdGraphPin'EdGraphPin_20043'
+	   NodePosX=-2224
+	   NodePosY=2464
 	   ErrorType=1
-	   ErrorMsg="Error Unexpected node type K2Node_MacroInstance encountered at  ForEachLoopWithBreak \nError Unexpected node type K2Node_MacroInstance encountered at  ForEachLoopWithBreak "
-	   NodeGuid=51B389B348B8F74BC1C03C8A7744EEB2
+	   ErrorMsg="Error Unexpected node type K2Node_MacroInstance encountered at  ForEachLoopWithBreak "
+	   NodeGuid=97C212D24A5797040845AB802EE42670
 	End Object
-	Begin Object Class=K2Node_CallFunction Name="K2Node_CallFunction_92334"
-	   Begin Object Class=EdGraphPin Name="EdGraphPin_1542"
+	Begin Object Class=K2Node_CallFunction Name="K2Node_CallFunction_2187"
+	   Begin Object Class=EdGraphPin Name="EdGraphPin_20058"
 	   End Object
-	   Begin Object Class=EdGraphPin Name="EdGraphPin_1543"
+	   Begin Object Class=EdGraphPin Name="EdGraphPin_20059"
 	   End Object
-	   Begin Object Name="EdGraphPin_1542"
+	   Begin Object Name="EdGraphPin_20058"
 		  PinName="self"
 		  PinFriendlyName="Target"
 		  PinToolTip="Target\nActor Reference"
 		  PinType=(PinCategory="object",PinSubCategoryObject=Class'/Script/Engine.Actor')
-		  LinkedTo(0)=EdGraphPin'K2Node_CallFunction_41995.EdGraphPin_1368'
+		  LinkedTo(0)=EdGraphPin'K2Node_CallFunction_2183.EdGraphPin_19876'
 	   End Object
-	   Begin Object Name="EdGraphPin_1543"
+	   Begin Object Name="EdGraphPin_20059"
 		  PinName="ReturnValue"
 		  PinToolTip="Return Value\nVector "
 		  Direction=EGPD_Output
-		  PinType=(PinCategory="struct",PinSubCategoryObject=ScriptStruct'/Script/CoreUObject.Object:Vector')
+		  PinType=(PinCategory="struct",PinSubCategoryObject=ScriptStruct'/Script/CoreUObject.Vector')
 		  DefaultValue="0, 0, 0"
 		  AutogeneratedDefaultValue="0, 0, 0"
-		  LinkedTo(0)=EdGraphPin'K2Node_CallFunction_169714.EdGraphPin_1403'
+		  LinkedTo(0)=EdGraphPin'K2Node_CallFunction_2184.EdGraphPin_19919'
 	   End Object
 	   bIsPureFunc=True
 	   bIsConstFunc=True
-	   FunctionReference=(MemberParentClass=Class'/Script/Engine.Actor',MemberName="K2_GetActorLocation")
-	   Pins(0)=EdGraphPin'EdGraphPin_1542'
-	   Pins(1)=EdGraphPin'EdGraphPin_1543'
-	   NodePosX=2416
-	   NodePosY=688
-	   NodeGuid=30F7DC8B49495736E721A09D0A7B8A21
+	   FunctionReference=(MemberParent=Class'/Script/Engine.Actor',MemberName="K2_GetActorLocation")
+	   Pins(0)=EdGraphPin'EdGraphPin_20058'
+	   Pins(1)=EdGraphPin'EdGraphPin_20059'
+	   NodePosX=-1696
+	   NodePosY=2608
+	   NodeGuid=E6C4DC04431CACDB9F1458BD922DB0C1
 	End Object
-	Begin Object Class=K2Node_CallFunction Name="K2Node_CallFunction_103176"
-	   Begin Object Class=EdGraphPin Name="EdGraphPin_1548"
+	Begin Object Class=K2Node_CallFunction Name="K2Node_CallFunction_2188"
+	   Begin Object Class=EdGraphPin Name="EdGraphPin_20064"
 	   End Object
-	   Begin Object Class=EdGraphPin Name="EdGraphPin_1549"
+	   Begin Object Class=EdGraphPin Name="EdGraphPin_20065"
 	   End Object
-	   Begin Object Name="EdGraphPin_1548"
+	   Begin Object Name="EdGraphPin_20064"
 		  PinName="self"
 		  PinFriendlyName="Target"
 		  PinToolTip="Target\nActor Reference"
 		  PinType=(PinCategory="object",PinSubCategoryObject=Class'/Script/Engine.Actor')
-		  LinkedTo(0)=EdGraphPin'K2Node_Knot_5.EdGraphPin_20252'
+		  LinkedTo(0)=EdGraphPin'K2Node_Knot_47.EdGraphPin_20233'
 	   End Object
-	   Begin Object Name="EdGraphPin_1549"
+	   Begin Object Name="EdGraphPin_20065"
 		  PinName="ReturnValue"
 		  PinToolTip="Return Value\nVector "
 		  Direction=EGPD_Output
-		  PinType=(PinCategory="struct",PinSubCategoryObject=ScriptStruct'/Script/CoreUObject.Object:Vector')
+		  PinType=(PinCategory="struct",PinSubCategoryObject=ScriptStruct'/Script/CoreUObject.Vector')
 		  DefaultValue="0, 0, 0"
 		  AutogeneratedDefaultValue="0, 0, 0"
-		  LinkedTo(0)=EdGraphPin'K2Node_CallFunction_94128.EdGraphPin_1559'
+		  LinkedTo(0)=EdGraphPin'K2Node_CallFunction_2189.EdGraphPin_20075'
 	   End Object
 	   bIsPureFunc=True
 	   bIsConstFunc=True
-	   FunctionReference=(MemberParentClass=Class'/Script/Engine.Actor',MemberName="K2_GetActorLocation")
-	   Pins(0)=EdGraphPin'EdGraphPin_1548'
-	   Pins(1)=EdGraphPin'EdGraphPin_1549'
-	   NodePosX=4256
-	   NodePosY=688
-	   NodeGuid=420D0BAB48F8161E90889CB9F92CD629
+	   FunctionReference=(MemberParent=Class'/Script/Engine.Actor',MemberName="K2_GetActorLocation")
+	   Pins(0)=EdGraphPin'EdGraphPin_20064'
+	   Pins(1)=EdGraphPin'EdGraphPin_20065'
+	   NodePosX=144
+	   NodePosY=2608
+	   NodeGuid=E876A015408E71B02F94C991ECF71A93
 	End Object
-	Begin Object Class=K2Node_CallFunction Name="K2Node_CallFunction_94128"
-	   Begin Object Class=EdGraphPin Name="EdGraphPin_1554"
+	Begin Object Class=K2Node_CallFunction Name="K2Node_CallFunction_2189"
+	   Begin Object Class=EdGraphPin Name="EdGraphPin_20070"
 	   End Object
-	   Begin Object Class=EdGraphPin Name="EdGraphPin_1555"
+	   Begin Object Class=EdGraphPin Name="EdGraphPin_20071"
 	   End Object
-	   Begin Object Class=EdGraphPin Name="EdGraphPin_1556"
+	   Begin Object Class=EdGraphPin Name="EdGraphPin_20072"
 	   End Object
-	   Begin Object Class=EdGraphPin Name="EdGraphPin_1557"
+	   Begin Object Class=EdGraphPin Name="EdGraphPin_20073"
 	   End Object
-	   Begin Object Class=EdGraphPin Name="EdGraphPin_1558"
+	   Begin Object Class=EdGraphPin Name="EdGraphPin_20074"
 	   End Object
-	   Begin Object Class=EdGraphPin Name="EdGraphPin_1559"
+	   Begin Object Class=EdGraphPin Name="EdGraphPin_20075"
 	   End Object
-	   Begin Object Name="EdGraphPin_1554"
+	   Begin Object Name="EdGraphPin_20070"
 		  PinName="execute"
 		  PinType=(PinCategory="exec")
-		  LinkedTo(0)=EdGraphPin'K2Node_CallFunction_103002.EdGraphPin_1525'
+		  LinkedTo(0)=EdGraphPin'K2Node_CallFunction_2190.EdGraphPin_20107'
 	   End Object
-	   Begin Object Name="EdGraphPin_1555"
+	   Begin Object Name="EdGraphPin_20071"
 		  PinName="then"
 		  Direction=EGPD_Output
 		  PinType=(PinCategory="exec")
 	   End Object
-	   Begin Object Name="EdGraphPin_1556"
+	   Begin Object Name="EdGraphPin_20072"
 		  PinName="self"
 		  PinFriendlyName="Target"
 		  PinToolTip="Target\nBTFunction Library Reference"
@@ -1228,137 +1263,147 @@ Here is what is going on in this graph:
 		  DefaultObject=Default__BTFunctionLibrary
 		  bHidden=True
 	   End Object
-	   Begin Object Name="EdGraphPin_1557"
+	   Begin Object Name="EdGraphPin_20073"
 		  PinName="NodeOwner"
 		  PinToolTip="Node Owner\nBTNode Reference"
 		  PinType=(PinCategory="object",PinSubCategoryObject=Class'/Script/AIModule.BTNode')
 		  bHidden=True
 	   End Object
-	   Begin Object Name="EdGraphPin_1558"
+	   Begin Object Name="EdGraphPin_20074"
 		  PinName="Key"
 		  PinToolTip="Key\nBlackboard Key Selector Structure (by ref)"
-		  PinType=(PinCategory="struct",PinSubCategoryObject=ScriptStruct'/Script/AIModule.BehaviorTreeTypes:BlackboardKeySelector',bIsReference=True,bIsConst=True)
-		  LinkedTo(0)=EdGraphPin'K2Node_VariableGet_266.EdGraphPin_1572'
+		  PinType=(PinCategory="struct",PinSubCategoryObject=ScriptStruct'/Script/AIModule.BlackboardKeySelector',bIsReference=True,bIsConst=True)
+		  LinkedTo(0)=EdGraphPin'K2Node_VariableGet_37.EdGraphPin_20088'
 		  bDefaultValueIsIgnored=True
 	   End Object
-	   Begin Object Name="EdGraphPin_1559"
+	   Begin Object Name="EdGraphPin_20075"
 		  PinName="Value"
 		  PinToolTip="Value\nVector "
-		  PinType=(PinCategory="struct",PinSubCategoryObject=ScriptStruct'/Script/CoreUObject.Object:Vector')
+		  PinType=(PinCategory="struct",PinSubCategoryObject=ScriptStruct'/Script/CoreUObject.Vector')
 		  DefaultValue="0, 0, 0"
 		  AutogeneratedDefaultValue="0, 0, 0"
-		  LinkedTo(0)=EdGraphPin'K2Node_CallFunction_103176.EdGraphPin_1549'
-	   End Object
-	   FunctionReference=(MemberParentClass=Class'/Script/AIModule.BTFunctionLibrary',MemberName="SetBlackboardValueAsVector")
-	   Pins(0)=EdGraphPin'EdGraphPin_1554'
-	   Pins(1)=EdGraphPin'EdGraphPin_1555'
-	   Pins(2)=EdGraphPin'EdGraphPin_1556'
-	   Pins(3)=EdGraphPin'EdGraphPin_1557'
-	   Pins(4)=EdGraphPin'EdGraphPin_1558'
-	   Pins(5)=EdGraphPin'EdGraphPin_1559'
-	   NodePosX=4704
-	   NodePosY=544
+		  LinkedTo(0)=EdGraphPin'K2Node_CallFunction_2188.EdGraphPin_20065'
+	   End Object
+	   FunctionReference=(MemberParent=Class'/Script/AIModule.BTFunctionLibrary',MemberName="SetBlackboardValueAsVector")
+	   Pins(0)=EdGraphPin'EdGraphPin_20070'
+	   Pins(1)=EdGraphPin'EdGraphPin_20071'
+	   Pins(2)=EdGraphPin'EdGraphPin_20072'
+	   Pins(3)=EdGraphPin'EdGraphPin_20073'
+	   Pins(4)=EdGraphPin'EdGraphPin_20074'
+	   Pins(5)=EdGraphPin'EdGraphPin_20075'
+	   NodePosX=592
+	   NodePosY=2464
+	   bHasCompilerMessage=True
 	   ErrorType=1
-	   ErrorMsg="Error Default value \'\' for  Key  is invalid: \'Literal values are not allowed for pass-by-reference parameters.\'"
-	   NodeGuid=ED0B02E34401D4B635C2A8A32B91F817
+	   ErrorMsg="Error This blueprint (self) is not a BTNode, therefore \' Node Owner \' must have a connection."
+	   NodeGuid=826018174CB60380AD257A8D0AD047AF
 	End Object
-	Begin Object Class=K2Node_VariableGet Name="K2Node_VariableGet_266"
-	   Begin Object Class=EdGraphPin Name="EdGraphPin_1572"
+	Begin Object Class=K2Node_VariableGet Name="K2Node_VariableGet_37"
+	   Begin Object Class=EdGraphPin Name="EdGraphPin_20088"
 	   End Object
-	   Begin Object Class=EdGraphPin Name="EdGraphPin_1573"
+	   Begin Object Class=EdGraphPin Name="EdGraphPin_20089"
 	   End Object
-	   Begin Object Name="EdGraphPin_1572"
+	   Begin Object Name="EdGraphPin_20088"
 		  PinName="TargetLocation"
 		  Direction=EGPD_Output
-		  PinType=(PinCategory="struct",PinSubCategoryObject=ScriptStruct'/Script/AIModule.BehaviorTreeTypes:BlackboardKeySelector')
-		  LinkedTo(0)=EdGraphPin'K2Node_CallFunction_94128.EdGraphPin_1558'
+		  PinType=(PinCategory="struct",PinSubCategoryObject=ScriptStruct'/Script/AIModule.BlackboardKeySelector')
+		  LinkedTo(0)=EdGraphPin'K2Node_CallFunction_2189.EdGraphPin_20074'
 	   End Object
-	   Begin Object Name="EdGraphPin_1573"
+	   Begin Object Name="EdGraphPin_20089"
 		  PinName="self"
 		  PinFriendlyName="Target"
-		  PinType=(PinCategory="object",PinSubCategoryObject=BlueprintGeneratedClass'/Game/Blueprints/AI/AgroCheck.AgroCheck_C')
+		  PinType=(PinCategory="object",PinSubCategoryObject=BlueprintGeneratedClass'/Game/StarterContent/Blueprints/Blueprint_Effect_Explosion.Blueprint_Effect_Explosion_C')
 		  bHidden=True
 	   End Object
 	   VariableReference=(MemberName="TargetLocation",bSelfContext=True)
-	   Pins(0)=EdGraphPin'EdGraphPin_1572'
-	   Pins(1)=EdGraphPin'EdGraphPin_1573'
-	   NodePosX=4512
-	   NodePosY=608
-	   NodeGuid=E34808B44AD9E4B46EF291AC372AF463
+	   Pins(0)=EdGraphPin'EdGraphPin_20088'
+	   Pins(1)=EdGraphPin'EdGraphPin_20089'
+	   NodePosX=400
+	   NodePosY=2528
+	   bHasCompilerMessage=True
+	   ErrorType=1
+	   ErrorMsg="Warning Could not find a variable named \"TargetLocation\" in \'Blueprint_Effect_Explosion_C\'.\nMake sure \'Blueprint_Effect_Explosion_C\' has been compiled for  Get TargetLocation \nError The property associated with  Target Location  could not be found"
+	   NodeGuid=8C8EE31C44973580F96EA28F1F10CD6A
 	End Object
-	Begin Object Class=K2Node_VariableGet Name="K2Node_VariableGet_303"
-	   Begin Object Class=EdGraphPin Name="EdGraphPin_1578"
+	Begin Object Class=K2Node_VariableGet Name="K2Node_VariableGet_38"
+	   Begin Object Class=EdGraphPin Name="EdGraphPin_20094"
 	   End Object
-	   Begin Object Class=EdGraphPin Name="EdGraphPin_1579"
+	   Begin Object Class=EdGraphPin Name="EdGraphPin_20095"
 	   End Object
-	   Begin Object Name="EdGraphPin_1578"
+	   Begin Object Name="EdGraphPin_20094"
 		  PinName="TargetToFollow"
 		  Direction=EGPD_Output
-		  PinType=(PinCategory="struct",PinSubCategoryObject=ScriptStruct'/Script/AIModule.BehaviorTreeTypes:BlackboardKeySelector')
-		  LinkedTo(0)=EdGraphPin'K2Node_CallFunction_103001.EdGraphPin_1608'
+		  PinType=(PinCategory="struct",PinSubCategoryObject=ScriptStruct'/Script/AIModule.BlackboardKeySelector')
+		  LinkedTo(0)=EdGraphPin'K2Node_CallFunction_2191.EdGraphPin_20128'
 	   End Object
-	   Begin Object Name="EdGraphPin_1579"
+	   Begin Object Name="EdGraphPin_20095"
 		  PinName="self"
 		  PinFriendlyName="Target"
-		  PinType=(PinCategory="object",PinSubCategoryObject=BlueprintGeneratedClass'/Game/Blueprints/AI/AgroCheck.AgroCheck_C')
+		  PinType=(PinCategory="object",PinSubCategoryObject=BlueprintGeneratedClass'/Game/StarterContent/Blueprints/Blueprint_Effect_Explosion.Blueprint_Effect_Explosion_C')
 		  bHidden=True
 	   End Object
 	   VariableReference=(MemberName="TargetToFollow",bSelfContext=True)
-	   Pins(0)=EdGraphPin'EdGraphPin_1578'
-	   Pins(1)=EdGraphPin'EdGraphPin_1579'
-	   NodePosX=3856
-	   NodePosY=880
-	   NodeGuid=7E7FA7C3479CB8466905F28714E7BFF6
+	   Pins(0)=EdGraphPin'EdGraphPin_20094'
+	   Pins(1)=EdGraphPin'EdGraphPin_20095'
+	   NodePosX=-256
+	   NodePosY=2800
+	   bHasCompilerMessage=True
+	   ErrorType=1
+	   ErrorMsg="Warning Could not find a variable named \"TargetToFollow\" in \'Blueprint_Effect_Explosion_C\'.\nMake sure \'Blueprint_Effect_Explosion_C\' has been compiled for  Get TargetToFollow \nError The property associated with  Target to Follow  could not be found"
+	   NodeGuid=61F581D441583D010D275F97173A007B
 	End Object
-	Begin Object Class=K2Node_VariableGet Name="K2Node_VariableGet_305"
-	   Begin Object Class=EdGraphPin Name="EdGraphPin_1584"
+	Begin Object Class=K2Node_VariableGet Name="K2Node_VariableGet_39"
+	   Begin Object Class=EdGraphPin Name="EdGraphPin_20100"
 	   End Object
-	   Begin Object Class=EdGraphPin Name="EdGraphPin_1585"
+	   Begin Object Class=EdGraphPin Name="EdGraphPin_20101"
 	   End Object
-	   Begin Object Name="EdGraphPin_1584"
+	   Begin Object Name="EdGraphPin_20100"
 		  PinName="TargetToFollow"
 		  Direction=EGPD_Output
-		  PinType=(PinCategory="struct",PinSubCategoryObject=ScriptStruct'/Script/AIModule.BehaviorTreeTypes:BlackboardKeySelector')
-		  LinkedTo(0)=EdGraphPin'K2Node_CallFunction_103002.EdGraphPin_1528'
+		  PinType=(PinCategory="struct",PinSubCategoryObject=ScriptStruct'/Script/AIModule.BlackboardKeySelector')
+		  LinkedTo(0)=EdGraphPin'K2Node_CallFunction_2190.EdGraphPin_20110'
 	   End Object
-	   Begin Object Name="EdGraphPin_1585"
+	   Begin Object Name="EdGraphPin_20101"
 		  PinName="self"
 		  PinFriendlyName="Target"
-		  PinType=(PinCategory="object",PinSubCategoryObject=BlueprintGeneratedClass'/Game/Blueprints/AI/AgroCheck.AgroCheck_C')
+		  PinType=(PinCategory="object",PinSubCategoryObject=BlueprintGeneratedClass'/Game/StarterContent/Blueprints/Blueprint_Effect_Explosion.Blueprint_Effect_Explosion_C')
 		  bHidden=True
 	   End Object
 	   VariableReference=(MemberName="TargetToFollow",bSelfContext=True)
-	   Pins(0)=EdGraphPin'EdGraphPin_1584'
-	   Pins(1)=EdGraphPin'EdGraphPin_1585'
-	   NodePosX=4048
-	   NodePosY=608
-	   NodeGuid=0C6B88654A6AE0188ED7F4B9B10DBDF6
+	   Pins(0)=EdGraphPin'EdGraphPin_20100'
+	   Pins(1)=EdGraphPin'EdGraphPin_20101'
+	   NodePosX=-64
+	   NodePosY=2528
+	   bHasCompilerMessage=True
+	   ErrorType=1
+	   ErrorMsg="Warning Could not find a variable named \"TargetToFollow\" in \'Blueprint_Effect_Explosion_C\'.\nMake sure \'Blueprint_Effect_Explosion_C\' has been compiled for  Get TargetToFollow \nError The property associated with  Target to Follow  could not be found"
+	   NodeGuid=841C4D864358568539E28785FC6C4D61
 	End Object
-	Begin Object Class=K2Node_CallFunction Name="K2Node_CallFunction_103002"
-	   Begin Object Class=EdGraphPin Name="EdGraphPin_1524"
+	Begin Object Class=K2Node_CallFunction Name="K2Node_CallFunction_2190"
+	   Begin Object Class=EdGraphPin Name="EdGraphPin_20106"
 	   End Object
-	   Begin Object Class=EdGraphPin Name="EdGraphPin_1525"
+	   Begin Object Class=EdGraphPin Name="EdGraphPin_20107"
 	   End Object
-	   Begin Object Class=EdGraphPin Name="EdGraphPin_1526"
+	   Begin Object Class=EdGraphPin Name="EdGraphPin_20108"
 	   End Object
-	   Begin Object Class=EdGraphPin Name="EdGraphPin_1527"
+	   Begin Object Class=EdGraphPin Name="EdGraphPin_20109"
 	   End Object
-	   Begin Object Class=EdGraphPin Name="EdGraphPin_1528"
+	   Begin Object Class=EdGraphPin Name="EdGraphPin_20110"
 	   End Object
-	   Begin Object Class=EdGraphPin Name="EdGraphPin_1529"
+	   Begin Object Class=EdGraphPin Name="EdGraphPin_20111"
 	   End Object
-	   Begin Object Name="EdGraphPin_1524"
+	   Begin Object Name="EdGraphPin_20106"
 		  PinName="execute"
 		  PinType=(PinCategory="exec")
-		  LinkedTo(0)=EdGraphPin'K2Node_IfThenElse_229.EdGraphPin_1624'
+		  LinkedTo(0)=EdGraphPin'K2Node_IfThenElse_27.EdGraphPin_20144'
 	   End Object
-	   Begin Object Name="EdGraphPin_1525"
+	   Begin Object Name="EdGraphPin_20107"
 		  PinName="then"
 		  Direction=EGPD_Output
 		  PinType=(PinCategory="exec")
-		  LinkedTo(0)=EdGraphPin'K2Node_CallFunction_94128.EdGraphPin_1554'
+		  LinkedTo(0)=EdGraphPin'K2Node_CallFunction_2189.EdGraphPin_20070'
 	   End Object
-	   Begin Object Name="EdGraphPin_1526"
+	   Begin Object Name="EdGraphPin_20108"
 		  PinName="self"
 		  PinFriendlyName="Target"
 		  PinToolTip="Target\nBTFunction Library Reference"
@@ -1366,60 +1411,63 @@ Here is what is going on in this graph:
 		  DefaultObject=Default__BTFunctionLibrary
 		  bHidden=True
 	   End Object
-	   Begin Object Name="EdGraphPin_1527"
+	   Begin Object Name="EdGraphPin_20109"
 		  PinName="NodeOwner"
 		  PinToolTip="Node Owner\nBTNode Reference"
 		  PinType=(PinCategory="object",PinSubCategoryObject=Class'/Script/AIModule.BTNode')
 		  bHidden=True
 	   End Object
-	   Begin Object Name="EdGraphPin_1528"
+	   Begin Object Name="EdGraphPin_20110"
 		  PinName="Key"
 		  PinToolTip="Key\nBlackboard Key Selector Structure (by ref)"
-		  PinType=(PinCategory="struct",PinSubCategoryObject=ScriptStruct'/Script/AIModule.BehaviorTreeTypes:BlackboardKeySelector',bIsReference=True,bIsConst=True)
-		  LinkedTo(0)=EdGraphPin'K2Node_VariableGet_305.EdGraphPin_1584'
+		  PinType=(PinCategory="struct",PinSubCategoryObject=ScriptStruct'/Script/AIModule.BlackboardKeySelector',bIsReference=True,bIsConst=True)
+		  LinkedTo(0)=EdGraphPin'K2Node_VariableGet_39.EdGraphPin_20100'
 		  bDefaultValueIsIgnored=True
 	   End Object
-	   Begin Object Name="EdGraphPin_1529"
+	   Begin Object Name="EdGraphPin_20111"
 		  PinName="Value"
 		  PinToolTip="Value\nObject Reference"
 		  PinType=(PinCategory="object",PinSubCategoryObject=Class'/Script/CoreUObject.Object')
-		  LinkedTo(0)=EdGraphPin'K2Node_Knot_5.EdGraphPin_20252'
-	   End Object
-	   FunctionReference=(MemberParentClass=Class'/Script/AIModule.BTFunctionLibrary',MemberName="SetBlackboardValueAsObject")
-	   Pins(0)=EdGraphPin'EdGraphPin_1524'
-	   Pins(1)=EdGraphPin'EdGraphPin_1525'
-	   Pins(2)=EdGraphPin'EdGraphPin_1526'
-	   Pins(3)=EdGraphPin'EdGraphPin_1527'
-	   Pins(4)=EdGraphPin'EdGraphPin_1528'
-	   Pins(5)=EdGraphPin'EdGraphPin_1529'
-	   NodePosX=4240
-	   NodePosY=544
-	   NodeGuid=673B107C4A74B4F348405090F18BAC71
+		  LinkedTo(0)=EdGraphPin'K2Node_Knot_47.EdGraphPin_20233'
+	   End Object
+	   FunctionReference=(MemberParent=Class'/Script/AIModule.BTFunctionLibrary',MemberName="SetBlackboardValueAsObject")
+	   Pins(0)=EdGraphPin'EdGraphPin_20106'
+	   Pins(1)=EdGraphPin'EdGraphPin_20107'
+	   Pins(2)=EdGraphPin'EdGraphPin_20108'
+	   Pins(3)=EdGraphPin'EdGraphPin_20109'
+	   Pins(4)=EdGraphPin'EdGraphPin_20110'
+	   Pins(5)=EdGraphPin'EdGraphPin_20111'
+	   NodePosX=128
+	   NodePosY=2464
+	   bHasCompilerMessage=True
+	   ErrorType=1
+	   ErrorMsg="Error This blueprint (self) is not a BTNode, therefore \' Node Owner \' must have a connection."
+	   NodeGuid=277F86FC42A73836499F53AD67650323
 	End Object
-	Begin Object Class=K2Node_CallFunction Name="K2Node_CallFunction_103001"
-	   Begin Object Class=EdGraphPin Name="EdGraphPin_1604"
+	Begin Object Class=K2Node_CallFunction Name="K2Node_CallFunction_2191"
+	   Begin Object Class=EdGraphPin Name="EdGraphPin_20124"
 	   End Object
-	   Begin Object Class=EdGraphPin Name="EdGraphPin_1605"
+	   Begin Object Class=EdGraphPin Name="EdGraphPin_20125"
 	   End Object
-	   Begin Object Class=EdGraphPin Name="EdGraphPin_1606"
+	   Begin Object Class=EdGraphPin Name="EdGraphPin_20126"
 	   End Object
-	   Begin Object Class=EdGraphPin Name="EdGraphPin_1607"
+	   Begin Object Class=EdGraphPin Name="EdGraphPin_20127"
 	   End Object
-	   Begin Object Class=EdGraphPin Name="EdGraphPin_1608"
+	   Begin Object Class=EdGraphPin Name="EdGraphPin_20128"
 	   End Object
-	   Begin Object Class=EdGraphPin Name="EdGraphPin_1609"
+	   Begin Object Class=EdGraphPin Name="EdGraphPin_20129"
 	   End Object
-	   Begin Object Name="EdGraphPin_1604"
+	   Begin Object Name="EdGraphPin_20124"
 		  PinName="execute"
 		  PinType=(PinCategory="exec")
-		  LinkedTo(0)=EdGraphPin'K2Node_IfThenElse_229.EdGraphPin_1625'
+		  LinkedTo(0)=EdGraphPin'K2Node_IfThenElse_27.EdGraphPin_20145'
 	   End Object
-	   Begin Object Name="EdGraphPin_1605"
+	   Begin Object Name="EdGraphPin_20125"
 		  PinName="then"
 		  Direction=EGPD_Output
 		  PinType=(PinCategory="exec")
 	   End Object
-	   Begin Object Name="EdGraphPin_1606"
+	   Begin Object Name="EdGraphPin_20126"
 		  PinName="self"
 		  PinFriendlyName="Target"
 		  PinToolTip="Target\nBTFunction Library Reference"
@@ -1427,104 +1475,107 @@ Here is what is going on in this graph:
 		  DefaultObject=Default__BTFunctionLibrary
 		  bHidden=True
 	   End Object
-	   Begin Object Name="EdGraphPin_1607"
+	   Begin Object Name="EdGraphPin_20127"
 		  PinName="NodeOwner"
 		  PinToolTip="Node Owner\nBTNode Reference"
 		  PinType=(PinCategory="object",PinSubCategoryObject=Class'/Script/AIModule.BTNode')
 		  bHidden=True
 	   End Object
-	   Begin Object Name="EdGraphPin_1608"
+	   Begin Object Name="EdGraphPin_20128"
 		  PinName="Key"
 		  PinToolTip="Key\nBlackboard Key Selector Structure (by ref)"
-		  PinType=(PinCategory="struct",PinSubCategoryObject=ScriptStruct'/Script/AIModule.BehaviorTreeTypes:BlackboardKeySelector',bIsReference=True,bIsConst=True)
-		  LinkedTo(0)=EdGraphPin'K2Node_VariableGet_303.EdGraphPin_1578'
+		  PinType=(PinCategory="struct",PinSubCategoryObject=ScriptStruct'/Script/AIModule.BlackboardKeySelector',bIsReference=True,bIsConst=True)
+		  LinkedTo(0)=EdGraphPin'K2Node_VariableGet_38.EdGraphPin_20094'
 		  bDefaultValueIsIgnored=True
 	   End Object
-	   Begin Object Name="EdGraphPin_1609"
+	   Begin Object Name="EdGraphPin_20129"
 		  PinName="Value"
 		  PinToolTip="Value\nObject Reference"
 		  PinType=(PinCategory="object",PinSubCategoryObject=Class'/Script/CoreUObject.Object')
 	   End Object
-	   FunctionReference=(MemberParentClass=Class'/Script/AIModule.BTFunctionLibrary',MemberName="SetBlackboardValueAsObject")
-	   Pins(0)=EdGraphPin'EdGraphPin_1604'
-	   Pins(1)=EdGraphPin'EdGraphPin_1605'
-	   Pins(2)=EdGraphPin'EdGraphPin_1606'
-	   Pins(3)=EdGraphPin'EdGraphPin_1607'
-	   Pins(4)=EdGraphPin'EdGraphPin_1608'
-	   Pins(5)=EdGraphPin'EdGraphPin_1609'
-	   NodePosX=4048
-	   NodePosY=816
-	   NodeGuid=48A427E74B008F9733459FB0C3F53FA7
+	   FunctionReference=(MemberParent=Class'/Script/AIModule.BTFunctionLibrary',MemberName="SetBlackboardValueAsObject")
+	   Pins(0)=EdGraphPin'EdGraphPin_20124'
+	   Pins(1)=EdGraphPin'EdGraphPin_20125'
+	   Pins(2)=EdGraphPin'EdGraphPin_20126'
+	   Pins(3)=EdGraphPin'EdGraphPin_20127'
+	   Pins(4)=EdGraphPin'EdGraphPin_20128'
+	   Pins(5)=EdGraphPin'EdGraphPin_20129'
+	   NodePosX=-64
+	   NodePosY=2736
+	   bHasCompilerMessage=True
+	   ErrorType=1
+	   ErrorMsg="Error This blueprint (self) is not a BTNode, therefore \' Node Owner \' must have a connection."
+	   NodeGuid=15590FFF47F60A3322F082AEE4DA84FF
 	End Object
-	Begin Object Class=K2Node_IfThenElse Name="K2Node_IfThenElse_229"
-	   Begin Object Class=EdGraphPin Name="EdGraphPin_1622"
+	Begin Object Class=K2Node_IfThenElse Name="K2Node_IfThenElse_27"
+	   Begin Object Class=EdGraphPin Name="EdGraphPin_20142"
 	   End Object
-	   Begin Object Class=EdGraphPin Name="EdGraphPin_1623"
+	   Begin Object Class=EdGraphPin Name="EdGraphPin_20143"
 	   End Object
-	   Begin Object Class=EdGraphPin Name="EdGraphPin_1624"
+	   Begin Object Class=EdGraphPin Name="EdGraphPin_20144"
 	   End Object
-	   Begin Object Class=EdGraphPin Name="EdGraphPin_1625"
+	   Begin Object Class=EdGraphPin Name="EdGraphPin_20145"
 	   End Object
-	   Begin Object Name="EdGraphPin_1622"
+	   Begin Object Name="EdGraphPin_20142"
 		  PinName="execute"
 		  PinType=(PinCategory="exec")
-		  LinkedTo(0)=EdGraphPin'K2Node_CallFunction_169714.EdGraphPin_1399'
+		  LinkedTo(0)=EdGraphPin'K2Node_CallFunction_2184.EdGraphPin_19915'
 	   End Object
-	   Begin Object Name="EdGraphPin_1623"
+	   Begin Object Name="EdGraphPin_20143"
 		  PinName="Condition"
 		  PinType=(PinCategory="bool")
-		  DefaultValue="false"
-		  AutogeneratedDefaultValue="false"
-		  LinkedTo(0)=EdGraphPin'K2Node_CallFunction_42644.EdGraphPin_1500'
+		  DefaultValue="true"
+		  AutogeneratedDefaultValue="true"
+		  LinkedTo(0)=EdGraphPin'K2Node_CallFunction_2186.EdGraphPin_20019'
 	   End Object
-	   Begin Object Name="EdGraphPin_1624"
+	   Begin Object Name="EdGraphPin_20144"
 		  PinName="then"
 		  PinFriendlyName="true"
 		  Direction=EGPD_Output
 		  PinType=(PinCategory="exec")
-		  LinkedTo(0)=EdGraphPin'K2Node_CallFunction_103002.EdGraphPin_1524'
+		  LinkedTo(0)=EdGraphPin'K2Node_CallFunction_2190.EdGraphPin_20106'
 	   End Object
-	   Begin Object Name="EdGraphPin_1625"
+	   Begin Object Name="EdGraphPin_20145"
 		  PinName="else"
 		  PinFriendlyName="false"
 		  Direction=EGPD_Output
 		  PinType=(PinCategory="exec")
-		  LinkedTo(0)=EdGraphPin'K2Node_CallFunction_103001.EdGraphPin_1604'
-	   End Object
-	   Pins(0)=EdGraphPin'EdGraphPin_1622'
-	   Pins(1)=EdGraphPin'EdGraphPin_1623'
-	   Pins(2)=EdGraphPin'EdGraphPin_1624'
-	   Pins(3)=EdGraphPin'EdGraphPin_1625'
-	   NodePosX=3680
-	   NodePosY=544
-	   NodeGuid=B15602754FAC5D57625079801DE9C963
+		  LinkedTo(0)=EdGraphPin'K2Node_CallFunction_2191.EdGraphPin_20124'
+	   End Object
+	   Pins(0)=EdGraphPin'EdGraphPin_20142'
+	   Pins(1)=EdGraphPin'EdGraphPin_20143'
+	   Pins(2)=EdGraphPin'EdGraphPin_20144'
+	   Pins(3)=EdGraphPin'EdGraphPin_20145'
+	   NodePosX=-432
+	   NodePosY=2464
+	   NodeGuid=A657CEA146A38DE8319BF7B5121DF69C
 	End Object
-	Begin Object Class=K2Node_CallFunction Name="K2Node_CallFunction_834898"
-	   Begin Object Class=EdGraphPin Name="EdGraphPin_1766"
+	Begin Object Class=K2Node_CallFunction Name="K2Node_CallFunction_2192"
+	   Begin Object Class=EdGraphPin Name="EdGraphPin_20154"
 	   End Object
-	   Begin Object Class=EdGraphPin Name="EdGraphPin_1767"
+	   Begin Object Class=EdGraphPin Name="EdGraphPin_20155"
 	   End Object
-	   Begin Object Class=EdGraphPin Name="EdGraphPin_1768"
+	   Begin Object Class=EdGraphPin Name="EdGraphPin_20156"
 	   End Object
-	   Begin Object Class=EdGraphPin Name="EdGraphPin_1769"
+	   Begin Object Class=EdGraphPin Name="EdGraphPin_20157"
 	   End Object
-	   Begin Object Class=EdGraphPin Name="EdGraphPin_1770"
+	   Begin Object Class=EdGraphPin Name="EdGraphPin_20158"
 	   End Object
-	   Begin Object Class=EdGraphPin Name="EdGraphPin_1771"
+	   Begin Object Class=EdGraphPin Name="EdGraphPin_20159"
 	   End Object
-	   Begin Object Name="EdGraphPin_1766"
+	   Begin Object Name="EdGraphPin_20154"
 		  PinName="execute"
 		  PinType=(PinCategory="exec")
-		  LinkedTo(0)=EdGraphPin'K2Node_VariableSet_149.EdGraphPin_1510'
-		  LinkedTo(1)=EdGraphPin'K2Node_IfThenElse_251.EdGraphPin_1804'
+		  LinkedTo(0)=EdGraphPin'K2Node_VariableSet_11.EdGraphPin_20029'
+		  LinkedTo(1)=EdGraphPin'K2Node_IfThenElse_28.EdGraphPin_20180'
 	   End Object
-	   Begin Object Name="EdGraphPin_1767"
+	   Begin Object Name="EdGraphPin_20155"
 		  PinName="then"
 		  Direction=EGPD_Output
 		  PinType=(PinCategory="exec")
-		  LinkedTo(0)=EdGraphPin'K2Node_VariableSet_86.EdGraphPin_1437'
+		  LinkedTo(0)=EdGraphPin'K2Node_VariableSet_10.EdGraphPin_19953'
 	   End Object
-	   Begin Object Name="EdGraphPin_1768"
+	   Begin Object Name="EdGraphPin_20156"
 		  PinName="self"
 		  PinFriendlyName="Target"
 		  PinToolTip="Target\nGameplay Statics Reference"
@@ -1532,165 +1583,150 @@ Here is what is going on in this graph:
 		  DefaultObject=Default__GameplayStatics
 		  bHidden=True
 	   End Object
-	   Begin Object Name="EdGraphPin_1769"
+	   Begin Object Name="EdGraphPin_20157"
 		  PinName="WorldContextObject"
 		  PinToolTip="World Context Object\nObject Reference"
 		  PinType=(PinCategory="object",PinSubCategoryObject=Class'/Script/CoreUObject.Object')
 		  bHidden=True
 	   End Object
-	   Begin Object Name="EdGraphPin_1770"
+	   Begin Object Name="EdGraphPin_20158"
 		  PinName="ActorClass"
 		  PinToolTip="Actor Class\nActor Class\n\nClass of Actor to find. Must be specified or result array will be empty."
 		  PinType=(PinCategory="class",PinSubCategoryObject=Class'/Script/Engine.Actor')
-		  DefaultObject=BlueprintGeneratedClass'/Game/Blueprints/AI/AI_Character.AI_Character_C'
 	   End Object
-	   Begin Object Name="EdGraphPin_1771"
+	   Begin Object Name="EdGraphPin_20159"
 		  PinName="OutActors"
-		  PinToolTip="Out Actors\nArray of AI Character References\n\nOutput array of Actors of the specified class."
-		  Direction=EGPD_Output
-		  PinType=(PinCategory="object",PinSubCategoryObject=BlueprintGeneratedClass'/Game/Blueprints/AI/AI_Character.AI_Character_C',bIsArray=True)
-		  LinkedTo(0)=EdGraphPin'K2Node_Knot_1.EdGraphPin_20243'
-	   End Object
-	   FunctionReference=(MemberParentClass=Class'/Script/Engine.GameplayStatics',MemberName="GetAllActorsOfClass")
-	   Pins(0)=EdGraphPin'EdGraphPin_1766'
-	   Pins(1)=EdGraphPin'EdGraphPin_1767'
-	   Pins(2)=EdGraphPin'EdGraphPin_1768'
-	   Pins(3)=EdGraphPin'EdGraphPin_1769'
-	   Pins(4)=EdGraphPin'EdGraphPin_1770'
-	   Pins(5)=EdGraphPin'EdGraphPin_1771'
-	   NodePosX=624
-	   NodePosY=240
-	   NodeGuid=0FA104764E6E171B1D9FFB986474E72C
+		  PinToolTip="Out Actors\nArray of Actor References\n\nOutput array of Actors of the specified class."
+		  Direction=EGPD_Output
+		  PinType=(PinCategory="object",PinSubCategoryObject=Class'/Script/Engine.Actor',bIsArray=True)
+		  LinkedTo(0)=EdGraphPin'K2Node_Knot_43.EdGraphPin_20208'
+	   End Object
+	   FunctionReference=(MemberParent=Class'/Script/Engine.GameplayStatics',MemberName="GetAllActorsOfClass")
+	   Pins(0)=EdGraphPin'EdGraphPin_20154'
+	   Pins(1)=EdGraphPin'EdGraphPin_20155'
+	   Pins(2)=EdGraphPin'EdGraphPin_20156'
+	   Pins(3)=EdGraphPin'EdGraphPin_20157'
+	   Pins(4)=EdGraphPin'EdGraphPin_20158'
+	   Pins(5)=EdGraphPin'EdGraphPin_20159'
+	   NodePosX=-3488
+	   NodePosY=2160
+	   NodeGuid=5EA186074133CEC8F715E6814F3F8E6A
 	End Object
-	Begin Object Class=K2Node_VariableGet Name="K2Node_VariableGet_245"
-	   Begin Object Class=EdGraphPin Name="EdGraphPin_1784"
+	Begin Object Class=K2Node_VariableGet Name="K2Node_VariableGet_40"
+	   Begin Object Class=EdGraphPin Name="EdGraphPin_20172"
 	   End Object
-	   Begin Object Class=EdGraphPin Name="EdGraphPin_1785"
+	   Begin Object Class=EdGraphPin Name="EdGraphPin_20173"
 	   End Object
-	   Begin Object Name="EdGraphPin_1784"
+	   Begin Object Name="EdGraphPin_20172"
 		  PinName="AI_CON_Ref"
 		  Direction=EGPD_Output
-		  PinType=(PinCategory="object",PinSubCategoryObject=BlueprintGeneratedClass'/Game/Blueprints/AI/Follower_AI_Con.Follower_AI_CON_C')
-		  LinkedTo(0)=EdGraphPin'K2Node_CallFunction_2644.EdGraphPin_20084'
+		  PinType=(PinCategory="object")
+		  LinkedTo(0)=EdGraphPin'K2Node_CallFunction_2193.EdGraphPin_20197'
 	   End Object
-	   Begin Object Name="EdGraphPin_1785"
+	   Begin Object Name="EdGraphPin_20173"
 		  PinName="self"
 		  PinFriendlyName="Target"
-		  PinType=(PinCategory="object",PinSubCategoryObject=BlueprintGeneratedClass'/Game/Blueprints/AI/AgroCheck.AgroCheck_C')
+		  PinType=(PinCategory="object",PinSubCategoryObject=BlueprintGeneratedClass'/Game/StarterContent/Blueprints/Blueprint_Effect_Explosion.Blueprint_Effect_Explosion_C')
 		  bHidden=True
 	   End Object
 	   VariableReference=(MemberName="AI_CON_Ref",bSelfContext=True)
-	   Pins(0)=EdGraphPin'EdGraphPin_1784'
-	   Pins(1)=EdGraphPin'EdGraphPin_1785'
-	   NodePosX=-656
-	   NodePosY=240
-	   NodeGuid=F4E5A3F7436BA7F125F8E2A836978821
+	   Pins(0)=EdGraphPin'EdGraphPin_20172'
+	   Pins(1)=EdGraphPin'EdGraphPin_20173'
+	   NodePosX=-4768
+	   NodePosY=2160
+	   bHasCompilerMessage=True
+	   ErrorType=1
+	   ErrorMsg="Warning Can\'t connect pins  AI CON Ref  and  A : Reference is not compatible with Object Reference.\nWarning Could not find a variable named \"AI_CON_Ref\" in \'Blueprint_Effect_Explosion_C\'.\nMake sure \'Blueprint_Effect_Explosion_C\' has been compiled for  Get AI_CON_Ref \nError The property associated with  AI CON Ref  could not be found"
+	   NodeGuid=743AD8D34E20E10A14823AA26FD37B35
 	End Object
-	Begin Object Class=K2Node_IfThenElse Name="K2Node_IfThenElse_251"
-	   Begin Object Class=EdGraphPin Name="EdGraphPin_1802"
+	Begin Object Class=K2Node_IfThenElse Name="K2Node_IfThenElse_28"
+	   Begin Object Class=EdGraphPin Name="EdGraphPin_20178"
 	   End Object
-	   Begin Object Class=EdGraphPin Name="EdGraphPin_1803"
+	   Begin Object Class=EdGraphPin Name="EdGraphPin_20179"
 	   End Object
-	   Begin Object Class=EdGraphPin Name="EdGraphPin_1804"
+	   Begin Object Class=EdGraphPin Name="EdGraphPin_20180"
 	   End Object
-	   Begin Object Class=EdGraphPin Name="EdGraphPin_1805"
+	   Begin Object Class=EdGraphPin Name="EdGraphPin_20181"
 	   End Object
-	   Begin Object Name="EdGraphPin_1802"
+	   Begin Object Name="EdGraphPin_20178"
 		  PinName="execute"
 		  PinType=(PinCategory="exec")
-		  LinkedTo(0)=EdGraphPin'K2Node_Event_44.EdGraphPin_1591'
+		  LinkedTo(0)=EdGraphPin'K2Node_CustomEvent_6.EdGraphPin_19768'
 	   End Object
-	   Begin Object Name="EdGraphPin_1803"
+	   Begin Object Name="EdGraphPin_20179"
 		  PinName="Condition"
 		  PinType=(PinCategory="bool")
-		  DefaultValue="false"
-		  AutogeneratedDefaultValue="false"
-		  LinkedTo(0)=EdGraphPin'K2Node_CallFunction_2644.EdGraphPin_20086'
+		  DefaultValue="true"
+		  AutogeneratedDefaultValue="true"
+		  LinkedTo(0)=EdGraphPin'K2Node_CallFunction_2193.EdGraphPin_20199'
 	   End Object
-	   Begin Object Name="EdGraphPin_1804"
+	   Begin Object Name="EdGraphPin_20180"
 		  PinName="then"
 		  PinFriendlyName="true"
 		  Direction=EGPD_Output
 		  PinType=(PinCategory="exec")
-		  LinkedTo(0)=EdGraphPin'K2Node_CallFunction_834898.EdGraphPin_1766'
+		  LinkedTo(0)=EdGraphPin'K2Node_CallFunction_2192.EdGraphPin_20154'
 	   End Object
-	   Begin Object Name="EdGraphPin_1805"
+	   Begin Object Name="EdGraphPin_20181"
 		  PinName="else"
 		  PinFriendlyName="false"
 		  Direction=EGPD_Output
 		  PinType=(PinCategory="exec")
-		  LinkedTo(0)=EdGraphPin'K2Node_DynamicCast_140.EdGraphPin_6203'
-	   End Object
-	   Pins(0)=EdGraphPin'EdGraphPin_1802'
-	   Pins(1)=EdGraphPin'EdGraphPin_1803'
-	   Pins(2)=EdGraphPin'EdGraphPin_1804'
-	   Pins(3)=EdGraphPin'EdGraphPin_1805'
-	   NodePosX=-192
-	   NodePosY=240
-	   NodeGuid=85E2A9F64FD32BDDAA97C2A4B3A3E9BF
+		  LinkedTo(0)=EdGraphPin'K2Node_DynamicCast_437.EdGraphPin_19779'
+	   End Object
+	   Pins(0)=EdGraphPin'EdGraphPin_20178'
+	   Pins(1)=EdGraphPin'EdGraphPin_20179'
+	   Pins(2)=EdGraphPin'EdGraphPin_20180'
+	   Pins(3)=EdGraphPin'EdGraphPin_20181'
+	   NodePosX=-4304
+	   NodePosY=2160
+	   NodeGuid=8859B2D7441D71FD581A5397AE8F0C47
 	End Object
-	Begin Object Class=EdGraphNode_Comment Name="EdGraphNode_Comment_85"
-	   NodePosX=-688
-	   NodePosY=160
+	Begin Object Class=EdGraphNode_Comment Name="EdGraphNode_Comment_11"
+	   NodePosX=-4800
+	   NodePosY=2080
 	   NodeWidth=720
 	   NodeHeight=192
 	   NodeComment="Is AI_CON_Ref set?"
-	   NodeGuid=4924281F40DAD32E3D7AEE8EFBADE3E4
+	   NodeGuid=119EAC1E4CCAB65876CD97B77D520A4E
 	End Object
-	Begin Object Class=K2Node_VariableGet Name="K2Node_VariableGet_340"
-	   Begin Object Class=EdGraphPin Name="EdGraphPin_1814"
+	Begin Object Class=K2Node_VariableGet Name="K2Node_VariableGet_41"
+	   Begin Object Class=EdGraphPin Name="EdGraphPin_20190"
 	   End Object
-	   Begin Object Class=EdGraphPin Name="EdGraphPin_1815"
+	   Begin Object Class=EdGraphPin Name="EdGraphPin_20191"
 	   End Object
-	   Begin Object Name="EdGraphPin_1814"
+	   Begin Object Name="EdGraphPin_20190"
 		  PinName="AI_CON_Ref"
 		  Direction=EGPD_Output
-		  PinType=(PinCategory="object",PinSubCategoryObject=BlueprintGeneratedClass'/Game/Blueprints/AI/Follower_AI_Con.Follower_AI_CON_C')
-		  LinkedTo(0)=EdGraphPin'K2Node_CallFunction_33307.EdGraphPin_1257'
+		  PinType=(PinCategory="object")
+		  LinkedTo(0)=EdGraphPin'K2Node_CallFunction_2180.EdGraphPin_19794'
 	   End Object
-	   Begin Object Name="EdGraphPin_1815"
+	   Begin Object Name="EdGraphPin_20191"
 		  PinName="self"
 		  PinFriendlyName="Target"
-		  PinType=(PinCategory="object",PinSubCategoryObject=BlueprintGeneratedClass'/Game/Blueprints/AI/AgroCheck.AgroCheck_C')
+		  PinType=(PinCategory="object",PinSubCategoryObject=BlueprintGeneratedClass'/Game/StarterContent/Blueprints/Blueprint_Effect_Explosion.Blueprint_Effect_Explosion_C')
 		  bHidden=True
 	   End Object
 	   VariableReference=(MemberName="AI_CON_Ref",bSelfContext=True)
-	   Pins(0)=EdGraphPin'EdGraphPin_1814'
-	   Pins(1)=EdGraphPin'EdGraphPin_1815'
-	   NodePosX=256
-	   NodePosY=576
-	   NodeGuid=4E767ACD4901CF0FE85D4BB703A11D15
+	   Pins(0)=EdGraphPin'EdGraphPin_20190'
+	   Pins(1)=EdGraphPin'EdGraphPin_20191'
+	   NodePosX=-3856
+	   NodePosY=2496
+	   bHasCompilerMessage=True
+	   ErrorType=1
+	   ErrorMsg="Warning Can\'t connect pins  AI CON Ref  and  Target : Reference is not compatible with Controller Reference.\nWarning Could not find a variable named \"AI_CON_Ref\" in \'Blueprint_Effect_Explosion_C\'.\nMake sure \'Blueprint_Effect_Explosion_C\' has been compiled for  Get AI_CON_Ref \nError The property associated with  AI CON Ref  could not be found"
+	   NodeGuid=20282A5A4019B770F28F249CD0A8AFE4
 	End Object
-	Begin Object Class=K2Node_CallFunction Name="K2Node_CallFunction_2644"
-	   Begin Object Class=EdGraphPin Name="EdGraphPin_20086"
-	   End Object
-	   Begin Object Class=EdGraphPin Name="EdGraphPin_20085"
+	Begin Object Class=K2Node_CallFunction Name="K2Node_CallFunction_2193"
+	   Begin Object Class=EdGraphPin Name="EdGraphPin_20196"
 	   End Object
-	   Begin Object Class=EdGraphPin Name="EdGraphPin_20084"
+	   Begin Object Class=EdGraphPin Name="EdGraphPin_20197"
 	   End Object
-	   Begin Object Class=EdGraphPin Name="EdGraphPin_20083"
+	   Begin Object Class=EdGraphPin Name="EdGraphPin_20198"
 	   End Object
-	   Begin Object Name="EdGraphPin_20086"
-		  PinName="ReturnValue"
-		  PinToolTip="Return Value\nBoolean"
-		  Direction=EGPD_Output
-		  PinType=(PinCategory="bool")
-		  DefaultValue="false"
-		  AutogeneratedDefaultValue="false"
-		  LinkedTo(0)=EdGraphPin'K2Node_IfThenElse_251.EdGraphPin_1803'
+	   Begin Object Class=EdGraphPin Name="EdGraphPin_20199"
 	   End Object
-	   Begin Object Name="EdGraphPin_20085"
-		  PinName="B"
-		  PinToolTip="B\nObject Reference"
-		  PinType=(PinCategory="object",PinSubCategoryObject=Class'/Script/CoreUObject.Object')
-	   End Object
-	   Begin Object Name="EdGraphPin_20084"
-		  PinName="A"
-		  PinToolTip="A\nObject Reference"
-		  PinType=(PinCategory="object",PinSubCategoryObject=Class'/Script/CoreUObject.Object')
-		  LinkedTo(0)=EdGraphPin'K2Node_VariableGet_245.EdGraphPin_1784'
-	   End Object
-	   Begin Object Name="EdGraphPin_20083"
+	   Begin Object Name="EdGraphPin_20196"
 		  PinName="self"
 		  PinFriendlyName="Target"
 		  PinToolTip="Target\nKismet Math Library Reference"
@@ -1698,155 +1734,178 @@ Here is what is going on in this graph:
 		  DefaultObject=Default__KismetMathLibrary
 		  bHidden=True
 	   End Object
+	   Begin Object Name="EdGraphPin_20197"
+		  PinName="A"
+		  PinToolTip="A\nObject Reference"
+		  PinType=(PinCategory="object",PinSubCategoryObject=Class'/Script/CoreUObject.Object')
+		  LinkedTo(0)=EdGraphPin'K2Node_VariableGet_40.EdGraphPin_20172'
+	   End Object
+	   Begin Object Name="EdGraphPin_20198"
+		  PinName="B"
+		  PinToolTip="B\nObject Reference"
+		  PinType=(PinCategory="object",PinSubCategoryObject=Class'/Script/CoreUObject.Object')
+	   End Object
+	   Begin Object Name="EdGraphPin_20199"
+		  PinName="ReturnValue"
+		  PinToolTip="Return Value\nBoolean"
+		  Direction=EGPD_Output
+		  PinType=(PinCategory="bool")
+		  DefaultValue="false"
+		  AutogeneratedDefaultValue="false"
+		  LinkedTo(0)=EdGraphPin'K2Node_IfThenElse_28.EdGraphPin_20179'
+	   End Object
 	   bIsPureFunc=True
-	   FunctionReference=(MemberParentClass=Class'/Script/Engine.KismetMathLibrary',MemberName="NotEqual_ObjectObject")
-	   Pins(0)=EdGraphPin'EdGraphPin_20083'
-	   Pins(1)=EdGraphPin'EdGraphPin_20084'
-	   Pins(2)=EdGraphPin'EdGraphPin_20085'
-	   Pins(3)=EdGraphPin'EdGraphPin_20086'
-	   NodePosX=-480
-	   NodePosY=240
-	   NodeGuid=3858C80640FBDB6E2236CE925D3DA63D
+	   FunctionReference=(MemberParent=Class'/Script/Engine.KismetMathLibrary',MemberName="NotEqual_ObjectObject")
+	   Pins(0)=EdGraphPin'EdGraphPin_20196'
+	   Pins(1)=EdGraphPin'EdGraphPin_20197'
+	   Pins(2)=EdGraphPin'EdGraphPin_20198'
+	   Pins(3)=EdGraphPin'EdGraphPin_20199'
+	   NodePosX=-4592
+	   NodePosY=2160
+	   bHasCompilerMessage=True
+	   ErrorType=3
+	   ErrorMsg="Warning Can\'t connect pins  A  and  AI CON Ref : Object Reference is not compatible with Reference."
+	   NodeGuid=AE3F83CF416CB0A3684633B550912EE9
 	End Object
-	Begin Object Class=K2Node_Knot Name="K2Node_Knot_1"
-	   Begin Object Class=EdGraphPin Name="EdGraphPin_20243"
+	Begin Object Class=K2Node_Knot Name="K2Node_Knot_43"
+	   Begin Object Class=EdGraphPin Name="EdGraphPin_20208"
 	   End Object
-	   Begin Object Class=EdGraphPin Name="EdGraphPin_20244"
+	   Begin Object Class=EdGraphPin Name="EdGraphPin_20209"
 	   End Object
-	   Begin Object Name="EdGraphPin_20243"
+	   Begin Object Name="EdGraphPin_20208"
 		  PinName="InputPin"
-		  PinType=(PinCategory="object",PinSubCategoryObject=BlueprintGeneratedClass'/Game/Blueprints/AI/AI_Character.AI_Character_C',bIsArray=True)
-		  LinkedTo(0)=EdGraphPin'K2Node_CallFunction_834898.EdGraphPin_1771'
+		  PinType=(PinCategory="object",PinSubCategoryObject=Class'/Script/Engine.Actor',bIsArray=True)
+		  LinkedTo(0)=EdGraphPin'K2Node_CallFunction_2192.EdGraphPin_20159'
 		  bDefaultValueIsIgnored=True
 	   End Object
-	   Begin Object Name="EdGraphPin_20244"
+	   Begin Object Name="EdGraphPin_20209"
 		  PinName="OutputPin"
 		  Direction=EGPD_Output
-		  PinType=(PinCategory="object",PinSubCategoryObject=BlueprintGeneratedClass'/Game/Blueprints/AI/AI_Character.AI_Character_C',bIsArray=True)
-		  LinkedTo(0)=EdGraphPin'K2Node_CallFunction_39224.EdGraphPin_1290'
-		  LinkedTo(1)=EdGraphPin'K2Node_Knot_2.EdGraphPin_20245'
+		  PinType=(PinCategory="object",PinSubCategoryObject=Class'/Script/Engine.Actor',bIsArray=True)
+		  LinkedTo(0)=EdGraphPin'K2Node_CallFunction_2182.EdGraphPin_19815'
+		  LinkedTo(1)=EdGraphPin'K2Node_Knot_44.EdGraphPin_20214'
 	   End Object
-	   Pins(0)=EdGraphPin'EdGraphPin_20243'
-	   Pins(1)=EdGraphPin'EdGraphPin_20244'
-	   NodePosX=1184
-	   NodePosY=512
-	   NodeGuid=11A7BA6C444F57D3DDBEEFAAE1DDE5D2
+	   Pins(0)=EdGraphPin'EdGraphPin_20208'
+	   Pins(1)=EdGraphPin'EdGraphPin_20209'
+	   NodePosX=-2928
+	   NodePosY=2432
+	   NodeGuid=E00D898F44BE7CA8FD01DFAFE1AC4896
 	End Object
-	Begin Object Class=K2Node_Knot Name="K2Node_Knot_2"
-	   Begin Object Class=EdGraphPin Name="EdGraphPin_20245"
+	Begin Object Class=K2Node_Knot Name="K2Node_Knot_44"
+	   Begin Object Class=EdGraphPin Name="EdGraphPin_20214"
 	   End Object
-	   Begin Object Class=EdGraphPin Name="EdGraphPin_20246"
+	   Begin Object Class=EdGraphPin Name="EdGraphPin_20215"
 	   End Object
-	   Begin Object Name="EdGraphPin_20245"
+	   Begin Object Name="EdGraphPin_20214"
 		  PinName="InputPin"
-		  PinType=(PinCategory="object",PinSubCategoryObject=BlueprintGeneratedClass'/Game/Blueprints/AI/AI_Character.AI_Character_C',bIsArray=True)
-		  LinkedTo(0)=EdGraphPin'K2Node_Knot_1.EdGraphPin_20244'
+		  PinType=(PinCategory="object",PinSubCategoryObject=Class'/Script/Engine.Actor',bIsArray=True)
+		  LinkedTo(0)=EdGraphPin'K2Node_Knot_43.EdGraphPin_20209'
 		  bDefaultValueIsIgnored=True
 	   End Object
-	   Begin Object Name="EdGraphPin_20246"
+	   Begin Object Name="EdGraphPin_20215"
 		  PinName="OutputPin"
 		  Direction=EGPD_Output
-		  PinType=(PinCategory="object",PinSubCategoryObject=BlueprintGeneratedClass'/Game/Blueprints/AI/AI_Character.AI_Character_C',bIsArray=True)
-		  LinkedTo(0)=EdGraphPin'K2Node_CallFunction_169714.EdGraphPin_1406'
+		  PinType=(PinCategory="object",PinSubCategoryObject=Class'/Script/Engine.Actor',bIsArray=True)
+		  LinkedTo(0)=EdGraphPin'K2Node_CallFunction_2184.EdGraphPin_19922'
 	   End Object
-	   Pins(0)=EdGraphPin'EdGraphPin_20245'
-	   Pins(1)=EdGraphPin'EdGraphPin_20246'
-	   NodePosX=2512
-	   NodePosY=512
-	   NodeGuid=326F8E184A88CF66FF4508815EBA8FED
+	   Pins(0)=EdGraphPin'EdGraphPin_20214'
+	   Pins(1)=EdGraphPin'EdGraphPin_20215'
+	   NodePosX=-1600
+	   NodePosY=2432
+	   NodeGuid=827668E54A90F50E765ABF88BF502A9A
 	End Object
-	Begin Object Class=K2Node_Knot Name="K2Node_Knot_3"
-	   Begin Object Class=EdGraphPin Name="EdGraphPin_20247"
+	Begin Object Class=K2Node_Knot Name="K2Node_Knot_45"
+	   Begin Object Class=EdGraphPin Name="EdGraphPin_20220"
 	   End Object
-	   Begin Object Class=EdGraphPin Name="EdGraphPin_20248"
+	   Begin Object Class=EdGraphPin Name="EdGraphPin_20221"
 	   End Object
-	   Begin Object Name="EdGraphPin_20247"
+	   Begin Object Name="EdGraphPin_20220"
 		  PinName="InputPin"
 		  PinType=(PinCategory="object",PinSubCategoryObject=Class'/Script/Engine.Actor')
-		  LinkedTo(0)=EdGraphPin'K2Node_CallFunction_41995.EdGraphPin_1368'
+		  LinkedTo(0)=EdGraphPin'K2Node_CallFunction_2183.EdGraphPin_19876'
 		  bDefaultValueIsIgnored=True
 	   End Object
-	   Begin Object Name="EdGraphPin_20248"
+	   Begin Object Name="EdGraphPin_20221"
 		  PinName="OutputPin"
 		  Direction=EGPD_Output
 		  PinType=(PinCategory="object",PinSubCategoryObject=Class'/Script/Engine.Actor')
-		  LinkedTo(0)=EdGraphPin'K2Node_Knot_4.EdGraphPin_20249'
+		  LinkedTo(0)=EdGraphPin'K2Node_Knot_46.EdGraphPin_20226'
 	   End Object
-	   Pins(0)=EdGraphPin'EdGraphPin_20247'
-	   Pins(1)=EdGraphPin'EdGraphPin_20248'
-	   NodePosX=2384
-	   NodePosY=1008
-	   NodeGuid=CB439B4B444669ACA45532BDFDF598BA
+	   Pins(0)=EdGraphPin'EdGraphPin_20220'
+	   Pins(1)=EdGraphPin'EdGraphPin_20221'
+	   NodePosX=-1728
+	   NodePosY=2928
+	   NodeGuid=BB1DCBCB4FCDEF5484DBE3AA406EEBC3
 	End Object
-	Begin Object Class=K2Node_Knot Name="K2Node_Knot_4"
-	   Begin Object Class=EdGraphPin Name="EdGraphPin_20249"
+	Begin Object Class=K2Node_Knot Name="K2Node_Knot_46"
+	   Begin Object Class=EdGraphPin Name="EdGraphPin_20226"
 	   End Object
-	   Begin Object Class=EdGraphPin Name="EdGraphPin_20250"
+	   Begin Object Class=EdGraphPin Name="EdGraphPin_20227"
 	   End Object
-	   Begin Object Name="EdGraphPin_20249"
+	   Begin Object Name="EdGraphPin_20226"
 		  PinName="InputPin"
 		  PinType=(PinCategory="object",PinSubCategoryObject=Class'/Script/Engine.Actor')
-		  LinkedTo(0)=EdGraphPin'K2Node_Knot_3.EdGraphPin_20248'
+		  LinkedTo(0)=EdGraphPin'K2Node_Knot_45.EdGraphPin_20221'
 		  bDefaultValueIsIgnored=True
 	   End Object
-	   Begin Object Name="EdGraphPin_20250"
+	   Begin Object Name="EdGraphPin_20227"
 		  PinName="OutputPin"
 		  Direction=EGPD_Output
 		  PinType=(PinCategory="object",PinSubCategoryObject=Class'/Script/Engine.Actor')
-		  LinkedTo(0)=EdGraphPin'K2Node_CallFunction_42644.EdGraphPin_1499'
+		  LinkedTo(0)=EdGraphPin'K2Node_CallFunction_2186.EdGraphPin_20018'
 	   End Object
-	   Pins(0)=EdGraphPin'EdGraphPin_20249'
-	   Pins(1)=EdGraphPin'EdGraphPin_20250'
-	   NodePosX=3264
-	   NodePosY=1008
-	   NodeGuid=7DB7D32E4E9C8F0FAEFCB49E54CC3B00
+	   Pins(0)=EdGraphPin'EdGraphPin_20226'
+	   Pins(1)=EdGraphPin'EdGraphPin_20227'
+	   NodePosX=-848
+	   NodePosY=2928
+	   NodeGuid=F2A4730849A279BDD702C5830C91C370
 	End Object
-	Begin Object Class=K2Node_Knot Name="K2Node_Knot_5"
-	   Begin Object Class=EdGraphPin Name="EdGraphPin_20251"
+	Begin Object Class=K2Node_Knot Name="K2Node_Knot_47"
+	   Begin Object Class=EdGraphPin Name="EdGraphPin_20232"
 	   End Object
-	   Begin Object Class=EdGraphPin Name="EdGraphPin_20252"
+	   Begin Object Class=EdGraphPin Name="EdGraphPin_20233"
 	   End Object
-	   Begin Object Name="EdGraphPin_20251"
+	   Begin Object Name="EdGraphPin_20232"
 		  PinName="InputPin"
 		  PinType=(PinCategory="object",PinSubCategoryObject=Class'/Script/Engine.Actor')
-		  LinkedTo(0)=EdGraphPin'K2Node_Knot_7.EdGraphPin_20256'
+		  LinkedTo(0)=EdGraphPin'K2Node_Knot_48.EdGraphPin_20239'
 		  bDefaultValueIsIgnored=True
 	   End Object
-	   Begin Object Name="EdGraphPin_20252"
+	   Begin Object Name="EdGraphPin_20233"
 		  PinName="OutputPin"
 		  Direction=EGPD_Output
 		  PinType=(PinCategory="object",PinSubCategoryObject=Class'/Script/Engine.Actor')
-		  LinkedTo(0)=EdGraphPin'K2Node_CallFunction_103176.EdGraphPin_1548'
-		  LinkedTo(1)=EdGraphPin'K2Node_CallFunction_103002.EdGraphPin_1529'
-	   End Object
-	   Pins(0)=EdGraphPin'EdGraphPin_20251'
-	   Pins(1)=EdGraphPin'EdGraphPin_20252'
-	   NodePosX=4000
-	   NodePosY=736
-	   NodeGuid=F5C0A3614769A81282A667815D8172A4
+		  LinkedTo(0)=EdGraphPin'K2Node_CallFunction_2188.EdGraphPin_20064'
+		  LinkedTo(1)=EdGraphPin'K2Node_CallFunction_2190.EdGraphPin_20111'
+	   End Object
+	   Pins(0)=EdGraphPin'EdGraphPin_20232'
+	   Pins(1)=EdGraphPin'EdGraphPin_20233'
+	   NodePosX=-112
+	   NodePosY=2656
+	   NodeGuid=AAF473AA4270950544001FA4CC114C56
 	End Object
-	Begin Object Class=K2Node_Knot Name="K2Node_Knot_7"
-	   Begin Object Class=EdGraphPin Name="EdGraphPin_20256"
+	Begin Object Class=K2Node_Knot Name="K2Node_Knot_48"
+	   Begin Object Class=EdGraphPin Name="EdGraphPin_20238"
 	   End Object
-	   Begin Object Class=EdGraphPin Name="EdGraphPin_20255"
+	   Begin Object Class=EdGraphPin Name="EdGraphPin_20239"
 	   End Object
-	   Begin Object Name="EdGraphPin_20256"
-		  PinName="OutputPin"
-		  Direction=EGPD_Output
-		  PinType=(PinCategory="object",PinSubCategoryObject=Class'/Script/Engine.Actor')
-		  LinkedTo(0)=EdGraphPin'K2Node_Knot_5.EdGraphPin_20251'
-	   End Object
-	   Begin Object Name="EdGraphPin_20255"
+	   Begin Object Name="EdGraphPin_20238"
 		  PinName="InputPin"
 		  PinType=(PinCategory="object",PinSubCategoryObject=Class'/Script/Engine.Actor')
-		  LinkedTo(0)=EdGraphPin'K2Node_CallFunction_42262.EdGraphPin_1467'
+		  LinkedTo(0)=EdGraphPin'K2Node_CallFunction_2185.EdGraphPin_19978'
 		  bDefaultValueIsIgnored=True
 	   End Object
-	   Pins(0)=EdGraphPin'EdGraphPin_20255'
-	   Pins(1)=EdGraphPin'EdGraphPin_20256'
-	   NodePosX=3392
-	   NodePosY=736
-	   NodeGuid=B9509E964D83708F1D879CB7BE225B27
+	   Begin Object Name="EdGraphPin_20239"
+		  PinName="OutputPin"
+		  Direction=EGPD_Output
+		  PinType=(PinCategory="object",PinSubCategoryObject=Class'/Script/Engine.Actor')
+		  LinkedTo(0)=EdGraphPin'K2Node_Knot_47.EdGraphPin_20232'
+	   End Object
+	   Pins(0)=EdGraphPin'EdGraphPin_20238'
+	   Pins(1)=EdGraphPin'EdGraphPin_20239'
+	   NodePosX=-720
+	   NodePosY=2656
+	   NodeGuid=59E5C0E248297E2662F5CE9F9ACEEC93
 	End Object
 [/VAR]
 
