#UE4doc #docs
Minor grammatical revision.

diff --git a/Engine/Documentation/Source/Programming/Tutorials/FirstPersonShooter/3/5/FirstPersonShooter_3_5.INT.udn b/Engine/Documentation/Source/Programming/Tutorials/FirstPersonShooter/3/5/FirstPersonShooter_3_5.INT.udn
--- a/Engine/Documentation/Source/Programming/Tutorials/FirstPersonShooter/3/5/FirstPersonShooter_3_5.INT.udn
+++ b/Engine/Documentation/Source/Programming/Tutorials/FirstPersonShooter/3/5/FirstPersonShooter_3_5.INT.udn
@@ -24,7 +24,7 @@ Version: 4.9
 %Steps%
 
 [REGION:note]
-Before getting started, download and extract our sample mesh from the following link:
+Before getting started, download and extract our sample image from the following link:
 *	[Sample Crosshair Image](Crosshair_fps_tutorial.zip "Sample Crosshair Image")
 [/REGION]
 
