#UE4doc #docs
Wrote first draft of the 'Animating your Character' section.
Availability: Docs

diff --git a/Engine/Documentation/Source/Programming/Tutorials/FirstPersonShooter/4/1/FirstPersonShooter_4_1.INT.udn b/Engine/Documentation/Source/Programming/Tutorials/FirstPersonShooter/4/1/FirstPersonShooter_4_1.INT.udn
--- a/Engine/Documentation/Source/Programming/Tutorials/FirstPersonShooter/4/1/FirstPersonShooter_4_1.INT.udn
+++ b/Engine/Documentation/Source/Programming/Tutorials/FirstPersonShooter/4/1/FirstPersonShooter_4_1.INT.udn
@@ -22,6 +22,71 @@ Version: 4.9
 
 %Steps%
 
-<!--TODO: Add tutorial steps here-->
+[REGION:note]
+Before getting started, download and extract our sample animations from the following link:
+*	[Sample Animations](FPP_Animations.zip "Sample Animations")
+[/REGION]
+
+During this step, you're going to get some animations into your game.
+
+1.	**Right-click** in the **Content Browser** and select **New Folder**. 
+
+	![](RightClickImport.png)
+
+1.	Name this new folder "Animations".
+
+1.	**Double-click** on the folder to open it.
+
+1.	Locate and select the five animation files, namely:
+
+	*	FPP_Idle.FBX
+	*	FPP_JumpEnd.FBX
+	*	FPP_JumpLoop.FBX
+	*	FPP_JumpStart.FBX
+	*	FPP_Run.FBX
+
+1.	**Click Open** to begin importing the animations to your project.
+
+1.	Select **HeroFPP_Skeleton** under the **Select Skeleton** heading, then **click Import All** to import all of the animations.
+
+	![](SelectHeroFPP_Skeleton.png)
+
+1.	You should now have all five animations inside the **Animations** folder in the **Content Browser**.
+
+	![](ImportedAnimations.png)
+
+1.	**Click** the **Save** button to save your imported mesh.
+
+## Creating the Animation Blueprint
+
+1.	**Click** the **Add New** button and select **Animation Blueprint**.
+
+	![](AddNewAnimBP.png)
+
+1.	Select **AnimInstance** as the Parent Class and choose **/Game/HeroFPP_Skeleton** as the Target Skeleton.
+
+	![](CreateAnimBP.png)
+
+1.	Name the new animation blueprint "Arms_AnimBP".
+
+1.	**Double-click Arms_AnimBP** to open the **Blueprint Editor**.
+
+	![](DoubleClickAnimBP.png)
+
+## Adding State Machine Transition Variables
+
+1.	In the **My Blueprint** tab, **click** on the **Add New** button and select **Variable**.
+
+	![](AddingNewVariable.png)
+
+1.	Make the variable a **Boolean**  and name it "IsRunning".
+
+1.	In the **My Blueprint** tab, **click** on the **Add New** button and select **Variable**.
+
+1.	Make the variable a **Boolean**  and name it "IsFalling".
+
+1.	You should now have two animation transition variables that will drive the animation state machine.
+
+	![](NewBooleans.png)
 
 %Steps%
\ No newline at end of file
