#UE4doc #docs
Completed first draft of the 'Setting up your Event Graph' section.
Availability: Docs

diff --git a/Engine/Documentation/Source/Programming/Tutorials/FirstPersonShooter/4/2/FirstPersonShooter_4_2.INT.udn b/Engine/Documentation/Source/Programming/Tutorials/FirstPersonShooter/4/2/FirstPersonShooter_4_2.INT.udn
--- a/Engine/Documentation/Source/Programming/Tutorials/FirstPersonShooter/4/2/FirstPersonShooter_4_2.INT.udn
+++ b/Engine/Documentation/Source/Programming/Tutorials/FirstPersonShooter/4/2/FirstPersonShooter_4_2.INT.udn
@@ -23,6 +23,92 @@ Version: 4.9
 
 %Steps%
 
-<!--TODO: Add tutorial steps here-->
+During this step, you'll edit the animation's Event Graph to ensure that your animation transition variables are firing properly while the game is running.
+
+## Updating your State Variables
+
+1.	Open the event graph by double-clicking on **EventGraph** in the **My Blueprint** tab.
+
+	![](OpenEventGraph.png)
+
+1.	**Right-click** in the graph to bring up the **Context Menu**.
+
+1.	Type "Update" in the **Context Menu** search field, then **click** on the **Event Blueprint Update Animation** to add that node.
+
+	[REGION:note]
+	The **Event Blueprint Update Animation** node allows us to update our state variables every time the Animation updates so that they are always in sync with the game state.  
+	![](EventBlueprintUpdateAnimationNode.png)
+	[/REGION]
+
+1.	**Right-click** in the graph to bring up the **Context Menu**.
+
+1.	Type "Owner" in the **Context Menu** search field, then **click** on **Try Get Pawn Owner** to add that node.
+
+	![](TryGetPawnOwnerBP_Node.png)
+
+1.	Draw off the output pin and select **Cast to Character** from the **Context Menu**.
+
+	![](CastToCharacterContextMenu.png)
+
+1.	Wire the output execution pin on the **Event Blueprint Update Animation** to the input execution pin on **Cast to Character**.
+
+	![](ConnectedCastToCharacterNode.png)
+
+1.	Drag off the **As Character** output pin and select **Get Character Movement**.
+
+	![](GetCharacterMovementNode.png)
+
+1.	Drag off the **Character Movement** output pin and select **Get Movement Mode**.
+
+	![](GetMovementMode.png)
+
+## Querying Character Movement
+
+1.	 Draff off the **Movement Mode** output pin and select **Equal (Enum)**.
+
+	![](EqualEnumNode.png)
+
+## Determining if the Character is Falling
+
+1.	Set the dropdown value on the **Equal (Enum)** node to **Falling**.
+
+	![](ConnectedIsFalling.png)
+
+1.	**Alt-click** on **IsFalling** in the **My Blueprint** tab and drag into the graph to create a **Set Is Falling** node.
+
+	![](IsFalling.png)
+
+1.	Connect the unlabeled output execution pin of the **Cast to Character** node to the input execution pin of the **Set Is Falling** node and connect the output Boolean data pin of the **Equal (Enum)** node to the input Boolean data pin of the **Set Is Falling** node.
+
+	![](ConnectedIsFalling.png)
+
+## Determining if the Character is Running
+
+1.	Go back to the **Cast To Character** node and drag off the **As Character** pin again. This time, select the **Get Velocity** node.
+
+	![](GetVelocityNode.png)
+
+1.	If the character is not standing still, the length of its velocity vector will be greater than zero. So, drag off the **Return Value** vector output pin and select **Vector Length** to add that node to the graph.
+
+	![](VectorLength.png)
+
+1.	Drag off the **Return Value** float output pin and select the **> (float)** node.
+
+	![](GreaterThan.png)
+
+1.	**Alt-click** on **IsRunning** in the **My Blueprint** tab and drag into the graph to create a **Set Is Running** node.
+
+	![](IsRunning.png)
+
+1.	Connect the output execution pin of the **Set Is Falling** node to the input execution pin of the **Set Is Running** node and connect the output Boolean pin to the **> (float)** node to the input Boolean pin of the **Set Is Running** node.
+
+	![](ConnectedIsRunning.png)
+
+1.	Your **Event Graph** should now look like the following:
+
+	[REGION:lightbox]
+		[![](CompleteEventGraph.png)(w:750)](CompleteEventGraph.png)
+	[/REGION]
+	_Click on the image to zoom in._
 
 %Steps%
\ No newline at end of file
