#UE4doc #docs
Completed first draft of the 'Adding a State Machine' section.
Availability: Docs

diff --git a/Engine/Documentation/Source/Programming/Tutorials/FirstPersonShooter/4/3/FirstPersonShooter_4_3.INT.udn b/Engine/Documentation/Source/Programming/Tutorials/FirstPersonShooter/4/3/FirstPersonShooter_4_3.INT.udn
--- a/Engine/Documentation/Source/Programming/Tutorials/FirstPersonShooter/4/3/FirstPersonShooter_4_3.INT.udn
+++ b/Engine/Documentation/Source/Programming/Tutorials/FirstPersonShooter/4/3/FirstPersonShooter_4_3.INT.udn
@@ -23,6 +23,24 @@ Version: 4.9
 
 %Steps%
 
-<!--TODO: Add tutorial steps here-->
+During this step, you're going to put together a state machine that will take advantage of your newly made variables.
+
+1.	**Double-click** on **AnimGraph** in the **My Blueprint** tab to open it.
+
+	![](AnimGraph.png)
+
+1.	**Right-click** the graph and select **State Machines > Add New State Machine... ** in the **Context Menu**.
+
+	![](AddNewSM.png)
+
+1.	**Right-click** on **New State Machine** in the **My Blueprint** tab and name it "Arms State Machine".
+
+	![](NewArmsStateMachine.png)
+
+1.	Connect the ouput execution pin on the "Arms State Machine" node to the **Result** input execution pin on the **Final Animation Pose** node.
+
+	![](ConnectedArmsStateMachine.png)
+
+1.	**Double-click** the "Arms State Machine" node to open its graph for editing.
 
 %Steps%
\ No newline at end of file
