#UE4doc #docs
Completed first draft of the 'Adding Animation Transition States' section.
Availability: Docs

diff --git a/Engine/Documentation/Source/Programming/Tutorials/FirstPersonShooter/4/4/FirstPersonShooter_4_4.INT.udn b/Engine/Documentation/Source/Programming/Tutorials/FirstPersonShooter/4/4/FirstPersonShooter_4_4.INT.udn
--- a/Engine/Documentation/Source/Programming/Tutorials/FirstPersonShooter/4/4/FirstPersonShooter_4_4.INT.udn
+++ b/Engine/Documentation/Source/Programming/Tutorials/FirstPersonShooter/4/4/FirstPersonShooter_4_4.INT.udn
@@ -23,6 +23,54 @@ Version: 4.9
 
 %Steps%
 
-<!--TODO: Add tutorial steps here-->
+During this step, you're going to add five states to the state graph, namely:
+
+*	Idle
+*	Run
+*	JumpStart
+*	JumpEnd
+*	JumpLoop
+
+1.	**Right-click** in the graph and select **Add State...** from the context menu.
+
+	![](AddState.png)
+
+1.	Name the state "Idle".
+
+1.	**Double-click** the state to begin editing it.
+
+1.	**Right-click** in the graph area and search for "Idle" in the **Context Menu**.
+
+1.	**Click** on **Play FPP_Idle** to insert that node.
+
+	![](PlayIdle.png)
+
+1.	Connect the output execution pin of the **Play FPP_Idle** node to the **Result** input execution pin of the **Final Animation Pose** node.
+
+	![](ConnectedIdleNode.png)
+
+[REGION:note]
+Repeat steps 1-6 for each of the other 4 states, namely:
+
+*	Run
+*	JumpStart
+*	JumpEnd
+*	JumpLoop
+[/REGION]
+
+1.	When you are done, the **Arms State Machine** graph should look like the following image:
+
+	[REGION:lightbox]
+		[![](AnimationSMGraph.png)(w:750)](AnimationSMGraph.png)
+	[/REGION]
+	_Click on the image to zoom in._
+
+	[REGION:note]
+	Each state should contain the appropriate **Play** node connected to the **Final Animation Pose** node.
+	[/REGION]
+
+1.	Drag a wire from the **Entry** node to the **Idle** state node.
+
+	![](ConnectEntrytoIdle.png)
 
 %Steps%
\ No newline at end of file
