#UE4doc #docs
Wrote first draft of the 'Setting Idle to/from Run' transitions.
Availability: Docs

diff --git a/Engine/Documentation/Source/Programming/Tutorials/FirstPersonShooter/4/4/1/FirstPersonShooter_4_4_1.INT.udn b/Engine/Documentation/Source/Programming/Tutorials/FirstPersonShooter/4/4/1/FirstPersonShooter_4_4_1.INT.udn
--- a/Engine/Documentation/Source/Programming/Tutorials/FirstPersonShooter/4/4/1/FirstPersonShooter_4_4_1.INT.udn
+++ b/Engine/Documentation/Source/Programming/Tutorials/FirstPersonShooter/4/4/1/FirstPersonShooter_4_4_1.INT.udn
@@ -23,6 +23,36 @@ Version: 4.9
 
 %Steps%
 
-<!--TODO: Add tutorial steps here-->
+During this step, you're going to set the idle to and from run transitions so that when your character starts moving, the state machine will transition from **Idle** to **Run**.
+
+1.	Drag a wire from the **Idle** state to the **Run** state to create a transition.
+
+	![](WireIdletoRun.png)
+
+1.	**Double-click** the transition to edit it.
+
+1.	**Control-click** on **IsRunning** in the **My Blueprint** tab and drag into the graph to create a **Get Is Running** node.
+
+1.	Connect the output pin from the **Get Is Running** node to the input pin of the **Can Enter Transition** pin on the **Result** node.
+
+	![](IsRunningtoResult.png)
+
+	[REGION:note]
+	When your character stops moving, the state machine should transition from the **Run** state to the **Idle** state.
+	[/REGION]
+
+1.	Return to the **Arms State Machine** graph and drag a wire from the **Run** state to the **Idle** state.
+
+	![](WireRuntoIdle.png)
+
+1.	**Double-click** the transition to edit it.
+
+1.	**Control-click** on **IsRunning** in the **My Blueprint** tab and drag into the graph to create a **Get Is Running** node.
+
+1.	Drag off the output Boolean pin on the **Get Is Running** node and create a **Not Boolean** node.
+
+1.	Connect the output pin from the **Not Boolean** node to the input pin of the **Can Enter Transition** pin on the **Result** node.
+
+	![](IsNotRunningtoResult.png)
 
 %Steps%
\ No newline at end of file
