#UE4doc #docs
Wrote the first draft of the 'Setting Idle to Jump Start Transition' section.
Availability: Docs

diff --git a/Engine/Documentation/Source/Programming/Tutorials/FirstPersonShooter/4/4/2/FirstPersonShooter_4_4_2.INT.udn b/Engine/Documentation/Source/Programming/Tutorials/FirstPersonShooter/4/4/2/FirstPersonShooter_4_4_2.INT.udn
--- a/Engine/Documentation/Source/Programming/Tutorials/FirstPersonShooter/4/4/2/FirstPersonShooter_4_4_2.INT.udn
+++ b/Engine/Documentation/Source/Programming/Tutorials/FirstPersonShooter/4/4/2/FirstPersonShooter_4_4_2.INT.udn
@@ -23,6 +23,18 @@ Version: 4.9
 
 %Steps%
 
-<!--TODO: Add tutorial steps here-->
+During this step, you're going to set the idle to jump start transition.
+
+1.	Return to the **Arms State Machine** graph and drag a wire from the **Idle** state to the **JumpStart** state.
+
+	![](IdletoJumpStart.png)
+
+1.	**Double-click** the transition to edit it.
+
+1.	**Control-click** on **IsFalling** in the **My Blueprint** tab and drag into the graph to create a **Get Is Falling** node.
+
+1.	Connect the output Boolean pin from the **Get Is Falling** node to the input Boolean **Can Enter Transition** pin on the **Result** node.
+
+	![](IsFallingtoResult.png)
 
 %Steps%
\ No newline at end of file
