#UE4doc #docs
Completed the first draft of the 'Setting the Run to JumpStart Transition' section.
Availability: Docs

diff --git a/Engine/Documentation/Source/Programming/Tutorials/FirstPersonShooter/4/4/3/FirstPersonShooter_4_4_3.INT.udn b/Engine/Documentation/Source/Programming/Tutorials/FirstPersonShooter/4/4/3/FirstPersonShooter_4_4_3.INT.udn
--- a/Engine/Documentation/Source/Programming/Tutorials/FirstPersonShooter/4/4/3/FirstPersonShooter_4_4_3.INT.udn
+++ b/Engine/Documentation/Source/Programming/Tutorials/FirstPersonShooter/4/4/3/FirstPersonShooter_4_4_3.INT.udn
@@ -23,6 +23,18 @@ Version: 4.9
 
 %Steps%
 
-<!--TODO: Add tutorial steps here-->
+During this step, you're going to set the run to jump start transition.
+
+1.	Return to the **Arms State Machine** graph and drag a wire from the **Run** state to the **JumpStart** state.
+
+	![](RuntoJumpStart.png)
+
+1.	**Double-click** the transition to edit it.
+
+1.	**Control-click** on **IsFalling** in the **My Blueprint** tab and drag into the graph to create a **Get Is Falling** node.
+
+1.	Connect the output Boolean pin from the **Get Is Falling** node to the input Boolean **Can Enter Transition** pin on the **Result** node.
+
+	![](IsFallingtoResult.png)
 
 %Steps%
\ No newline at end of file
