#UE4doc #docs
Wrote first draft of the 'Setting the Jump Start to JumpLoop Transition' section.
Availability: Docs

diff --git a/Engine/Documentation/Source/Programming/Tutorials/FirstPersonShooter/4/4/4/FirstPersonShooter_4_4_4.INT.udn b/Engine/Documentation/Source/Programming/Tutorials/FirstPersonShooter/4/4/4/FirstPersonShooter_4_4_4.INT.udn
--- a/Engine/Documentation/Source/Programming/Tutorials/FirstPersonShooter/4/4/4/FirstPersonShooter_4_4_4.INT.udn
+++ b/Engine/Documentation/Source/Programming/Tutorials/FirstPersonShooter/4/4/4/FirstPersonShooter_4_4_4.INT.udn
@@ -23,6 +23,20 @@ Version: 4.9
 
 %Steps%
 
-<!--TODO: Add tutorial steps here-->
+During this step, you're going to set your jump start to jump loop transition.
+
+1.	Return to the **Arms State Machine** graph and drag a wire from the **JumpStart** state to the **JumpLoop** state.
+
+	![](JumpStarttoJumpLoop.png)
+
+1.	**Double-click** the transition to edit it.
+
+1.	**Right-click** in the graph, search for and select the **TimeRemaining for 'FPP_JumpStart'** node.
+
+1.	Drag off the **Time Remaining** output pin and add a **<= (float)** node using the **Context Menu**.
+
+1.	Enter 0.1 in the other output field on the **<= (float)** node and then wire the Boolean outputpin from that node to the **Can Enter Transition** input pin of the **Result** node.
+
+	![](TimeRemaining.png)
 
 %Steps%
\ No newline at end of file
