#UE4doc #docs
Wrote the first draft of the 'Setting the Jump Loop to Jump End Transition' section.
Availability: Docs

diff --git a/Engine/Documentation/Source/Programming/Tutorials/FirstPersonShooter/4/4/5/FirstPersonShooter_4_4_5.INT.udn b/Engine/Documentation/Source/Programming/Tutorials/FirstPersonShooter/4/4/5/FirstPersonShooter_4_4_5.INT.udn
--- a/Engine/Documentation/Source/Programming/Tutorials/FirstPersonShooter/4/4/5/FirstPersonShooter_4_4_5.INT.udn
+++ b/Engine/Documentation/Source/Programming/Tutorials/FirstPersonShooter/4/4/5/FirstPersonShooter_4_4_5.INT.udn
@@ -23,6 +23,20 @@ Version: 4.9
 
 %Steps%
 
-<!--TODO: Add tutorial steps here-->
+During this step, you're going to set your jump loop to jump end transition.
+
+1.	Return to the **Arms State Machine** graph and drag a wire from the **JumpLoop** state to the **JumpEnd** state.
+
+	![](JumpLooptoJumpEnd.png)
+
+1.	**Double-click** the transition to edit it.
+
+1.	**Control-click** on **IsFalling** in the **My Blueprint** tab and drag into the graph to create a **Get Is Falling** node.
+
+1.	Drag off the output Boolean pin on the **Get Is Falling** node and create a **Not Boolean** node.
+
+1.	Connect the output pin from the **Not Boolean** node to the input pin of the **Can Enter Transition** pin on the **Result** node.
+
+	![](IsNotFalling.png)
 
 %Steps%
\ No newline at end of file
