#UE4doc #docs
Completed first draft of 'Associating the Animation Blueprint with the Character Blueprint' section.
Availability: Docs

diff --git a/Engine/Documentation/Source/Programming/Tutorials/FirstPersonShooter/4/5/FirstPersonShooter_4_5.INT.udn b/Engine/Documentation/Source/Programming/Tutorials/FirstPersonShooter/4/5/FirstPersonShooter_4_5.INT.udn
--- a/Engine/Documentation/Source/Programming/Tutorials/FirstPersonShooter/4/5/FirstPersonShooter_4_5.INT.udn
+++ b/Engine/Documentation/Source/Programming/Tutorials/FirstPersonShooter/4/5/FirstPersonShooter_4_5.INT.udn
@@ -22,6 +22,31 @@ Version: 4.9
 
 %Steps%
 
-<!--TODO: Add tutorial steps here-->
+During this final section of the tutorial, you're going to associat the animation blueprint you created with your character blueprint.
+
+1.	Compile and save the **Arms_AnimBP** Animation Blueprint before closing it.
+
+1.	Navigate to the **Blueprints** folder in the **Content Browser** and open the **BP_FPSCharacter** Blueprint.
+
+1.	In **Defaults** mode, locate the **FirstPersonMesh** subsection inside the **Animation** section.
+
+1.	Set the **AnimationBlueprint** for the **FirstPersonMesh** to the **Arms_AnimBP** Animation Blueprint you just created.
+
+1.	While you're in the **Defaults** mode, change the **FirstPersonMesh** transform to {0, 0, -150} for **Location** and {0, 0, 0} for **Rotation**.
+
+1.	Compile and save the Blueprint before closing it.
+
+1.	**Click** the **Play In** button in the **Level Editor Toolbar**. Your arms will be animated as they transition through the animations that you hooked together in the Animation State Machine.
+
+	![](FinalShot.png)
+
+Congratulations! You've learned how to:
+
+&#x2713; Set-up a Project  
+&#x2713; Implement a Character  
+&#x2713; Implement a Projectile  
+&#x2713; Animate a Character  
+
+You've just completed the C++ First Person Shooter tutorial!
 
 %Steps%
\ No newline at end of file
