#UE4 Docs: Checking in the first pass of the des docs.

diff --git a/Engine/Documentation/Source/Engine/Physics/Destructibles/HowTo/DestructibleMesh/DestructibleMesh.INT.udn b/Engine/Documentation/Source/Engine/Physics/Destructibles/HowTo/DestructibleMesh/DestructibleMesh.INT.udn
new file mode 100644
--- /dev/null
+++ b/Engine/Documentation/Source/Engine/Physics/Destructibles/HowTo/DestructibleMesh/DestructibleMesh.INT.udn
@@ -0,0 +1,96 @@
+Availability:Docs
+Title:Destructible Meshes
+Description:Creating and using Destructible Mesh in UE4.
+SkillLevel:Intermediate
+Version: 4.9
+
+This page covers how to create and use a Destructible Static Mesh Actors in Unreal Engine 4. 
+In this example we will first create a Destructible Static Mesh using a Static Mesh from  the Content Browser. 
+Once the Destructible Static Mesh has been created we will then bring it into a level so that when the level is played the Destructible Static Mesh will fall to the floor and "Fracture" into many smaller pieces.
+
+
+[REGION:note]
+For this tutorial we are using the **Blueprint First Person** template with **Starter Content** enabled. 
+[/REGION]
+
+1. First inside of the **Content Browser** go to the **Geometry** folder and from the **Meshes** folder select the **1M_Cube** Static Mesh.
+
+	![](Pick_SM_For_Des.png)
+
+1. With the 1M_Cube selected, **Right-click** on the 1M_Cube and from the menu that is displayed select the **Create Destructible Mesh** option. 
+
+	![](Create_Des_Mesh.png)
+
+1. Once that has been completed the following two things will happen. 
+	First a new Destructible Mesh asset will be created and placed in the Content Browser using the same name as the Static Mesh that is was derived from but with a **_DM** added to the end of the name.
+	The second thing that will happen is the **Destructible Mesh Editor** window will be displayed.
+	The **Destructible Mesh Editor** tool allows you to affect many different aspects of how the Static Mesh that is used for destruction will be cut up and how those pieces will react when they are broken off from the rest of the Destructible Mesh.
+
+	![](DSM_Editor.png)
+
+1. Inside of the **Destructible Mesh Editor** tool locate the **Fracture Settings** and under the **[Voronoi](https://en.wikipedia.org/wiki/Voronoi_diagram)** section set the **Cell Site Count** from 25 to **50** and then set the **Random Seed** to **240**.
+
+	![](Setting_Up_Fracture.png)
+
+1. Now press the **Fracture Mesh** button to see how the mesh will look when it has been fractured.
+
+	![](Applying_Fracture.png)
+
+	[region:tip]
+	You can use the **Explode Amount** slider to move the fractured pieces apart from one another making it easier to see what is going on. 
+	[/region]
+
+1. Once the mesh has been fractured go back to the **Meshes*** folder in the Content Browser and create two new Materials named **MAT_Inside** and **MAT_Outside**. 
+
+	![](Create_Two_New_Materails.png)
+
+1. Open up the **Mat_Outside** Material and add the following two Material Expression nodes to the Material Graph, **Constant 3 Vector** expression node and a **Constant** expression node.
+
+	![](Adding_MEN_To_Graph.png)
+
+1. Now plug the **Constant 3 Vector** into the **Base Color** input and set its color to be **Red** then plug the **Constant** into the **Roughness** input and set it's value to **.25** then press the **Apply** and **Save** buttons to compile and save your work. 
+
+	![](Connecting_And_SettingUp_MEN.png)
+
+1. Open up the **MAT_Inside** Material and set up in the same manner as the **Mat_Outside** but instead of making the color red make the color **green** and set it's **Roughness** to **0.5**.
+
+	![](Setting_Up_Inside_Material.png)
+
+	[region:note]
+	Any Material that is used with Destructible Meshes will need to have the **Usage** flag, **Used with Skeletal Mesh** enabled. 
+	If you forget to enable this flag UE4 will enable it for you automatically and re-compile the Material when you put this Material on a Destructible Mesh.
+	
+	![](Enable_Usage_Flag.png)
+	[/region]
+
+1. Open up the **1M_Cube_DM** Destructible Mesh and then locate the **Skeletal Mesh** section. Expand the two **Materials** sections and then apply the **MAT_Inside** and **MAT_Outside Materials**.
+
+	![](Applied_Materials.png)
+
+1. Once the Materials have been applied expand the **Damage** section and then enable the **Enable Impact Damage** option by clicking on the checkmark box next to the name.
+
+	![](Enable_Impact_Damge.png)
+
+1. Close the Destructible Mesh editor and then from the Content Browser, add the Destructible Mesh to the level by dragging it from the Content Browser into the level.
+
+	![](Adding_DM_To_Level.png)
+
+1. Select the Destructible Mesh that was just placed in the level and then in the **Details** panel under the **Physics** section turn on the **Simulate Physics** option.
+
+1. Now press the **Play** button and when the level starts you should see the Destructible Mesh that was added fall down to the ground. 
+	Once the Destructible Mesh makes contact with the ground it will fracture into many smaller pieces like in the video below.
+
+	[OBJECT:EmbeddedVideo]
+		[PARAMLITERAL:width]
+		640
+		[/PARAMLITERAL]
+		[PARAMLITERAL:height]
+		360
+		[/PARAMLITERAL]
+		[PARAMLITERAL:videoid]
+		CT1bIbJz-a0
+		[/PARAMLITERAL]
+	[/OBJECT]
+
+
+
