#UE4doc #docs - Added engine version tag.

diff --git a/Engine/Documentation/Source/Engine/Deployment/Deployment.INT.udn b/Engine/Documentation/Source/Engine/Deployment/Deployment.INT.udn
--- a/Engine/Documentation/Source/Engine/Deployment/Deployment.INT.udn
+++ b/Engine/Documentation/Source/Engine/Deployment/Deployment.INT.udn
@@ -2,6 +2,7 @@ Availability:Public
 Title:Packaging and Cooking Games
 Crumbs:%ROOT%, Engine
 Description:Generating platform-specific content for assets used in the game.
+Version: 4.9
 
 [VAR:Topic]
 [OBJECT:Topic]
@@ -53,6 +54,8 @@ Description:Generating platform-specific content for assets used in the game.
 	* [](Engine/Basics/Projects/Packaging "%Engine/Basics/Projects/Packaging:description%")
 	* [](Engine/Deployment/Launching "%Engine/Deployment/Launching:description%")
 	* [](Engine/Deployment/Cooking "%Engine/Deployment/Cooking:description%")
+	* [](Engine/Deployment/Releasing "/Engine/Deployment/Releasing:description%")
+	* [](Engine/Deployment/Patching "Engine/Deployment/Patching:description%")
 	* [](Engine/Deployment/UnrealFrontend "%Engine/Deployment/UnrealFrontend:description%")
 	[/PARAM]
 [/OBJECT]
