#UE4doc #docs - added engine version tag.

diff --git a/Engine/Documentation/Source/Engine/Deployment/Cooking/ContentCooking.INT.udn b/Engine/Documentation/Source/Engine/Deployment/Cooking/ContentCooking.INT.udn
--- a/Engine/Documentation/Source/Engine/Deployment/Cooking/ContentCooking.INT.udn
+++ b/Engine/Documentation/Source/Engine/Deployment/Cooking/ContentCooking.INT.udn
@@ -2,6 +2,7 @@ Availability:Public
 Title:Content Cooking
 Crumbs:%ROOT%, Engine, Engine/Deployment
 Description:Generating platform-specific content for assets used in the game.
+Version: 4.9
 
 Unreal Engine stores content assets in particular formats which it uses internally, such as PNG for texture 
 data or WAV for audio. However, this content needs to be converted to different formats for the various platforms, 
