#UE4doc #docs - added engine version tag.

diff --git a/Engine/Documentation/Source/Engine/Deployment/Launching/LaunchingToDevices.INT.udn b/Engine/Documentation/Source/Engine/Deployment/Launching/LaunchingToDevices.INT.udn
--- a/Engine/Documentation/Source/Engine/Deployment/Launching/LaunchingToDevices.INT.udn
+++ b/Engine/Documentation/Source/Engine/Deployment/Launching/LaunchingToDevices.INT.udn
@@ -2,7 +2,7 @@ Availability:Public
 Title:Launching to Devices
 Crumbs:%ROOT%, Engine, Engine/Deployment
 Description:One-click deploying your project to devices like iOS and Android for testing
-
+Version: 4.9
 
 ![](LaunchMenu.png)
 
@@ -14,7 +14,7 @@ In the drop down menu, you will see a list of possible platforms you can launch
  map on. Generally, you would not need Windows or Mac, as you can use the Play button for that without needing to cook your data. This page will primarily cover the information you need for launching on 
  mobile platforms.
 
-For Android, there will be multiple selections. See the [Android Texture Formats](Platforms\Android\Reference#androidtextureformats) page for more information.
+For Android, there will be multiple selections. See the [**Android Texture Formats**](Platforms\Android\Reference#androidtextureformats) page for more information.
 
 If you have multiple devices for a particular platform, you can choose between them here. Clicking on the platform (and device), the editor will start to cook the map, install the data on the device, and 
 then run the map on the device.
@@ -24,14 +24,14 @@ This method will only install the current map on the device for fast iteration,
 [/REGION]
 
 ### Generic UE4Game App
-If you are making a content-only project, the executable that runs for One Click Deploy is actually the generic "UE4Game" game (since it can be used with any content-only project). When it runs on a 
+If you are making a content-only project, the executable that runs for One-Click Deploy is actually the generic "UE4Game" game (since it can be used with any content-only project). When it runs on a 
 device, the icon that gets installed is just labeled as UE4Game. However, when we package a game, the name of your project is used in the final packaged build. When you install the packaged build,
  the icon will have your name (as well as your icons if you updated the default icons).
 
 
 
 ## Advanced Mode (UnrealFrontend)
-There is an additional tool for performing advanced building, cooking, deploying, packaging, and launching options. It is called UnrealFrontend and is located in the following locations depending on your platform:
+There is an additional tool for performing advanced building, cooking, deploying, packaging, and launching options. It is called **Unreal Frontend** and is located in the following locations depending on your platform:
 
 |Platform | Location |
 | --------| ---------|
@@ -39,7 +39,7 @@ There is an additional tool for performing advanced building, cooking, deploying
 | Mac | [ENGINE INSTALL LOCATION]\Engine\Binaries\Mac\UnrealFrontend.app |
 
 This tool will let you only cook certain maps, change the commandline, or even run your game without precooking all the data (called "Cook on the Fly" mode). These are advanced tools - for more information, 
-see [](Engine\Deployment\UnrealFrontend).
+see **[](Engine\Deployment\UnrealFrontend)**.
 
 
 
