#UE4doc #docs
Updated introductory section.

diff --git a/Engine/Documentation/Source/Programming/Tutorials/FirstPersonShooter/FirstPersonShooter.INT.udn b/Engine/Documentation/Source/Programming/Tutorials/FirstPersonShooter/FirstPersonShooter.INT.udn
--- a/Engine/Documentation/Source/Programming/Tutorials/FirstPersonShooter/FirstPersonShooter.INT.udn
+++ b/Engine/Documentation/Source/Programming/Tutorials/FirstPersonShooter/FirstPersonShooter.INT.udn
@@ -6,6 +6,7 @@ SkillLevel: Beginner
 Version: 4.9
 
 ![](FPS_CoverImage.png)
+_This is what you'll see at the end of this tutorial._
 
 ## Goals
 
