#UE4doc #docs
Minor revision.

diff --git a/Engine/Documentation/Source/Programming/Tutorials/FirstPersonShooter/1/5/FirstPersonShooter_1_5.INT.udn b/Engine/Documentation/Source/Programming/Tutorials/FirstPersonShooter/1/5/FirstPersonShooter_1_5.INT.udn
--- a/Engine/Documentation/Source/Programming/Tutorials/FirstPersonShooter/1/5/FirstPersonShooter_1_5.INT.udn
+++ b/Engine/Documentation/Source/Programming/Tutorials/FirstPersonShooter/1/5/FirstPersonShooter_1_5.INT.udn
@@ -49,6 +49,6 @@ Congratulations! You've learned how to:
 &#x2713; Compile your First C++ Class  
 &#x2713; Set the Default Game Mode
 
-You're now ready to learn how to implement your character.
+You're now ready to learn how to implement your character in the next section.
 
 %Steps%
\ No newline at end of file
