#UE4doc #docs
Added cover image and updated concluding paragraph.

diff --git a/Engine/Documentation/Source/Programming/Tutorials/FirstPersonShooter/2/8/FirstPersonShooter_2_8.INT.udn b/Engine/Documentation/Source/Programming/Tutorials/FirstPersonShooter/2/8/FirstPersonShooter_2_8.INT.udn
--- a/Engine/Documentation/Source/Programming/Tutorials/FirstPersonShooter/2/8/FirstPersonShooter_2_8.INT.udn
+++ b/Engine/Documentation/Source/Programming/Tutorials/FirstPersonShooter/2/8/FirstPersonShooter_2_8.INT.udn
@@ -314,4 +314,17 @@ Before continuing, download and extract our sample mesh from the following link:
 
 1.	Click the **Stop** button in the Level Editor to exit Play in Editor (PIE) mode.
 
+Congratulations! You've learned how to:
+
+&#x2713; Make a New Character  
+&#x2713; Set-up Axis Mapping  
+&#x2713; Implement Character Movement Functions  
+&#x2713; Implement Mouse Camera Control  
+&#x2713; Implement Character Jumping  
+&#x2713; Add a Mesh to your Character  
+&#x2713; Change the Camera View  
+&#x2713; Add a First Person Mesh to your Character  
+
+You're now ready to learn how to implement projectiles in the next section.
+
 %Steps%
\ No newline at end of file
diff --git a/Engine/Documentation/Source/Programming/Tutorials/FirstPersonShooter/2/FirstPersonShooter_2.INT.udn b/Engine/Documentation/Source/Programming/Tutorials/FirstPersonShooter/2/FirstPersonShooter_2.INT.udn
--- a/Engine/Documentation/Source/Programming/Tutorials/FirstPersonShooter/2/FirstPersonShooter_2.INT.udn
+++ b/Engine/Documentation/Source/Programming/Tutorials/FirstPersonShooter/2/FirstPersonShooter_2.INT.udn
@@ -7,7 +7,8 @@ Related: Programming/Tutorials/FirstPersonShooter
 SkillLevel: Beginner
 Version: 4.9
 
-<!--TODO: Insert Cover Image-->
+![](FreeMovingCamera.png)
+_This is what you'll see at the end of this section._
 
 ## Goals
 
