#UE4doc #docs
Added cover image and revised concluding paragraph.

diff --git a/Engine/Documentation/Source/Programming/Tutorials/FirstPersonShooter/3/5/FirstPersonShooter_3_5.INT.udn b/Engine/Documentation/Source/Programming/Tutorials/FirstPersonShooter/3/5/FirstPersonShooter_3_5.INT.udn
--- a/Engine/Documentation/Source/Programming/Tutorials/FirstPersonShooter/3/5/FirstPersonShooter_3_5.INT.udn
+++ b/Engine/Documentation/Source/Programming/Tutorials/FirstPersonShooter/3/5/FirstPersonShooter_3_5.INT.udn
@@ -240,12 +240,14 @@ During this step, you'll add a crosshair HUD element to your game so that you ca
 
 1.	Click the **Stop** button in the Level Editor to exit Play in Editor (PIE) mode.
 
-Congratulations, you can now:
+Congratulations! You've learned how to:
 
-*	Move around your level
-*	Jump around your level
-*	Control your camera
-*	Shoot projectiles that interact with physics objects
-*	Aim with crosshairs
+&#x2713; Add Projectiles to your Game  
+&#x2713; Implement Shooting  
+&#x2713; Set-up Projectile Collision and Lifetime  
+&#x2713; Get your Projectiles to Interact with the World  
+&#x2713; Add Crosshairs to your Viewport  
+
+You're now ready to learn how to animate your character in the next section.
 
 %Steps%
\ No newline at end of file
diff --git a/Engine/Documentation/Source/Programming/Tutorials/FirstPersonShooter/3/FirstPersonShooter_3.INT.udn b/Engine/Documentation/Source/Programming/Tutorials/FirstPersonShooter/3/FirstPersonShooter_3.INT.udn
--- a/Engine/Documentation/Source/Programming/Tutorials/FirstPersonShooter/3/FirstPersonShooter_3.INT.udn
+++ b/Engine/Documentation/Source/Programming/Tutorials/FirstPersonShooter/3/FirstPersonShooter_3.INT.udn
@@ -7,7 +7,8 @@ Related: Programming/Tutorials/FirstPersonShooter
 SkillLevel: Beginner
 Version: 4.9
 
-<!--TODO: Insert Cover Image-->
+![](CrosshairsInGame.png)
+_This is what you'll see at the end of this section._
 
 ## Goals
 
