#UE4doc #doc
Made some minor revisions and updated an image.

diff --git a/Engine/Documentation/Source/Programming/Tutorials/FirstPersonShooter/4/FirstPersonShooter_4.INT.udn b/Engine/Documentation/Source/Programming/Tutorials/FirstPersonShooter/4/FirstPersonShooter_4.INT.udn
--- a/Engine/Documentation/Source/Programming/Tutorials/FirstPersonShooter/4/FirstPersonShooter_4.INT.udn
+++ b/Engine/Documentation/Source/Programming/Tutorials/FirstPersonShooter/4/FirstPersonShooter_4.INT.udn
@@ -7,7 +7,8 @@ Related: Programming/Tutorials/FirstPersonShooter
 SkillLevel: Beginner
 Version: 4.9
 
-<!--TODO: Insert Cover Image-->
+![](FinalShot.png)
+_This is what you'll see at the end of this section._
 
 ## Goals
 
