Attempting to add again, hopefully it takes

diff --git a/Engine/Documentation/Source/Support/Builds/ReleaseNotes/2015/4_9/ReleaseNotes_4_9.INT.udn b/Engine/Documentation/Source/Support/Builds/ReleaseNotes/2015/4_9/ReleaseNotes_4_9.INT.udn
new file mode 100644
--- /dev/null
+++ b/Engine/Documentation/Source/Support/Builds/ReleaseNotes/2015/4_9/ReleaseNotes_4_9.INT.udn
@@ -0,0 +1,4726 @@
+Availability:Public
+Title:Unreal Engine 4.9 Release Notes
+Crumbs:%ROOT%, Engine
+Description:
+
+This release brings hundreds of updates for Unreal Engine 4, including 129 improvements submitted by the community of Unreal Engine developers on GitHub! Thanks to all of these contributors to Unreal Engine 4.9:
+
+[REDACTED_CONTRIBUTORS]
+
+# Major Features 
+
+## Enhanced Support for Mobile Devices
+
+We've been working to greatly improve support for mobile platforms in UE4, and this release contains many exciting rendering features to help you create leading-edge mobile titles.  These include efficient dynamic shadows, movable lights and decals.  Expect to see continued improvements to mobile and HTML5 rendering quality in the future.  In this release, IOS In-App Purchases has been improved and we've added Remote Push Notifications and support for CloudKit.  For Android, we've added OpenGL ES 3.1, Immersive Mode support, asynchronous audio decompression, and better orientation features.  And we now include NVIDIA AndroidWorks to make it easy to get up and running on your device.  UE4 now includes a Flurry analytics plugin for iOS out of the box, also!  Many of the new features for mobile devices are also available in your HTML5 games, too.  
+
+![image alt text](image_0.png)
+
+## New: Dynamic Character Shadows for Mobile
+
+Dynamic modulated character shadows from directional lights are now supported on mobile devices!
+
+![image alt text](image_1.png)
+
+* You can use dynamic modulated shadows even if your scene is statically lit
+
+* To enable this feature, turn on the "Cast Modulated Shadows" option on a directional light!
+
+## New: Dynamic Point Lights for Mobile
+
+Dynamic point lights are now supported on mobile devices and HTML5!
+
+![](image_2.gif)(convert:false)
+
+* These can be used to enhance short lived visual effects such as explosions, projectiles, particle effects etc.
+
+* Up to four dynamic point lights are supported on each object being illuminated.
+
+* Note:  Dynamic shadows cast from point lights are not yet supported
+
+## New: Decals on Mobile
+
+Decals are now supported by the mobile renderer!
+
+![image alt text](image_3.png)
+
+* Translucent, additive and modulated blend modes are supported
+
+* Note:  Only unlit decals are supported on mobile and web, currently.
+
+## New: Major VR Updates
+
+### Steam VR (HTC Vive)
+
+Unreal Engine's SteamVR plugin includes many major fixes to make it easier to use, and more performant.
+
+![image alt text](image_4.png)
+
+SteamVR improvements include:
+
+* Fix for latency issues that appeared in 4.8 Preview 2+
+
+* Support for the standard Motion Controller abstraction, so both Vive controllers can route to the same player
+
+* C++ projects are now cleanly supported by the plugin API
+
+* Various improvements and bug fixes
+
+### Gear VR Updated to Mobile SDK 0.6.0
+
+This revision fixes several major Gear VR rendering bugs and issues present in the 0.5.0 release, as well as some hardware compatibility issues. 
+
+![image alt text](image_5.png)
+
+## New: VR Motion Controller Support
+
+Motion Controllers are now supported through a common abstraction layer in UE4!  
+
+![image alt text](image_6.png)
+
+This means that motion controller devices will now all be available through a common interface.  The HTC Vive controller is currently supported, with more (such as the PlayStation Move) to come soon.  Input button presses and triggers can now be routed to a single player from multiple controllers, using the Motion Controller key abstraction.  Additionally, you can add motion tracking to your project by simply adding a Motion Controller component to your character, which automatically updates anything attached to it to follow the position of your controllers!  We recommend that all plugins for motion controllers move over to the new system, so that they can be part of the device-agnostic interface.
+
+## New: Experimental DirectX 12 Support
+
+DirectX 12 is now supported as an experimental feature! If you are using Windows 10, try it out by running the engine with "-DX12" on the command line.
+
+![image alt text](image_7.png)
+
+Microsoft's engineers added support for DirectX 12 to UE4 and we have worked with them to integrate their changes into 4.9. The feature is still new and is considered experimental. DirectX 12 offers a much lower-level rendering API that is more efficient and allows for rendering commands to be submitted in parallel across many threads, a feature inspired by console rendering APIs. Going forward, we'll continue to improve support for DirectX 12 and look for ways to leverage the new API in upcoming versions of the engine.
+
+## New: Full Scene Particle Collision with Mesh Distance Fields
+
+Distance field GPU particle collision allows particles to efficiently collide with the entire scene.
+
+![image alt text](image_8.png)
+
+This can be used for [sparks colliding with surfaces](https://www.youtube.com/watch?v=3YH8oTzSQwY).
+
+![image alt text](image_9.gif)(convert:false)
+
+It can also be used for various effects like and [snow that accumulates on surfaces](https://www.youtube.com/watch?v=Z-2e8j0t0Ww),
+
+![image alt text](image_10.gif)(convert:false)
+
+* The Collision module on GPU particle emitters has a new property called 'Collision Mode' where you can specify 'Distance Field'
+
+![image alt text](image_11.png)
+
+* In order for a particle to collide we need to have a colliding plane, and this can be extracted very efficiently from the distance field.  A single sample of the distance field gives distance to the nearest colliding surface.  Computing the gradient of the distance field provides the direction to that surface, and the surface's normal.  Putting these together gives the collision plane.
+
+* Particle collision using this method is more reliable than the existing Scene Depth collision, which only collides against what's on your screen, restricting it to effects with short lifetimes.  However, the global distance field is fairly low resolution so particles will pass through thin objects, and sharp corners are rounded.
+
+* As with all distance field features, this requires 'Generate Mesh Distance Fields' to be enabled in the Rendering Project settings.  Distance field particle collision is a Shader Model 5 feature, on hardware that doesn't support SM5, the particles will fall back to scene depth collision instead.
+
+* Particle collision with distance fields is very efficient - it costs about the same as particle collision with scene depth.
+
+## New: Hierarchical LOD Fast Preview and Clustering
+
+Hierarchical Level-of-Detail (HLOD) is a system introduced in 4.8 which allows many objects in your levels to be collapsed down to only a few objects when they are small on screen.  This helps you achieve much higher quality levels when viewing objects up close, and faster overall performance for your level.  We've made improvements to HLOD in 4.9 and will continue to in future releases.
+
+### Fast HLOD Preview
+
+In order to allow for rapid iteration on Hierarchical LOD builds, you can now use the Preview HLOD option. This will create the HLOD Actors with only a visual representation of their bounds and objects without merging and creating the HLODs meshes themselves. Their bounds are rendered according to the Draw Distance setting inside the LOD System settings for the specific Hierarchical LOD Level. 
+
+![image alt text](image_12.gif)(convert:false)
+
+The image above shows both the first and second Hierarchical LOD level clusters popping into view, this gives a clear indication how the level is clustered without having to go through the time-consuming process of merging and creating the Hierarchical LOD meshes. Be aware that this is an initial implementation of the preview system and may tend to change and improve over time.
+
+### Custom LOD Clustering (HLOD Volumes)
+
+In order to allow for more granular control over the clustering process for the Hierarchical LODs, you can now place Hierarchical LOD Volumes which will define a cluster for the Actors it encapsulates. Use this to manually define which areas and objects within a level should be clustered together.
+
+![image alt text](image_13.png)
+
+## New: "Arch Vis" Character Controls
+
+The new "ArchVisCharacter” plugin adds a character class that provides out-of-the-box controls appropriate for architectural visualization applications. It is intended to feel pleasant and natural to control at realistic world scales, and to enable smoother videos and live demos.
+
+![image alt text](image_14.png)
+
+You can hook up ArchVisCharacter directly as your Pawn Class, or you can create custom blueprints from it to tweak the movement settings as you see fit.
+
+## New: Widget Depth Ordering
+
+Widget Components in screen space now automatically sort based on distance from the viewer, ensuring the closest screen space widget component is always on top of the other components.
+
+![image alt text](image_15.png)
+
+We now sort all widget components against one another based on distance to the camera. Previously it was based a manual Z Order, or whoever was added to the viewport first.
+
+## New: Area Shadows (for stationary lights)
+
+You're no longer limited to sharp shadows with uniform penumbra sizes.  The Lightmass ray-tracer now supports area shadows are now fully support for Stationary lights!
+
+![image alt text](image_16.gif)(convert:false)
+
+Turn on the "Use Area Shadows for Stationary Light" option to use this feature.  The 'Light Source Angle' (or 'Source Radius' for a point light) controls how soft these shadows are.
+
+![image alt text](image_17.png)
+
+Performance details:
+
+* Uniform penumbra shadows are still better at dealing with low lightmap resolutions, so those stay the default.  Previously the uniform penumbra size was controlled with an obscure project setting, now it is controlled by the 'Light Source Angle' property, just like the area shadows are.  Double the source angle to get twice as soft shadows, etc.
+
+* There's also a new optimization for areas of maps which are only affected by a single stationary light channel (eg one Stationary Directional light).  Their shadow maps will be packed into a single channel texture instead of a 4 channel texture, for 1/4th the memory cost.
+
+## New: Ambient Occlusion Material Mask
+
+The new Ambient Occlusion Material Mask feature lets you access Lightmass calculated AO in your material, which can be useful for procedural texturing, for example to add in aging effects and dirt in areas where it would accumulate.
+
+![image alt text](image_18.png)
+
+The above screenshot uses an AO mask to blend a dirt layer automatically into corners of the environment.
+
+To use AO mask, you'll need to enable both 'Use Ambient Occlusion' and 'Generate Ambient Occlusion Material Mask' under World Settings -> Lightmass settings, and then build lighting.  The other AO controls like Max Occlusion Distance can be useful to tweak the look.  (Be sure to set Direct and Indirect Occlusion Fraction to 0, so that this AO will not be applied to the actual level lighting.)
+
+![image alt text](image_19.png)
+
+Then in your material you can access the AO as a 0-1 mask with the PrecomputedAOMask node.  This node will provide 1 in occluded areas and 0 elsewhere. 
+
+![image alt text](image_20.png)
+
+Performance details:  
+
+* Using the PrecomputedAOMask node is just a standard texture lookup.  Enabling the feature adds half a byte per lightmap texel, where the default usage is about 4 bytes per lightmap texel.  
+
+* AOMaterialMask texture memory can be inspected with the ‘ListTextures' console command.
+
+## New: Mesh Distance Field Materials
+
+This feature allows your materials to cheaply determine the distance to the nearst solid from any point in space.  There are two new nodes to access the global distance field properties within the material editor. 
+
+![image alt text](image_21.png)
+
+DistanceToNearestSurface returns signed distance in world space units.  Note that the distance will always be near 0 on the surface of an opaque mesh.  DistanceToNearestSurface can be used to have effects avoid occluders:
+
+UE4 user Roel Bartstra has already found some interesting uses for this - one to [get a soft body look](https://www.youtube.com/watch?v=shmJIbCayhc), and another to [create flow maps that automatically avoid surfaces](https://www.youtube.com/watch?v=3lCkrRDOWXM).
+
+As with all distance field features, this requires 'Generate Mesh Distance Fields' to be enabled in the Rendering Project settings.  Distance field material access is a Shader Model 5 feature.  Use a FeatureLevelSwitch material node to provide fallbacks for Shader Model 4 hardware.
+
+## New: Improved Distance Field Ambient Occlusion
+
+Sky occlusion gives nice soft shadows around objects lit by the sky, especially important in overcast lighting scenarios, a use case traditionally hard to render at high quality without precomputed lighting. Distance Field Ambient Occlusion solves this for fully dynamic games where both the time of day and environment can change at any time. We've put a lot of effort into improving the quality and performance of the technology for 4.9 and believe that it is now a shippable feature on medium spec PC and PlayStation 4 level hardware with a total GPU cost of under 4ms for a typical game scene!
+
+The previous method did adaptive sampling, so flat surfaces did less work, but this also caused a lot of splotchiness in clean environments.
+
+![image alt text](image_22.gif)(convert:false)
+
+The AO computations are now fast enough to get rid of the adaptive sampling, so occlusion is much smoother.
+
+![image alt text](image_23.gif)(convert:false)
+
+The catch is that the new method requires a longer history filter, which can contribute to ghosting, especially when shadow casters are moved.  Hopefully this can be improved in the future.
+
+**Performance**
+
+Distance Field AO has been made much faster such that it can run on medium spec PC and PS4.  It also has a much more reliable cost now, so that it's mostly constant, with a slight dependency on object density.  
+
+In cases with a static camera and mostly flat surfaces, the new method is **1.6x faster**.  In complex scenes with foliage and a fast moving camera, the new method is **5.5x faster**.  The cost of Distance Field AO on PS4 for a full game scene is 3.7ms.
+
+**Technical details**
+
+The main optimization is the use of a global distance field, which follows the camera around.  This is created by compositing all the usual per-object distance fields into a few volume textures centered around the camera, called clipmaps.  Only newly visible areas or those affected by a scene modification need to be updated, so the composition doesn't cost much.
+
+Visualization of the texel size of the clipmaps - each clipmap is a different color
+
+![image alt text](image_24.png)
+
+The global distance field is lower resolution than the object distance fields, so it can't be used for everything.  When computing cone traces for sky occlusion, the object distance fields are sampled near the point being shaded, while the much faster global distance field is sampled further away. 
+
+This is a visualization of ray tracing the global distance field vs object distance fields.  Surfaces in the global distance field have become blobby and thin objects disappear.
+
+![image alt text](image_25.png)
+
+## New: Content Browser Advanced Search
+
+The Content Browser now supports advanced search syntax.  Unreal Engine 4 is used on many large projects, including some with teams comprising hundreds of developers, and content libraries with hundreds of thousands of assets!  These new search features improve the workflow for finding your content on projects of any size.
+
+![image alt text](image_26.png)
+
+This augments the existing simple search syntax by providing a way to match key->value pairs from the asset metadata, as well as the following special keys:
+
+* Name - Tests against the asset name.
+
+* Path - Tests against the asset path.
+
+* Class (alias: Type) - Tests against the asset class.
+
+* Collection (alias: Tag) - Tests against the names of any collections that contain the asset.
+
+## New: Collection Improvements
+
+This release includes many changes to Collections that allow you to help you manage assets in your game such as collection nesting, smart collections and tagging functionality.
+
+### Nested Collections
+
+* Your collections of assets can now be organized into a hierarchy.
+
+![image alt text](image_27.png)
+
+* Child collections can be created by choosing the "New…" option from a given collections context menu. Alternatively, collections can be parented via drag-and-drop.
+
+![image alt text](image_28.png)
+
+### Dynamic Collections
+
+* Dynamic collections leverage the power of the Advanced Content Browser Search Syntax to allow you to create and share Content Browser filters.
+
+![image alt text](image_29.png)
+
+* Anything you can do with the Content Browser text search can be saved as a dynamic collection, including referencing other collections (both static and dynamic).
+
+* You can create a dynamic collection using the "Save" button to the right of the Content Browser text filter.
+
+	![image alt text](image_30.png)
+
+### Quick Asset Management
+
+* You can now quickly add an asset to many collections at once using the new Quick Asset Management check boxes.
+
+![image alt text](image_31.png)
+
+* These are available within the collections view itself, or alternatively, via the context menu of any selected asset(s).
+
+![image alt text](image_32.png)
+
+### Improved Robustness
+
+* Collections will now automatically follow any redirectors that they contain, which avoids the issue of an asset seeming to vanish from a collection when it was renamed or moved. Additionally, collections are now considered when fixing up (and ultimately removing) redirectors via the Content Browser.
+
+### Improved Status Reporting
+
+* Each collection now has a little status indicator at the right hand side of its row item. This changes colour depending on the current state of the collection, with each colour meaning:
+
+    * Red - The collection is in some kind of bad or read-only state. See the tooltip for how to resolve the issue.
+
+    * Orange - The collection is not up-to-date with the latest version in source control.
+
+    * Blue - The collection has local unsaved changes. This may happen if a save or auto-checkin fails, or if the collection had its redirectors followed.
+
+    * Green - The collection is non-empty and up-to-date.
+
+    * Grey - The collection is empty and up-to-date.
+
+### Improved Asset Tooltips
+
+* The static collections for a given object are now shown as part of its tooltip in the Content Browser.
+
+![image alt text](image_33.png)
+
+## New: Plugin Creation Wizard
+
+Its now easy to create new C++ plugins for Unreal Engine 4!  Check out or new plugin wizard.
+
+![image alt text](image_34.png)
+
+* The Plugins browser (accessible from the Window menu) now lets you create a new plugin!
+
+* You'll be able to give your new plugin a name and select from a few preset plugin types (Blank, Toolbar, Standalone Window).
+
+* This will create all of the initial files you need for your plugin to compile and be loaded into the editor
+
+## New: Enhanced Curve Editor
+
+We've made many improvements to the Curve Editor to make it easier to manipulate curve keys and tangents.
+
+![image alt text](image_35.gif)(convert:false)
+
+* Middle mouse now also moves selected keyframes or selected tangents
+
+* Added support for selecting and manipulating multiple tangents
+
+* Selection now respects control (toggle) and shift (add to existing selection)
+
+* Holding down shift will restrict movement to one axis
+
+* Navigation: alt-middle mouse pans the view and alt-right mouse zooms the view
+
+* Added tangent display options: Always show tangents, show tangents for selected points, never show tangents
+
+* Added pre-post infinity extrapolation options for curves
+
+* Make adding keys more predictable
+
+    * Keys are only added under the mouse when there is only one curve to edit
+
+    * When multiple curves are available clicking on a curve will add a key inline on that curve, and clicking off of the curve will add keys to all curves
+
+    * Changed the text of the menu item to make it clear what's going to happen
+
+## New: Multiple Return Nodes in Blueprints
+
+A function can now have multiple return nodes. As soon as execution reaches a return the function's execution is terminated. This can be used to return early from a Blueprint function, or to branch and return differing values.
+
+![image alt text](image_36.png)
+
+## New: Construct Custom Objects in Blueprints
+
+In Blueprints, you can now spawn baser object types (not just actors/components). The Construct Object From Class node takes a class and creates a new object of that type, similar to Spawn Actor From Class (except for types that are not an actor).
+
+![image alt text](image_37.png)
+
+* The "Outer" input will serve as the new object's owner, which controls the lifetime of the created object.
+
+* For actor classes, you'll still use the Spawn Actor From Class node.  And for widgets, you'll use the Create Widget node.  In a future release we may try to combine these different features.
+
+## New: Blueprint Class Defaults
+
+In Blueprints, it is now possible to access a class's defaults through the new "Get Class Defaults" node. This can be helpful when working with data-only Blueprints, giving you access to those Blueprints' values without having to create an instance.
+
+![image alt text](image_38.png)
+
+* To keep the output pins from growing too unwieldy, use the node's details to show/hide specific defaults
+
+## New: Level Blueprint Communication
+
+Level Blueprints can now be communicated with through interfaces. They can adopt and implement interface functions just as Class Blueprints can. After an interface has been added to a Level Blueprint, calls can be made through interface message nodes (using a streaming level as the target).
+
+![image alt text](image_39.png)
+
+* Use the "Get Streaming Level" node to get the level to communicate with
+
+    * Since you have to use "Get Streaming Level" to target a specific level, this will only work for sub-levels
+
+## New: Optimized Math Expression
+
+We've optimized the Math Expression node so that it is more performant than regular nodes. It is now roughly twice as fast to use a Math Expression node, when compared to a series of operation/function nodes representing the same expression.
+
+![image alt text](image_40.png)
+
+## New: Blueprint Asset IDs
+
+We've made it possible for Blueprints to load assets on-demand instead of right away.  There are two new data types exposed to Blueprints: "Asset ID" and “Class Asset ID”. An Asset ID represents a loaded or unloaded asset, and can be resolved into an object using the “Resolve Asset” node. If the asset was not loaded, then the “Resolve Asset” node returns an invalid object. You can load the asset using the “Load Asset” node. Similarly, a “Class Asset ID” will resolve into a class, and can be loaded with the “Load Asset Class” node. 
+
+![image alt text](image_41.png)
+
+## New: Montage Element Timing Interface
+
+​We've added a new panel to the Montage editor to help you control the order that Montage events will fire when your animations execute at runtime.
+
+![image alt text](image_42.png)
+
+## New: Non-Linear Animation Blending
+
+We now support many different blending functions between animation poses and transitions!
+
+![image alt text](image_43.gif)(convert:false)
+
+Previously, only linear or cubic blending was supported.  Now you can select from many different blending functions:  Linear, Cubic, Hermite Cubic, Sinusoidal, Quadratic, Circular, Exponential and Custom.  For most types, you can separately control whether you want easing on either the "in" or the “out” of the curve.
+
+The "Custom" option requires a “Curve Float” asset.  After creating one, you can assign it in the details panel for the relevant nodes:
+
+![image alt text](image_44.png)
+
+![image alt text](image_45.png)
+
+This will allow you to specify just about any blend you want to perform. The length of the curve will be normalized and scaled down to meet the specified blend time and any values that are outside of the range 0-1 are clamped back to fit (this limitation will hopefully be lifted soon, watch this space). If a custom curve is not specified the system will fall back to a linear blend.
+
+![image alt text](image_46.gif)(convert:false)
+
+## New: Bone-driven Animation Controllers
+
+This allows a "Driver" bone to dynamically affect the motion of one or more “Driven” bones.  This is great for characters that have accessories attached!  You can avoid geometry intersecting during an animation, even when a lot of blending is used.
+
+![image alt text](image_47.gif)(convert:false)
+
+In the above example, the attached accessory (green) has no authored animations, and is being driven in two axes as a function of the character's thigh bone.  This is all calculated at runtime, so a blend of animations ends up working quite well here, even without hand-authored tweaking.
+
+You can either set the "Driver" value directly with a multiplier, remap it into a whole new range, or just use a Curve Asset to drive the motion.  Using a curve is usually the best approach as it lets you define the response naturally and interactively tweak points/tangents, seeing the changes in real-time.
+
+Here are the new settings for bone-driven animation controllers:
+
+![image alt text](image_48.png)
+
+## New: Animation Transition Rules
+
+This release contains improvements to animation blending trees, particularly to how transition rules work.
+
+### Better handling of animation Asset Overrides
+
+Nodes like Time Remaining in transition rules now correctly respond to changes in animation length due to overriden animations in Child Animation Blueprints.  This means that animations referenced by transition rules no longer have to be the same length in every child Animation Blueprint.
+
+### Referencing the most relevant animation player
+
+In order to make state machines more maintainable, you can now use a new class of getter that always picks the highest weighted animation from the source state, rather than targeting a specific animation that might change in the future.  The following nodes are now available:
+
+![image alt text](image_49.png)
+
+### Custom Blend Graph Improvements
+
+Additional information is now exposed in Custom Blend graphs to allow them to be used for more than just a one-shot custom transition animation.  The following nodes that provide information about the corresponding transition node and source/destination states are now available in custom blend graphs:
+
+![image alt text](image_50.png)
+
+You can also use getter methods that reference the asset players, just like in a transition rule.
+
+## New: Animation Curve Evaluation Change
+
+Animation Curve Evaluation used to happen in Update phase of animation, but now it is moved to Evaluate phase. This comes with following benefits. 
+
+1. The curve's weight is properly is evaluated with blended weights. 
+
+    1. Cache Nodes
+
+    2. Additive Nodes: Additive nodes will apply the delta of additive curve data from base to the current pose as animation bone transform does
+
+    3. Layered Node: This was a bit more tricky to make it work properly as you don't know which curve is influencing which part of joint. Right now it does provide option of how to blend curves. 
+
+        1. Max Weight : It choose the max weight of curves from blended poses
+
+        2. Normalize by Weight : It sums up all weights of blended poses for all curves and normalize to 1
+
+        3. Blend by Weight: It just sums up (weights * curve value)
+![image alt text](image_51.png)
+
+2. Curve Evaluation has moved to multi-threading if you're using parallel evaluation.
+
+3. Eventually this will allow us to create a node that can drive curve from bone transform and vice versa
+
+4. However, if you need a curve that hasn't been evaluated (if you're ticking but not evaluating), you won't get the up-to-date data. Curve data is treated as bone transform now as opposed to treated as notifies before. 
+
+## New: Animation Asset Metadata Support
+
+We now support Metadata that can be added to animation assets.  The metadata is a Blueprintable class derived from the Anim Meta Data class.  This allows you to add custom metadata to animation assets.  Anim Sequence, Anim Montage, Anim Composite, and Blendspace are all supported. 
+
+![image alt text](image_52.png)
+
+You can query those data from an animation asset in C++ by calling the GetMetaData() method, or by using GetSectionMetaData() on Anim Montages.
+
+​
+
+## New: Sound Quality Levels
+
+A new feature has been added that gives sound designers the ability to manage audio memory usage and (to some extent) performance for lower end machines/devices.
+
+The available Sound Quality Levels are defined in the Audio section of the Project Settings. Each Quality Level can currently specify the maximum number of audio channels created by the audio device (though some platforms, notably Android, can still limit that value a lower number).
+
+![image alt text](image_53.png)
+
+In a Sound Cue you can place a Quality Level sound node that will provide an input pin for each Quality Level defined in the Project Settings and only the branch connected to the active Quality Level will be executed.  This can be used to reduce the number of variations that are loaded or it could be used to specify reduced fidelity wave files.
+
+![image alt text](image_54.png)
+
+When playing a standalone game, the Quality Level is specified via the Game User Settings and only the Sound Waves that are needed for that Quality Level will be loaded in to memory. The default value can be set per game or platform via the appropriate GameUserSettings.ini file and games can allow setting the value via their settings UI.  Currently, the game will need to be restarted for the quality changes to take effect.
+
+For Play in Editor you can specify which Quality Level to use via the Level Editor Play Settings. Changing this value will work while the Play session is active and new sounds will use the current Quality Level.
+
+## New: Custom Audio Attenuation Curves
+
+You can now define your own custom curve for attenuating sounds rather than simply being limited to the built in algorithms.
+
+In the attenuation settings set the Distance Algorithm to "Custom" and you will see the custom curve section appear.  You can then either specify an external curve asset from the content browser or define a custom curve directly in the details panel yourself.
+
+![image alt text](image_55.png)
+
+## New: Actor Tick Intervals
+
+You can now specify an interval at which a tick function will be executed allowing you to reduce the overhead of per-frame tick execution without the overhead or complexity of timers.
+
+This is exposed to end users via the Actor and Component default properties and for programmers it can be set in C++ similarly to any other Tick Function property.
+
+![image alt text](image_56.png)
+
+A few implementation details to be aware of.
+
+* All Tick functions execute the first frame they are registered, at which point they will wait the tick interval until their next tick.
+
+* Changing the tick interval of an already ticking function will not change when the next scheduled tick will occur. If you wish the change to take effect immediately you should disable the tick, change the interval, and then re-enable the tick.
+
+* Tick functions will occur as close as possible to the specified interval by providing a credit for any overrun when scheduling the next tick.  For example if a function is specified to occur every 1 second and it has been 1.1 seconds since the last tick at the point the tick function is executed it will be scheduled to occur once another 0.9 seconds have passed.  However, no tick function will execute more than once in a frame even if a multiple of its interval has occurred during the frame.
+
+* DeltaTime for tick functions remains the same regardless of tick interval.  It is not the time since the last time this function ticked, but always the delta time of the frame in question.
+
+* An Actor's Custom Time Dilation does not apply to the Tick Interval.  The frequency with which the tick is dispatched remains constant regardless of whether the individual Actor has its own Time Dilation, that Time Dilation is only applied to the DeltaTime passed in to the tick function.
+
+## New: Actor Encroachment Detection
+
+Encroachment detection is also more robust and works in more situations.  We've added new customization options for encroachment, too.
+
+![image alt text](image_57.jpg)
+
+The "Spawn Even If Colliding" option ("bNoCollisionFail" in C++) when spawning an actor has extended to support four methods of handling collision at the spawn point:
+
+1. Spawn it there anyway.
+
+2. Spawn at suitable location nearby if possible, else just spawn it there anyway.
+
+3. Spawn only if a suitable location can be found nearby, else don't spawn.
+
+4. Just don't spawn it.
+
+Existing blueprints will update automatically to choose the correct corresponding choice. Some code has been deprecated and will need to be manually updated.
+
+Actor classes have a new blueprint-exposed member called Spawn Collision Handling Method that controls how to resolve spawn collisions by default. This can be overridden by the setting on the Spawn Actor node or function call.
+
+For actors with a Movement Component, it is assumed the updated component is the primary collision shape and will be the only component tested for encroachment. For other actors, all components with appropriate collision settings will now be checked. Be aware this may result in some situations where actors fail to spawn after upgrading -- double check your collision settings, handling method, and spawn transform if this occurs.
+
+## New: Post Process Blending
+
+Post process "Blendables" are used to transition post process effects between different states. We've made improvements to this system.
+
+You can now set a blending weight value for each Blendable. This blending weight can be modified at runtime in either Blueprints or C++ code, so you can do some interesting things with post process transitions.
+
+You're also now able to define post process settings "inline" right within the post process volume, or you can opt to create a standalone Blendables asset that you can share between many objects!  
+
+![image alt text](image_58.png)           ![image alt text](image_59.png)
+
+Also, we've made it easier to define your own custom blendable settings. You no longer have to modify the engine's built-in PostProcessSettings structure, but instead can define new structures with settings, such as the Light Propagation Volume Blendable settings above.
+
+Finally, here's an example of how you'd use Blueprints to edit the blending weight value for a Blendable at runtime:
+
+![image alt text](image_60.png)
+
+## New: Runtime Asset Cache
+
+Some games need to support storing local copies of assets that may not have shipped with the game. For example if you have a lot of procedurally created assets, or assets that come from a server (such as advertisement graphics.) The new Runtime Asset Cache is a generic asset cache system that can be used to persistently store data generated during runtime in configurable asset buckets.
+
+![image alt text](image_61.png)
+
+The Runtime Asset Cache has the following characteristics:
+
+* Multiple buckets, separately storing assets of types intended by user, e.g. Character Image cache and Advertisement Image cache.
+
+* Each bucket has separately configurable size.
+
+* Asset versioning, allowing to rebuild cache entry when it becomes obsolete.
+
+* Synchronous and asynchronous cache querying.
+
+* When cache grows larger than specified size, oldest entries are removed from it.
+
+## New: Volume Decals (Experimental only!)
+
+3D function field rendering is an alternative to polygons for describing geometry. It is better-suited to solving some hard rendering problems such as shadows, Boolean operations, smoothing, ambient occlusion, ray tracing, and deformation. 2D and 3D distance functions are already in use in various places in the engine, like fonts, lightmaps, and distance field shadows/ambient occlusion. The new Volume Decals feature enables a distance function to be rendered as opaque objects into the GBuffer, using all of the information normally available to deferred passes, including ambient occlusion, screen-space reflections, lighting, and depth of field. The bowl in the following image is a Volume Decal.
+
+![image alt text](image_62.png)
+
+The object was created (see image below, from left to right) 
+
+* by defining a sphere with the distance to the decal center and a texture projection
+
+* defining a sphere another sphere and taking the maximum of both (boolean intersection)
+
+* moving the spheres to adjust the shape
+
+* adding another sphere and changing the material (boolean subtraction)
+
+* adjusting the shape again
+
+![image alt text](image_63.png)
+
+This is how the distance function for the sphere is defined:
+
+![image alt text](image_64.png)
+
+The feature is unfinished. Shadow casting, normal map support, and indirect lighting are missing, and performance and quality optimizations have not been made. In addition, the feature will always be costly on the GPU (distance function needs to be evaluated maybe hundreds of times for each pixel). We will polish it further by adding material functions, shadow casting, documentation, etc., but it might never be fast enough for real-time use.
+
+## New: UE4 Documentation Tags, Version and Skill Level
+
+* Documentation pages now have skill level, engine version, and category tags which are displayed at the top and bottom of pages that enable you to find other pages with matching tags.
+
+![image alt text](image_65.png)
+
+* Added ability to filter the sitemap by skill level and/or engine version. The sitemap is now an expandable tree view. When filtered, the tree is expanded to show matching items while non-matching items appear grayed out
+
+![image alt text](image_66.png)
+
+* Skill Level, Engine Version, and Tags on pages also now link to the sitemap and automatically filter it.  Clicking a tag opens the sitemap filtered to only show pages with that tag.
+
+![image alt text](image_67.png)
+
+* Links to tutorials (how-to's, quick starts, etc.) display the skill level and engine version of that tutorial:
+
+![image alt text](image_68.png)
+
+## New and Updated UE4 Documentation
+
+* [Gear VR Global Menu](https://docs.unrealengine.com/latest/INT/Platforms/GearVR/GlobalMenu/index.html)
+
+* [Reducing APK Package Size](https://docs.unrealengine.com/latest/INT/Platforms/Android/ReducingAPKSize/index.html)
+
+* [Mobile Content Scale Factor](https://docs.unrealengine.com/latest/INT/Platforms/Mobile/Performance/index.html#mobilecontentscalefactor)
+
+* [Blueprints Best Practices](https://docs.unrealengine.com/latest/INT/Engine/Blueprints/BestPractices/index.html)
+
+* [Guidelines for Programming for Blueprints](https://docs.unrealengine.com/latest/INT/Engine/Blueprints/TechnicalGuide/Guidelines/index.html)
+
+* [Gameplay Debugger](https://docs.unrealengine.com/latest/INT/Gameplay/Tools/GameplayDebugger/index.html)
+
+* [Visual Logger](https://docs.unrealengine.com/latest/INT/Gameplay/Tools/VisualLogger/index.html)
+
+* [Network Profiler](https://docs.unrealengine.com/latest/INT/Gameplay/Tools/NetworkProfiler/index.html)
+
+* [Behavior Tree Quick Start](https://docs.unrealengine.com/latest/INT/Engine/AI/BehaviorTrees/QuickStart/index.html)
+
+* [Behavior Tree Node Reference](https://docs.unrealengine.com/latest/INT/Engine/AI/BehaviorTrees/NodeReference/index.html)
+
+* [Environment Query System Quick Start](https://docs.unrealengine.com/latest/INT/Engine/AI/EnvironmentQuerySystem/QuickStart/index.html)
+
+* In-Game Analytics Provider Pages
+
+    * [Apsalar](https://docs.unrealengine.com/latest/INT/Gameplay/Analytics/Apsalar/index.html)
+
+    * [File Logging](https://docs.unrealengine.com/latest/INT/Gameplay/Analytics/FileLogging/index.html)
+
+    * [Flurry](https://docs.unrealengine.com/latest/INT/Gameplay/Analytics/Flurry/index.html)
+
+    * [Multicast](https://docs.unrealengine.com/latest/INT/Gameplay/Analytics/Multicast/index.html)
+
+## New: Asynchronous Real-time Audio Decompression
+
+Platforms that support real-time ogg-vorbis based audio decompression (PC, Xbox One, Mac, Android) now perform the decompression asynchronously when possible to avoid blocking operations on the main game thread. 
+
+## New: Shared Resources for Feature Packs and Templates
+
+We now support shared resources in templates and feature packs. This is achieved by allowing both templates and feature packs to specify other feature packs to include. This means that many assets that were common to all templates now only exist once. It also means that we no longer need to duplicate assets that are shared between the Blueprint and C++ versions of templates.
+
+Multiple target platforms are supported, so different assets can exist in a template or feature pack for mobile and desktop targets. That's useful when you want lower detail models on one platform but not another, for example.
+
+#3New: Improved HTML5 (Easy Setup, Amazon S3 Support, Networking)
+
+You no longer need to install any third party SDKs to start development with HTML5 and package your game for a web browser! To package, simply select HTML5 in the File -> Package menu. 
+
+Other improvements include:
+
+* Emitting compressed javascript which reduce the size of the builds drastically. 
+
+* A small web server is included with the game during packaging which can handle serving compressed files for local testing. 
+
+* HTML5 Networking is now supported! There is now a Unreal Build Tool configuration variable which can be used to create dedicated servers that support HTML5 clients. 
+
+* We now support uploading to Amazon S3 during packaging. Once the credentials for AWS are filled in, the build will be uploaded to S3 and the game can be served to browser directly from S3. 
+
+![image alt text](image_69.png)
+
+# Release Notes
+
+## AI
+
+* New: A dedicated AI asset category has been created and is available in Content Browser's context menu.![image alt text](image_70.png)
+
+* New: Added log errors when EQS (Environmental Querying System) queries take longer than 0.25 seconds (cumulatively). When the query completes we also log how long each test took to perform.
+
+* New: Added support for trace based projection in EQS generators.
+
+* New: Added Z offset for circle center(s) in the On Circle EQS generator that uses data bindings.
+
+* New: AI Controller's Path Following Component has been exposed to blueprint and can be tweaked via blueprint-created classes' default properties.
+
+* New: Blackboard assets have been made blueprint types to allow blueprint users to interact with them.
+
+* New: Exposed filtering and scoring options for multiple contexts of EQS query.
+
+* Fixed "Skip Item" in EQS not always showing results in debug view.
+
+* Fixed "Skip Item" acting like "Pass" and in some cases like "Fail" when handling filtering.
+
+* Fixed crash when unregistering the Pawn Actions Component when the Controlled Pawn is pending kill pending.
+
+* Fixed AI stimuli never expiring which resulted in AI never forgetting perceived actors.
+
+* Fixed crash in pathfinding batch EQS test.
+
+* Fixed crash on running a EQS query without valid options.
+
+* Fixed EQS pathing grid generator getting stuck on points inside geometry.
+
+* Fixed restoring EQS query asset with multiple options.
+
+* Fixed crash when doing a pathfinding batch EQS Query Test when test items are not considered valid.
+
+* Fixed "On Circle" EQS generator crashing if its query owner is not an actor.
+
+* Blueprint compilation and loading are much more efficient.
+
+### Behavior Tree
+
+* New: Added dynamic subtree injection support for behavior trees.
+
+* New: Added new restart mode for behavior tree tasks, they can now ignore restart which leads back into running the same task again. Use case depends on task type and its parameters: some of them work better with restarting every time (e.g. movement), some should finish their action (e.g. animation).
+
+* Blueprint based behavior tree tasks will now ignore Finish Execute and Finish Abort calls depending on state of task, fixed latent blueprint actions not being cleared when task finishes execution.
+
+* Fixed behavior tree decorators losing blackboard observers.
+
+* Fixed behavior tree restarts from blackboard decorator in On Value Change mode.
+
+* Fixed Behavior Tree tag cooldown decorator testing the current value of the tag cooldown end time against its duration. (Duration in the Tag Cooldown decorator should only be applied when we deactivate, not used for comparison.)
+
+* Fixed behavior tree's search update for branch with loop decorator.
+
+* Fixed bug with aux nodes losing all changes in memory block during discarded search.
+
+* Fixed crash on pasting composite decorator node in behavior tree editor.
+
+* Fixed drag & drop operation for composite decorators.
+
+* Fixed duplicated subnodes in behavior tree editor.
+
+* Fixed order of behavior tree nodes when creating a new node from existing pin.
+
+* Fixed rare crash on aborting behavior tree task connected to parallel node.
+
+* Fixed Rotate to Face task getting stuck on receiving abort event.
+
+* Replaced Has Reached Goal decorator with Is At Location.
+
+### Debugging Tools
+
+* New: Added "-LogNavOctree" command line parameter, to log Nav Octree debug geometry to visual logger, on stop logging.
+
+* New: Added new shapes to log with Visual Logger: mesh, convex polygon, navarea / pulled convex and Nav Octree dump from given bounding box.
+
+* New: Minor improvements to gameplay debugger.
+
+    * Added multiple logs selection to Visual Logger tool, to analyze all selected objects simultaneously.
+
+    * Added debug camera to gameplay debugger, to have a way to fly around (Tab key to switch between cameras).
+
+    * Added alternate keyboard bindings to gameplay debugger, for keyboards without numpad (can be enabled in gameplay debugger settings). It uses Alt + [regular number] shortcuts.
+
+* Added more data about EQS queries to vlogs.
+
+* EQS score for the winning item will no longer be set to 1 when picking a single item. That way, scores can be compared appropriately since none of the other scores have been normalized. They will all be on the same relative scale.
+
+* Fixes Gameplay Debugger's client server issues. Everything should replicate to clients correctly again.
+
+* Tweaked description width for EQS debug display from 200 to 312 to allow for slightly longer descriptions without overlapping the next column. Item Description Width, Item Score Width, and Test Score Width can now be tweaked in code (they are now variables rather than in-line values).
+
+### Navigation
+
+* New: Added a dedicated blueprint-bindable event to AI Perception Component that's being triggered for every newly sensed stimulus.
+
+* New: Added bUseNavAgentGoalLocation to pathfinding Reached tests for moving to Actors. (Sometimes it is useful to move towards / test against the actual actor's location, rather than its location projected into the navigation system.)
+
+* New: Added custom navigation export for skeletal meshes.
+
+* New: Added navigation agent selectors to navmesh bounds volume, defining zones per agent.
+
+* New: Added navigation updates for destructible mesh.
+
+* New: Added new option to navigation link, which allows connecting to cheapest area in snap radius.
+
+* A bug in editor-time navigation system, resulting in removing all the saved data in static navmesh on map load, has been fixed.
+
+* Abstract navigation data actor will no longer be visible in scene outliner.
+
+* Added timeout for path following's waiting state.
+
+* Fixed a bug in navigation system preventing collision-less shape components from affecting navmesh generation by other means (like nav areas).
+
+* Fixed bounds of "standalone" nav relevant components (without navigation parent), which caused them to be ignored during navmesh generation.
+
+* Fixed crash in crowd movement when AI is using direct path to location.
+
+* Fixed crowd simulation updates for moving without pathfinding.
+
+* Fixed Floating Pawn Movement component's issue resulting in teleporting AI following a path, ignoring Max Speed limits.
+
+* Fixed geometry projection for navmesh walking mode using cached location even if pawn moved far away on Z axis.
+
+* Fixed handling of custom navigation links loaded with map before navigation system is created.
+
+* Fixed issue with navigation system being stuck on Initial Lock in editor mode.
+
+* Fixed missing abstract navigation data (handling direct paths) in game.
+
+* Fixed missing navigation links with big snap height.
+
+* Fixed navigable collision export from convex elements of Body Setup not using element's local transform.
+
+* Fixed navigation relevancy and update order of custom navigation links in Nav Link Proxy actor.
+
+* Fixed Navigation System not handling properly deep Actor hierarchies. Now you can have Actors attached to Actors attached to Actors and all of them will affect navmesh generation just as expected!
+
+* Fixed updating Nav Relevant Components attached to an actor without collision component.
+
+* Navmesh's rendering component has been made transient and we made sure any navigation rendering component saved with old maps won't break anything.
+
+* Updated deprecation messages for Navigation System's "get random point" functions family to be more precise.
+
+## Animation
+
+* New: "Observe Bone" nodes for debugging animations.
+
+* These let you print out the current position of a bone (at that particular point in the animation tree) and have a couple of options for coordinate space and whether to display relative to the reference pose.
+
+* The observe bone nodes are meant for debugging.  They don't cost much but are not currently compiled out of a shipping build, so you should probably remove them when you are done working with the graph.
+
+![image alt text](image_71.png)
+
+* New: Added turntable feature for Persona.
+
+![image alt text](image_72.gif)(convert:false)
+
+* New: Added a "Remove this LOD" button to each LOD (after 0) in Persona's mesh mode.
+
+* New: Added an icon overlay to the 3D thumbnail for Animation Blueprint assets in the Content Browser.
+
+* New: Added 'Blend Out Trigger Time' to AnimMontages.
+
+    * Time from Sequence End to trigger blend out.
+
+    * <0 means using Blend Out Time, so Blend Out finishes as Montage ends.
+
+    * >=0 means using "Sequence End - Blend Out Trigger Time" to trigger blend out.
+
+    * Also added as an optional parameter to Play Slot Animation As Dynamic Montage.
+
+* New: Added compilation warnings to the Modify Bone skeletal control when it will do nothing (e.g., when all components are set to be ignored or no bone is selected).
+
+* New: Added event for Post Evaluation.
+
+    * This gets triggered after bone evaluation is done.
+
+    * If your mesh component has physics, it will trigger after bone is blended back to it.
+
+* New: Optimized animation evaluation by looping over the bones directly instead of indirectly.
+
+* New: Added flag to notifies to stop them triggering on dedicated servers.
+
+* New: Added friendly display names to Bone Modification Mode enumerants.
+
+* New: Added new node "Bind Cloth To Master Pose Component" that makes cloth on a slaved skeletal mesh component to reuse the simulation data from the master component, rather than processing the cloth separately.
+
+* New: Added new node "Is Section Using Cloth" to query whether a skeletal mesh section has cloth data.
+
+* New: Added reference pose toggle toolbar.
+
+    * Also now you can edit during reference pose.
+
+* New: Added sorting of Smart Name curves.
+
+* New: Added support for Montage-to-Montage synchronization.
+
+    * A montage can only have a single leader. A leader can have multiple followers.
+
+    * Synchronization is performed once before any of the leader or followers are updated in a given frame.
+
+    * Essentially, synchronization is performed one frame late, so tick order between montages is not an issue.
+
+* New: Allowed skeletal controls to specify the coordinate space that they would prefer the gizmo to be displayed in by implementing GetWidgetCoordinateSystem.
+
+* New: First pass of the dependency preloading system.
+
+    * Uses asset information created during cooking to determine the dependencies for all packages and make them load BEFORE the requested package. When using pak files, the pak file will be ordered to match this load order, giving much reduced seek times.
+
+    * Currently disabled by default, but can be enabled with "s.PreloadPackageDependencies 1" in the console.
+
+* New: Look At Control now supports look axis in local space.
+
+* New: Started moving animation nodes out of Engine into a new AnimGraphRuntime module.
+
+* New: Support Append Frames in the begin/end of the AnimSequence.
+
+    * Also fixed to rearrange notifies and curve keys.
+
+* New: Wrapped re-entrancy checks during AnimationBlueprint node Update/Evaluate in DO_CHECK tests so that they are compiled out in shipping builds.
+
+* Animation and joint name now supports UTF-8 when importing.
+
+* Animation Record gives default path and available name.
+
+* Changed MasterPoseComponent serialization so that you can use it in a construction script.
+
+* Changed animation stat names to reflect whether they are running on the main game thread or a worker thread.
+
+* Changed new NativeNotify menu to display the localized class name.
+
+* Fix for crash in BlendSpaces when referenced AnimSequences have been deleted.
+
+* Fixed a bug where the latent action manager wasn't correctly removing actions.
+
+* Fixed a crash due to incorrect current time when switching a sequence player node's animation to one that is shorter than the current one.
+
+* Fixed a crash when compiling a malformed Animation Blueprint (containing a missing node caused by deleted code or a disabled plugin) and added a graph validation pass prior to compilation.
+
+* Fixed adding an extra frame when importing animations.
+
+* Fixed an issue where the message log was not displayed/refreshed when using the "click to recompile AnimBP" prompt rather than the Compile button.
+
+* Fixed client/server disagreements when playing root motion animations.
+
+    * When a Montage ends, some left over root motion could stick around and be used the next time around, potentially creating positional differences between client and server.
+
+    * If sending a DualServerMove when root motion starts, server can update root motion for a move that didn't use root motion on the client yet.
+
+* Fixed crash during animation blueprint compilation due to incorrect parent class being accessed during skeleton recompile.
+
+* Fixed crash during skeleton class recompile when structurally changing some animation blueprints.
+
+* Fixed crash in blend node with when evaluating a LOD'd skeleton (bones were stripped out.)
+
+* Fixed crash when enabling Show Uncompressed Bones.
+
+* Fixed crash when pausing the game and using "Tick when paused" on a SkeletalMeshComponent.
+
+* Fixed crash when you mix additive and non additive animations within montage.
+
+* Fixed crashes relating to animation notifies that had invalid slot indices when linked in montages.
+
+* Fixed divide by zero when importing zero-length animations. (Contributed by [REDACTED_PERSON], pull request 1148.)
+
+* Fixed incorrect state machine weights when using a custom blend on a transition.
+
+* Fixed montage linkable elements sometimes losing segment links when external animations had changed.
+
+* Fixed non-human-controlled characters moving much more than they should while running root motion on dedicated servers with AnimUpdateRateOptimizations and RootMotionFromMontages ON.
+
+* Fixed 'phantom' warnings that could show up the first time an Animation Blueprint was opened in an editor session but would go away when you recompiled interactively.
+
+* Fixed reimporting breaking RetargetBasePose.
+
+* Fixed variable initialization order bug that caused evaluation skipping on SkeletalMeshComponents to not function properly.
+
+* If you import a mesh that adds bones to a skeleton, but don't save the skeleton, it could cause a crash. Added the ability to recover these bones in the editor.
+
+* Minor style fixes for padding/fonts of bone references and the alpha scale/bias of a skeletal control.
+
+* Removed WITH_EDITOR define around GetActiveSocketList in QuerySupportedSockets so behaviour of QuerySupportedSockets is consistent between editor and non editor builds.
+
+* Reworked animation track segment looping logic to reduce floating point error problems.
+
+* When Skeleton detects inconsistent hierarchy, it will try fix up. You're going to notice merge bone effort with error message.
+
+    * When this message comes up, please resave skeleton.
+
+* While piloting an object, 'move to' object button functionality works properly.
+
+## Audio
+
+* New: Added ability to set the volume headroom per platform with a new Engine.ini field.
+
+* New: Improved Gameplay Statics sound API:
+
+    * New Spawn at Location functions return an Audio Component that allows manipulation of the played sound similar to the Attached version.
+
+    * Play at Location still does not return the audio component and should be used for one-shot sounds.
+
+    * Rotation parameter added to API calls that specify location or attachment. This causes non-spherical attenuation shapes to face the intended direction.
+
+    * Play Attached renamed to Spawn Attached.
+
+    * Sound and Dialogue APIs are now consistent.
+
+* Fixed volume differences on the Mac platform due to splitting and merging audio streams when using reverb, EQ, or radio effects.
+
+    * Added low-pass filter for voice audio units to closer match PC audio.
+
+* Fixed to use proper linear-to-dB volume equation on the Android platform.
+
+* Fixed a Core Audio API crash for Mac when the number of active sounds exceeded the number of available audio units.
+
+## Automation
+
+* New: Basic lighting tests have been separated from the 'General Editor Test'.
+
+    * They are now: "System.Promotion.Editor.Lighting.Place Scale Rotate", "System.Promotion.Editor.Lighting.Modify Properties", and "System.Promotion.Editor.Lighting.Duplicate and Copy Paste".
+
+* New: GenerateProjectFiles.bat will now exit with a non-zero exit code (1) when any failure happens.
+
+* New: Material Editor promotion tests are now contained in the MaterialEditor module. Fully exported UMaterialGraph.
+
+* New: Particle system promotion tests are now contained in the Cascade module.
+
+* New: The geometry automation tests have been separated from the 'General Editor Tests' automation test. It is now: "System.Promotion.Editor.Geometry Validation".
+
+* New: UAT and UBT will no longer wait for the global mutex, and will fail immediately instead. Previously, each instance would wait in line serially, hiding the fact that the command was not actually running.
+
+* Fixed UAT base path being invalid if the UE4 workspace is mounted directly onto a drive letter (e.g. C:\Engine, D:\Engine, etc.).
+
+## Blueprints
+
+* New: Added Blueprint-exposed GetBoundingBox function to StaticMesh.
+
+* New: Added Equals and NearlyEquals functions for Transforms to Blueprint math library.
+
+* New: Added TimeSecondsToString function to convert float seconds to formatted string.
+
+* New: Added a node icon to identify editor-only properties referenced in Blueprint script.
+
+    * Properties like a Character's arrow component, or a TargetPoint's sprite.
+
+    * Fixed crashes that would occur from referencing these properties in cooked builds (replaced with a logged exception instead).
+
+* New: Added editor .ini config settings that can be used to hide specific Blueprint menu entries.
+
+    * Use the "[BlueprintEditor.Menu]" section header to add exclusions under in the .ini file.
+
+    * In the .ini file, specific fields can be hidden like so: +BlueprintHiddenFields="/Script/Engine.Actor:GetTransform".
+
+    * In the .ini file, specific nodes can be hidden like so: +BlueprintHiddenNodes="/Script/BlueprintGraph.K2Node_IfThenElse".
+
+* New: Added Get Map Range Value Unclamped and renamed Get Mapped Range Value to Get Mapped Range Value Clamped.
+
+* New: Added Invert Transform blueprint node.
+
+* New: Added support for expressing a per-parameter Display Name using the markup UPARAM(DisplayName="Something neat") before the parameter declaration.
+
+    * This allows renaming parameters without using a K2ParamRedirect, and the use of names that are not legal C++ identifiers.
+
+    * e.g. static void DoSomethingAwesome(UPARAM(DisplayName="Awesome Param") float BoringParam);
+
+* New: Added the ability to filter out library functions and macros from the Blueprint context menu.
+
+* New: Added various improvements related to Select node. The node doesn't return a copy of input parameter, it returns a reference to the input.
+
+* New: An "InternalUseParam" meta tag for UFUNCTION has been added to support hiding function parameters from Blueprints. The pin for that parameter will be hidden and not connectable.
+
+* New: Blueprint Timer API improvements.
+
+    * Set Timer By Delegate now returns a Timer Handle to be used when interacting with the timer via other functions. Interacting with timers by handle is significantly more performant and enables the same timer delegate to be used for multiple timers simultaneously.
+
+    * Set Timer By Name, and all timer manipulation functions, by name and delegate have been deprecated.
+
+* New: Blueprint Timeline UI updates:
+
+    * Made individual tracks collapsable.
+
+    * Added an option to turn off the view synchronization for an individual track.
+
+    * General padding and layout tweaks.
+
+* New: Game/URL Options parsing nodes are now exposed to blueprints.
+
+* New: Can change a variable's type in the My Blueprint window by clicking on its type icon, or right-clicking to toggle whether or not it's an array.
+
+* New: Can drag and drop custom events from the My Blueprint window to add a call function for the event.
+
+* New: Can promote pins to local variables in Blueprint function graphs (or their child graphs).
+
+* New: Can set the Advanced Display flag on Blueprint member variables.
+
+* New: Made the boolean pin in Blueprints less subtle.
+
+* New: Fixed a crash that could occur with data-only Blueprints that explicitly called their parent's construction script.
+
+* New: Fixed error about missing World Context pin in blueprint library.
+
+* New: CreateEvent nodes are now more user friendly - they will use the current blueprint as the default context and no longer suggest that the user select an asset for context.
+
+* New: Improved details panel for nodes: Make, Break and Set Member. Properties are grouped in categories. Improvements for PostProcessSettings.
+
+* New: Improved the logic for when to reuse an existing graph tab versus open a new one.
+
+    * All subgraphs will now open in the parent when double-clicking on a node in a graph.
+
+    * Holding down Shift while clicking on an item in the My Blueprint tree or in the graph will force the creation of a new tab regardless of what is already opened.
+
+* New: Node "Set member in .." has an output pin for the post-modified struct.
+
+* New: Object Reference, Class Reference, Asset ID, Asset Class ID pin types have been spun off into a sub-menu on the item to reduce duplication of these items in the listview.
+
+* New: The output pin subtype for a GetClass function call node is now linked to the input pin's subtype.
+
+* New: Added support for multi-line editing to blueprint comment nodes.
+
+    * To enter a newline in the comment title use shift + enter.
+
+* New: Exposed NOR and NAND node support to blueprints.
+
+    * The nodes can be found in the math/Boolean section in the blueprint action menus.
+
+* "Color" name is not longer restricted.
+
+* Added "magnitude" keyword to vector length functions.
+
+* Adding or removing members from an editor created structure will no longer cause 'inappropriate outermost' warnings when blueprints that use that structure are saved without recompiling.
+
+* Adding pins to a MakeArray node will give the new pin a default values.
+
+* After context fail (Access None error) return data are properly cleared.
+
+* Attaching a bool to the "Select" nodes index pin will force the node to update it's pin names even when there are only two options.
+
+* Axis mapping value getters can no longer be called from other blueprints.
+
+* Blueprint Diff Tool no longer has buttons to switch between the components, graph, and defaults views. Navigation is now done via the list of differences.
+
+* Blueprint Editor context menu filtering is much more efficient.
+
+* Blueprint member variables of type component that are not generated by SCS now have the details options normally available to exposed variables.
+
+* Breakpoint is triggered on Access None error.
+
+* C++ Functions that are tagged as BlueprintImplementableEvent and are overridden as a function graph in Blueprints will display an error when attempting to drag and drop into the graph to inform the user that the function is not Blueprint Callable.
+
+* Call Function nodes will no longer change the function reference when duplicating the node or duplicating the entire Blueprint.
+
+* Can no longer split pins on Blueprint nodes if the struct does not have any valid properties.
+
+* Can now Make and Break the DateTime struct in Blueprints.
+
+* Can right click in collapsed graphs to add local variables that are part of the outer function graph.
+
+* Can right click on any node to search for references of it. Does initial search only in current Blueprint.
+
+* Can use unicode characters in Math Expression nodes.
+
+    * All unicode "letters" are now allowed for variable/input/output names.
+
+    * This does not extend to symbols and only the underscore symbol is still allowed.
+
+* Cannot name Timeline tracks the same as any of the native pins on the node.
+
+    * Any nodes previously broken by this functionality should work again.
+
+* Changing a pin type on a Select node will recombine all sub pins on the option and result pins before setting the new pin type.
+
+* Changing the Class of a "Spawn Actor from Class" node will no longer disconnect the result pin's links if they are still valid.
+
+* Child components in a Blueprint Class asset attached to a default scene root will no longer be lost after reparenting to a Blueprint Class that also has a default scene root.
+
+* Connecting function result node array pins to a reroute node will no longer prevent the node from compiling with all its output parameters.
+
+* Copying and pasting a K2Node_CallFunction with an interface pin for the self pin will no longer change the self pin into a normal Object Reference pin.
+
+* Deleting a custom event when there are no other events in the Blueprint will no longer disallow you to name other items placed in the Blueprint the same as the deleted custom event.
+
+* Drag and dropping a component from the SCS editor to a graph will no longer incorrectly report that the graph is outside of scope when it is a part of the same Blueprint.
+
+* Dragging and dropping variables from the My Blueprint tab will no longer default to assuming that the variable is read-only.
+
+* Duplicating function or macro graphs that contain nodes with Text pins will no longer incorrectly link the original and duplicated pins so that changes in one appear to affect the other.
+
+* Duplicating the "Add Scene Component" node while the transform pin is split will no longer make visible the transform pin on the duplicated node.
+
+* Editing a blueprint that is instanced in a hidden sub level no longer resets transforms on the hidden instances.
+
+* Ensure 'bExchangedRoles' is not set after serialization of an Actor-based Blueprint Class asset.
+
+* Examining the value of an enum array pin will show the enum value's display name.
+
+* Fixed a crash when compiling a Blueprint while intermediate graphs are open.
+
+* Fixed a crash while compiling a Blueprint with a validated Get node that is missing a property.
+
+* Fixed a crash when compiling a Blueprint while Find-in-Blueprints is searching all Blueprints.
+
+* Fixed a crash when vectors are compared in Math Expression using operator!=.
+
+* Fixed a potential crash during component registration when constructing instances of a Blueprint Class.
+
+* Fixed a potential crash on load during component template serialization for older Blueprint Class assets.
+
+* Fixed a potential crash when copying an array of USTRUCTs from one Actor instance to another.
+
+* Fixed an issue where deleting the scene root component in the Blueprint editor could lead to incorrect state of instances of the Blueprint class post-reconstruction.
+
+* Fixed crash executing cast to a class that is not loaded (presumably due to referencing a class in an unloaded module).
+
+* Fixed crash on reload of maps saved with constructed component instances that match archetypes by name but no longer by type.
+
+* Fixed a crash that would occur when editing an instanced variable nested within another instanced variable.
+
+* Fixed propagation of material swaps when editing "native" component templates (i.e. inherited from the C++ class) in the Blueprint Class editor.
+
+* Fixed socket browsing in the Blueprint editor for inherited component templates.
+
+* Fixed stale components persisting in a Blueprint Class asset after changing a native parent class's inheritance hierarchy.
+
+* Fixed vector axis input event nodes displaying a float output pin instead of a vector pin.
+
+* Fixed a bug that prevented new components from having the same name as a deleted one (prior to a compile).
+
+* Fixed a bug that was preventing copy, duplicate, and rename hotkeys from working in the Blueprint viewport.
+
+* Fixed a bug where certain keywords for Blueprint menu items were being ignored.
+
+* Fixed a bug where dirty Blueprints were not being compiled before you selected "Launch".
+
+* Fixed a bug where static functions were being hidden from Blueprints with Hide Categories metadata.
+
+* Fixed a bug with cyclic parent/child dependencies, where inherited components could be trashed when compiled.
+
+* Fixed a bug with parent/child cyclic dependencies, where the child's overridden property values would be cleared on load.
+
+* Fixed a crash after renaming local variables when other functions have local variables of the same name.
+
+* Fixed a crash that could occur when passing a default Blueprint Session Result object to the "Join Session" Blueprint node. Instead of crashing, the join will fail an execution will continue from the "On Failure" pin, as normal.
+
+* Fixed a crash that could occur when undoing, after compiling an Actor Blueprint with components.
+
+* Fixed a crash that would occur in specific cyclic dependency cases involving an Actor and Actor Component Blueprints.
+
+* Fixed a crash that would occur when hovering over the "Edit Blueprint" button, for Blueprint instances that had components added in their construction script.
+
+* Fixed a crash that would occur when loading a Blueprint that depended on a plugin that was disabled.
+
+* Fixed a crash that would occur when removing an event/function param in a Blueprint that was associated with the set GameMode.
+
+* Fixed a crash that would occur when removing a pin that had a debug watch applied to it.
+
+* Fixed a crash when deleting a map while it's level Blueprint was opened.
+
+* Fixed a crash when restarting level after streaming sublevels in.
+
+* Fixed a cyclical dependency bug involving macro libraries, where Blueprints could load with a false error.
+
+* Fixed a potential failure to restore component instance data during reconstruction of an Actor-based Blueprint Class instance.
+
+* Fixed an issue where it would appear that the Increment/Decrement macros would occur twice when using the output pin from the macro.
+
+* Fixed an issue where some component property values on a Blueprint class instance in the current scene could be reset as a result of recompiling the Blueprint class.
+
+* Fixed crash when Local Variable of array type is compiled.
+
+* Fixed crash when making changes with multiple components selected in the blueprint editor.
+
+* Fixed crash when selecting for "None" class type on the Spawn Actor From Class node after having already selected another class.
+
+* Fixed crash when struct value is copied from invalid object using Set node.
+
+* Fixed crash, when two empty User Defined Enums have the same name.
+
+* Fixed decrement int macro node to actually decrement.
+
+* Fixed duplication of components within the component tree in the Blueprint editor.
+
+* Fixed editing of Instanced Static Mesh Component and Hierarchical Instanced Static Mesh Component.
+
+* Fixed exposed APEX cloth functions not appearing in blueprint editors.
+
+* Fixed incorrect tooltips between "Convert to validated Get" and "Convert to pure Get" options on Get variable nodes.
+
+* Fixed issues that could occur if an existing Blueprint class default object or instance contained one or more invalid property values for native (C++-defined) component subobjects on load.
+
+* Fixed issues with graph node bubbles rendering incorrectly when there is nothing in the top slot
+
+* Fixed issues with hiding categories in Blueprints not always hiding all the expected nodes.
+
+* Fixed issues with node titles being out of date (e.g., animation skeletal controls) by forcing node titles to be updated when a property on a graph node is modified
+
+* Fixed issues with Target pins being invalid after reparenting a blueprint.
+
+* Fixed rare crash in the Blueprint diff tool that occurred when user dragged certain UI elements.
+
+* Fixed the cause of a 'Critical Failure' that could occur when replacing references to one blueprint with references to a different blueprint.
+
+* Fixed crash with Find-in-Blueprints, fallout from changes to FString.
+
+* Hidden pins generated from native functions will correctly use their default values.
+
+* Improved the Pin Type Selector filter so order of words and case sensitivity does not matter.
+
+* In the pin type selector dropdown, hitting the shift key will no longer move the cursor back to the start of the search box and hitting the down and up arrow keys will again navigate the list view.
+
+* Input action and key events will no longer be placeable in Blueprint functions.
+
+* Input Touch event node cannot be placed in function or macro graphs.
+
+* Local variables in Blueprint functions will now be set to the default values of the user defined struct their type represents.
+
+* Math expression functions without any parameters will still provide their end parenthesis in the Math Expression node.
+
+* Math Expression nodes will again correctly lookup by the display name of a function.
+
+* Math Expression nodes will no longer always appear as changed when doing a diff on a Blueprint.
+
+    * Improved diffing support to see how the expression has changed between revisions.
+
+* MathExpression nodes will not count function inputs as valid variables to access via getter nodes.
+
+* More strict restrictions related to WeakObjectPtr exposed to blueprint. No arrays of WeakObjectPtr.
+
+* On Blueprint enum pins, tooltips will appear when hovering over items in the drop down menu.
+
+* On the select node, setting the Index pin's type will always change the pin names of the option pins to "False" and "True".
+
+* Exposed FMath::FixedTurn to Blueprints. (Contributed by [REDACTED_PERSON], pull request 1254.)
+
+* Preventing in UK2Node::AutowireNode hidden pins from being auto-connected to passed in Wildcard pins.
+
+    * This prevents issues with wildcards connecting to the "self" pin of functions when dragging off them.
+
+* Proper error is generated when incompatible Blueprint pins are connected.
+
+* Recompiling a Blueprint with a Component parent class now correctly re-registers any instances of the blueprinted component in the level. Previously existing instances could disappear.
+
+* Removed a check that was asserting when compiling a Blueprint with invalid Call Function or Call Parent Function nodes.
+
+* Renaming a component and then adding a new component of the same type in the Blueprint editor will no longer result in broken defaults editing.
+
+* Select nodes using enums with hidden values will no longer incorrectly match up the Option pin connections to the improper enum value.
+
+* Selecting a pin type for a Select node will not fail the first time.
+
+* Setting default values of Text properties in the User Defined Struct or for Local Variable properties in Blueprints will correctly assign the value and preserve the value between instances of the editor.
+
+* Splitting pins on the Function Result node will no longer attempt to compile the split pins as result pins.
+
+* The "Get" pin(s) on "Set" variable nodes will now display the name of the pin if the pin has been generated by splitting a struct pin.
+
+* The details view in Blueprints of non-member variables (either from a parent class or an external member reference) will no longer show various detail uneditable options as editable.
+
+* The eye icon appears now for marking a user created component variable as editable in the My Blueprints window.
+
+* Tooltips will display for sub-menu items in the Pin Type Selector.
+
+* Tweaked the tooltip for "Expand Nodes" menu option to better describe what the action will do.
+
+* When dragging off a pin in a function graph, local variable nodes will no longer be filtered out of the context menu for placement if they are valid for connection.
+
+* When indexing all Blueprints, hitting "No" will correctly skip checking out all un-indexed Blueprints.
+
+* When loading a blueprint pin types that don't match are no longer disconnected, instead we rely on compile errors to prevent the potentially broken blueprint from executing bad code.
+
+* When signature of a function is updated, its instances in the Event Graph are properly updated.
+
+## Core
+
+* New: A plugin with the same name but of different versions can be placed in multiple game projects. This allows studios with multiple projects to utilize the same plugin at different revisions as long as the plugin is part of the game project.
+
+* New: Added module path caching to improve editor startup times.
+
+* New: Added pre-allocated memory pool which is freed when the system runs out of memory to properly handle such crashes.
+
+* New: Added 'stat slow' console command.
+
+    * Displays the game and render thread stats.
+
+    * Can take two optional parameters -ms=x -depth=n that specify the minimum stat duration and the displayed hierarchy depth respectively.
+
+* New: Added the ability to force garbage collection when worker threads are inside of garbage collection lock.
+
+    * In this case main thread will be blocked and wait until lock is over.
+
+    * This can be configured in Project Settings under Garbage Collection section.
+
+* New: Added two new parameters to AutomationTool's BuildCookRun command:
+
+    * IgnoreCookErrors - Ignores all cook errors and continues with deploying the project.
+
+    * FastCook - Passes fast cook parameter to the cooker, which is handled on a per-project basis.
+
+* New: Changed the file reader to use async I/O for a 1.6x+ speedup reading files on Windows based platforms.
+
+* New: CMake project generator now works on Mac and Windows.
+
+    * Pass -cmakefile to GenerateProjectFiles and use CLion or other IDEs to build UE4.
+
+* New: DirectoryWatcher can now optionally ignore changes in subdirectories (works on Windows and Linux; Mac implementation ignores this flag).
+
+* New: Exposed garbage collection settings to Project Settings under Garbage Collection category.
+
+* New: Exposed Streaming settings to Project Settings under Streaming section.
+
+* New: Improved missing import message to include the source asset name.
+
+* New: Improved the clean command in Unreal Build Tool. It will now properly clean up modules no longer referenced by the project.
+
+* New: Made zlib to also use standard Unreal memory allocators, which makes it easier to track them.
+
+* New: Memory stats will now also be reported when the engine crashes.
+
+* New: More verbose log message when the engine crashes because of a corrupted file.
+
+* New: Optimized transaction processing, saving unnecessary iteration when performing an ObjectRecord preload.
+
+* New: Improvement in garbage collection performance by implementing three micro-optimizations. (Contributed by [REDACTED_PERSON], pull request 1348.)
+
+* New: Slightly optimized object loading by looking up the object in memory before flushing async loading.
+
+* New: Small optimizations for structure serialization:
+
+    * Added native serialization for immutable structures.
+
+    * Added POD flags for various math classes.
+
+* New: The amount of logging information displayed in the editor has been reduced. The removed logging is now only written to the log file.
+
+* New: TAssetPtr now supports conversion from derived to base.
+
+* Fixed a memory leak in the FLinkerLoad in debug builds, never clearing itself from the list of live linker loaders.
+
+* Added a system for deferring linker deletion to fix crashes during async loading.
+
+* Added attribute serialization to the XML File writer.
+
+* Added code to force exit when a worker thread crashes.
+
+* Added message log error for source name collisions in redirector ini sections.
+
+* Always use the Engine directory instead of the Process directory to locate executables on Windows.
+
+* Changed "CreateExport: Failed to load Outer..." message to be emitted to the message log so it is better surfaced to the user when running the editor.
+
+* Changed the message about missing config files on startup to make it more accurate.
+
+* Disabled runtime crash processing for launcher builds to reduce memory footprint. Callstacks will still be processed by CrashReporter.
+
+* Editing "Project Settings" value now only edits a single value in the config ini file.
+
+    * With the exception of Arrays, which need a whole new implementation.
+
+* EPackageFlags & EObjectFlags can now be properly viewed in the debugger as a collection of enum entries, e.g. RF_Public | RF_Transactional, rather than just the raw numerical value.
+
+* Fixed crash reporting on Windows to make it more reliable.
+
+* Fixed "nearly equal" math comparisons not returning true when values are exactly equal and tolerance is zero.
+
+* Fixed a bug in TBB allocator that caused aligned allocations to overlap incorrectly.
+
+* Fixed Get Value Range function on Curve Base giving Min instead of Max. (Contributed by [REDACTED_PERSON], pull request 1314.)
+
+* Fixed a bug where game thread graph was not always enabled after loading a UE4Stats file in the profiler.
+
+* Fixed a bug where profiler could not perform 2 captures in a row through UI.
+
+* Fixed a crash after renaming a map or asset that couldn't have been renamed due to errors.
+
+* Fixed a crash in Online Async Task Manager Null when running the game with "-nothreading" on Windows.
+
+* Fixed a crash on exit after streaming levels in a cooked game.
+
+* Fixed an assert when checking for existing linker after async loading.
+
+* Fixed an intermittent deadlock when shutting down a game.
+
+* Fixed an issue where installed plugins would cause a content-only game to build when packaging.
+
+* Fixed Blueprint compilation errors after hot-reload.
+
+* Fixed build failures that could happen when user tried to add modules' dependencies.
+
+* Fixed crash on exit in packaged games caused by random order of destruction of some of the core UObject systems.
+
+* Fixed crash reporting and the engine not reporting correct windows version for Windows 8 and 10.
+
+* Fixed redirectors assets "Fail to load" messages that could appear in certain situations in the editor making redirector assets unusable and unfixable.
+
+* Fixed saving failure after hot-reload when blueprint node contained a reference to a class default object.
+
+* Fixed sporadic audio crash after hot reload.
+
+* Fixed the search order for detecting a target for game projects.
+
+* Fixed the search paths for plugins when running a packaged content-only game.
+
+* Fixed BSP texture align tools all being initialized to default. (Contributed by [REDACTED_PERSON], pull request 1136.)
+
+* Fixed a crash if list of delayed linkers includes multiple entries of the same linker. (Contributed by [REDACTED_PERSON], pull request 1197.)
+
+* Fixed broken UnrealBuildTool dependency file caches when the process is terminated unexpectedly.
+
+* Fixed UnrealBuildTool not being able to harvest environment variables when the user's Windows login name contains non-ASCII characters.
+
+* Fixed UnrealBuildTool not updating the makefile when a .Build.cs or .uproject file changes.
+
+* Fixed UnrealHeaderTool crash that was occurring in some rare cases if parsing step had other errors.
+
+* Fixed UnrealSwarm "Channel already registered" error to log the actual reason why exception has happened, cause until now it was very misleading.
+
+* Fixed UnrealSwarm connection problems, which were happening if coordinator and agent were on the same machine and coordinator address in agent's settings was set to non-loopback IP.
+
+* Incredibuild binaries are now sought in Path and not just Program Files.
+
+* Integrated some UnrealBuildTool performance fixes which were added in 4.7 but omitted from 4.8.
+
+* Made sure default subobject (component) properties are not reset to default on actor construction.
+
+* Made sure packages are not marked as dirty when they're being streamed in.
+
+* Minor performance optimizations for UnrealHeaderTool that shave a couple of seconds off of an invocation for a typical game project.
+
+* Moved linker serialization from UObject to UPackage to fix inaccurate UObject mem reporting.
+
+* Optimized object referencing system when collecting garbage while streaming.
+
+* Profiler will now load all frames, previously the last frame was ignored.
+
+* Quadrupled trunc length for callstack lines when forming the human readable string for windows callstacks.
+
+* Removed some redundant calls from FPaths::MakeStandardFilename to make it faster.
+
+* Rendering thread callstacks are once again correctly printed to the log for non-editor builds.
+
+* Save to default config no longer writes out duplicate entries.
+
+* Smoke test added to ensure platform consistency on opening files for append, and ability to seek said file handles.
+
+* The TAttribute no longer allows you to directly compare another TAttribute for equality.
+
+    * They've been deprecated, and corrected to only allow comparison against literal values instead of comparison attributes, which was causing implicit construction of attributes when comparing against a literal.
+
+    * There's a new IdenticalTo() function on TAttribute which can actually compare two attributes to see if they are really the same.
+
+* When bAutoCompleteWhenLoadingCompletes is false in the MoviePlayer, the movie correctly continues playing until prompted to stop after loading completes.
+
+* When queueing new async packages, async loading code will also look for existing packages in loaded packages to process to prevent infinite streaming times.
+
+## Editor and Tools
+
+* New: Added "Zip Up Project" button under File -> Package Project to grab necessary project files and package into a zip. For easily sending the project to other people excluding unnecessary files.
+
+* Smart Resizing for Transform Details UI:
+
+    * The vector and rotator widgets now support automatically rendering in a more compact mode as the details panel shrinks, making it possible to edit properties over a wider range of widths.
+
+![image alt text](image_73.gif)(convert:false)
+
+* Rotator UI Consistency:
+
+    * Rotators now display everywhere in a consistent order as XYZ, rather than being displayed as XYZ in the level editor but as Pitch, Yaw, Roll (YZX order) in Blueprints.
+
+    * In a Blueprint:
+
+![image alt text](image_74.png)
+
+    * In the level editor:
+
+![image alt text](image_75.png)
+
+    * All existing content is automatically updated and will work as expected, but you may notice wires crossing that didn't cross before since the order has changed.
+
+    * As a refresher (this information is also displayed in the tooltips):
+
+        * X is the roll value (rotation about the X axis, typically the forward vector, e.g. tilting your head sideways).
+
+        * Y is the pitch value (rotation about the Y axis, typically the right vector, e.g. looking up/down).
+
+        * Z is the yaw value (rotation about the Z axis, typically the up vector, e.g. turning to face left/right).
+
+* New: "Add Auto Connect To Server" option to the advanced Play In Editor options, defaulting to old value of true. If set to false it will launch multiple clients but they won't connect to each other, which is useful for UI testing.
+
+* New: Added "No to All", "Yes to All" and "Cancel" options to the "Failed to Merge Bones" dialog when importing skeletal meshes to avoid having to dismiss large a large amount of windows.
+
+* New: Added "Save All" button to Visual Logger tool.
+
+* New: Added a check to determine whether plugins are compatible when opening a content only project.  This is needed because a content only project will not attempt to build incompatible plugins later, users can now choose to disable them.
+
+* New: Added a support URL to plugins so that it is easy for users to know how to contact the creators of plugins for support.
+
+* New: Added file modified date and file size to map file asset tool tips. Allows sorting by these values in column view mode.
+
+* New: Added SkeletalMesh Resource Size in content browser.
+
+* New: Added support for multi-line text in the Note actor's text field.
+
+* New: Added support for positioning and always centering standalone play windows. Fields are added under Level Editor - Play advanced settings to position the standalone play window and turn always centering on and off.
+
+* New: Added the option to display the current Unreal Engine's version number in the project badge. This feature can be enabled in the Miscellaneous category inside of the Editor Preferences.
+
+* New: All integer property types are now supported in details panels instead of only int32 properties.
+
+* New: Allows play in editor windows to always be displayed on top of other windows, enabled/disabled under Level Editor - Play advanced settings.
+
+* New: Clarified the use of Server Game Options and Command Line Options in the tooltips for the Level Editor - Play settings.
+
+* New: Enabled engine/game plugin switch in plugin wizard. (Contributed by [REDACTED_PERSON], pull request 1145.)
+
+* New: FBX importing detection has changed.
+
+    * If you have animated curves, it will detect as Skeletal Mesh even if you don't have deformer.
+
+    * This is so that it is easier to import animation for rigid bodies.
+
+* New:"'Light Types” show menu:
+
+    * "Lighting Components" -> “Ambient Occlusion” affects both DFAO and SSAO now.
+
+    * Added "Screen Space Ambient Occlusion" show flag under ”'Lighting Features”.
+
+* New: Made the Blueprint Editor viewport camera speed a per-project editor user setting, which is saved between sessions.
+
+* New: Major optimizations to the BSP building process and BSP Undo.
+
+* New: Moved Plugins menu from Windows to Settings/Configuration.
+
+* New: Optimizations to setting the material or its texture transform on a BSP surface.
+
+* New: Placed cameras can now be looked through and piloted more quickly by using the camera dropdown in the level editor viewport.
+
+* New: Added hooks to Material Editor for plugin support. (Contributed by [REDACTED_PERSON], pull request 1168.)
+
+* New: Removed the close button from the Convert Project modal dialog, as it was returning an undefined reason code (and was not necessary as there was already a Cancel button).
+
+* New: Removed the SaveBrushAsCollision command and its bindings, as the builder brush is no longer a user-facing concept.
+
+* New: Renamed the editor preference "Automatically Hot Reload New Classes" to "Automatically Compile Newly Added C++ Classes" to convey that it implies a compile.
+
+* New: Text FrontendFilter now tries to match with components of the path name, as well as the asset name itself.
+
+* New: The "Maximum Texture Size" property now yields an error dialog when attempting to change it for non-power-of-two textures which do not have a padded Power of Two Mode.
+
+* New: The RenderToTexture Blueprint toolset now has its own set of options for EXR exporting and does not require interfacing with the 'High Resolution Screenshot" dialogue box.
+
+* New: Users can now edit PlayerStateClass and SpectatorClass for a game mode via the Project Settings menu.
+
+* New: The green flash on tutorials will not return in future sessions once the tutorial is started. It is no longer necessary to dismiss the tutorial manually or complete it.
+
+* New: When building the plugin category tree, plugins are now also added to their parent category until we reach the root, so that all plugins can be searched from the root level.
+
+* New: When using "-List" command with UnrealPak, the files will now be printed out in the order they appear in the pak.
+
+* Actor Merging tool now correctly assigns user specified lightmap channel to a merged mesh.
+
+* Actors generated by 'Hierarchical LOD' system can no longer be transformed using Details panel.
+
+* Added an ensure to the main frame to indicate when it is being incorrectly loaded in a -game build.
+
+* Added an option to not use MikkTSpace tangent generation for static mesh importing.
+
+* Added error message if a material expression parameter name is too long.
+
+* All PIE world objects are now marked pending kill when exiting a play in editor session to break any references that could prevent garbage collecting the PIE world.
+
+* Attempting to enter greater than the maximum allowed number of FName characters into a name property will now show an error message in the text box in a details panel instead of crashing.
+
+* Changed common and divisions of 360 rotation grid snap values.
+
+* Changed Plugin Creator to allow plugins to be created in subdirectories.
+
+* Decals now drag and drop properly when dropped on top of meshes.
+
+* Disabled Plugin Creator plugin in content only projects as it currently won't generate projects/compile any plugins that are created.
+
+* Ensured that dependencies on missing plugins can be removed from the .uproject file when loading the editor.
+
+* Ensured that Plugin Browser displays correct results after adding a new plugin.
+
+* Fix for crash while trying to hide the Tick property category within a class.
+
+* Fixed "Multiple Values" being committed to property text blocks when they lose focus or when the user presses enter.
+
+* Fixed a crash executing source control commands with the Subversion source control provider.
+
+* Fixed a crash that could occur during a Play In Editor session with two or more players, using multiple processes, and having the editor multiplayer mode set to "Play as Client".
+
+* Fixed a crash when a level that has a Matinee Actor is unloaded at runtime.
+
+* Fixed a crash when changing a struct that was being used by an open Data Table editor.
+
+* Fixed a crash when dragging junk text from an external program into the graph editor.
+
+* Fixed a crash when immediately closing Visual Studio after opening a source file via the editor.
+
+* Fixed a crash when importing multiple objects with the same name at the same time when one of them is an FBX file.
+
+* Fixed a crash which could occur when attempting to open a map in the content browser that is currently loaded in the world composition browser.
+
+* Fixed an issue where selecting a surface with no actor selected did not bring up the details panel specialization for the surface properties.
+
+* Fixed an issue where some of the shape textures in engine/content did not have textures.
+
+* Fixed BSP selected surface rendering in Geometry Mode.
+
+* Fixed camera roll being applied to level viewport cameras after ejecting from the player in a Play in Viewport session.
+
+* Fixed crash accessing the engine when attempting to play a slate sound after the engine has been shut down.
+
+* Fixed crash converting texture sample to parameter when MipValue is set to "Level" or "Bias".
+
+* Fixed crash copy and pasting material function call nodes while live preview is disabled.
+
+* Fixed crash ending Play in Editor if transactions made before starting PIE were undone while PIE was running.
+
+* Fixed crash in visual logger when you run the editor with an invalid .uproject.
+
+* Fixed crash when reimporting a mesh while editing collision in the static mesh editor.
+
+* Fixed duplicated brush actors not appearing in game.
+
+* Fixed editor viewports not remembering maximized state after shutting down in immersive mode.
+
+* Fixed error where during Mesh Paint mode using the Remove button would only remove the Vertex Colours from the first LOD level.
+
+* Fixed existing collision not being restored if there is no collision in the FBX file when importing over existing Static Meshes.
+
+* Fixed FBX import failure caused by incorrect import factory attempting to handle incoming assets.
+
+* Fixed geometry surface alignment tools for BSPs.
+
+* Fixed issue where BSP surface texture transforms could break when rebuilding geometry if their parent brush had been rotated about more than one axis.
+
+* Fixed issue where BSP surface transforms could get broken when applying them to geometry surfaces which are out of sync with the owner Brush (because Auto Build BSP is disabled).
+
+* Fixed issue where cancelling unsaved packages dialog when launching resulted in the Launch button permanently disabled.
+
+* Fixed issue where creating a child of an anim blueprint would result in a regular blueprint being created. Also fixed a similar issue with GameplayAbilityBlueprints.
+
+* Fixed issue where the AssetImportData for an asset was being lost when loaded.
+
+* Fixed issue where warnings during PIE were not being announced via a notification at the end of a PIE session.
+
+* Fixed issue with EyeDropper tool returning incorrect colour due to lack of gamma-correction.
+
+* Fixed issue with selection box rendering in the Curve Editor.
+
+* Fixed issues with redirector fixup, so that redirectors are deleted cleanly and referencing map packages are correctly fixed up.
+
+    * In addition, fixed an issue where sometimes it was necessary to perform redirector fixup twice.
+
+* Fixed issues with the Delete Assets dialog so that the operation is always performed cleanly where possible, and not erroneously requiring a "Force Delete" operation.
+
+* Fixed issues with the geometry Optimize operation so that it always works correctly, without sometimes first requiring a Triangulate operation.
+
+* Fixed lights not being hidden when changing the editor visibility via the World Outliner.
+
+* Fixed manually entered rotations in the editor sometimes not resulting in the correct object rotation (mainly when editing Pitch).
+
+    * Component transform editing now makes sure to only edit the single component (Yaw, Pitch, or Roll) that has been changed, and makes sure Quat<->Rotation conversion does not mess with the end result.
+
+* Fixed manually entered rotations in the editor sometimes not resulting in the correct object rotation (mainly when editing Pitch).
+
+    * Component transform editing now makes sure to only edit the single component (Yaw, Pitch, or Roll) that has been changed, and makes sure Quat<->Rotation conversion does not mess with the end result.
+
+* Fixed Matinee Editor Preview Viewport FOV lagging by a frame. (Esp. noticeable during cuts between cameras of different FOVs)
+
+* Fixed menus so that they set the widget to focus upon activation to the actual content widget, rather than the wrapper. This fixes issues with the Ctrl+Tab file selection shortcut.
+
+* Fixed mirror transforms on BSP / Volumes.
+
+* Fixed missing engine controls overlay for mac and linux in the lower-left part of the viewport.
+
+* Fixed numeric math expression input when using a combination of units in unit-based numeric entry boxes.
+
+* Fixed Paper2D JSON files no longer importing as Paper2D assets.
+
+* Fixed rare crash in Visual Logger during shutdown of engine.
+
+* Fixed rare editor crashes caused by proxy updates happening after proxy destruction.
+
+* Fixed rendering issue with the Collision Analyzer window when opened from a non-editor build via the console.
+
+* Fixed rendering of graph node popups (for example, dropdown boxes) in graph editors. They are now rendered and sized at 100% scale, regardless of the zoom of the graph itself.
+
+* Fixed some users being unable to create blueprinted components without visual studio installed.
+
+* Fixed Tag and Tag Query Editor windows on Mac.
+
+    * These windows were parented wrongly on Mac and couldn't be clicked.
+
+* Fixed Undo in Geometry Mode so that the geometry selection state is correctly restored.
+
+* Fixes to Visual Logger tool, to fix how we use this tool only with keyboard.
+
+* If you import a Skeletal Mesh with scale 0 of bone for reference pose, it will print warning.
+
+* Improved Data Table handling of User Defined Enums:
+
+    * The Data Table editor UI will now show you the friendly name of the enums in its grid panel (matching what the property panel shows).
+
+    * The Data Table import will now fallback to trying to import using a friendly name if no matching enum entry was found for the given name.
+
+* Improved instructions for what to do after creating new plugins with the wizard.
+
+* Improved performance when importing thousands of mesh sections from one FBX file.
+
+* Improved the localization support of the class wizard dialogs.
+
+* Made StartPIEGameInstance virtual to allow for game-specific PIE code.
+
+* New game windows now default to being positioned centered on the screen.
+
+* Optimizations to histogram graphs from Visual Logger tool.
+
+* Optimized pasting of actors into levels that contain a large number of actors.
+
+* Option to enable Hierarchical LOD is now more discoverable and can be found in 'World Settings' under the LODSystem category.
+
+* Performance improvements to OpenEXR exporter.
+
+* Placement mode can now correctly place new actors at locations occupied by foliage meshes.
+
+* Fixed small but annoying typos in ContentWrangleCommandlet.
+
+* Fixed Animation Sequence 'Reimport With new Source' context action.
+
+* Fixed issue with Animation Sequence reimport not respecting 'Default Sample Rate' flag set by the import UI.
+
+* Reduced ULevel::BuildStreamingData editor hitches from 600ms to 30ms.
+
+* Removed the map check warning "has simple collision but is being scaled non-uniformly - collision creation will fail" as a collision shape is created regardless of the fact that an object is non-uniformly scaled.
+
+* Renamed "Tick" category in UActorComponent to "Component Tick" and Actor tick to "Actor Tick" to differentiate between the two when they are both shown in a details panel.
+
+* Resizing a viewport no longer flushes texture streaming, causing a hitch.
+
+* Stopped ensure being triggered when scaling a skeletal mesh to 0.
+
+* Textures are no longer compressed for auto-saves to reduce the time it takes to save.
+
+* The color picker no longer assumes a straight inverse gamma conversion from colors sampled on the screen using the eye dropper. This change makes it so sRGB is the assumed space all colors are sampled in, making sampling colors from Photoshop bring in the identical version of the color converted into linear space, which when transformed back into sRGB space will be the same color you saw in photoshop.
+
+    * FColor and FLinearColors can now default to sRGB space using a UPROPERTY meta tag of sRGB=true (or false), you can force a specific default checking of the sRGB checkbox.
+
+    * Several places in UMG now force this default since sRGB is the space UMG will be rendered in (on PC). Additional work will be needed on Mac where the final rendering space is gamma 2.2.
+
+* The editor now properly saves out array elements to the ini files.
+
+* The Far View Plane setting for level editor viewports is now saved between sessions.
+
+* The Scalability state is now saved back to the same .ini file that it was loaded from, so that launching the editor doesn't reset Scalability settings set by the game.
+
+* There is now no longer a max limit to the number of static mesh sections that can be generated when importing meshes. The previous limit was 64 sections.
+
+* Updated Material Preview Mesh properties to include Adjacency Info.
+
+## Content Browser
+
+* New: Improved the behavior of Ctrl+B in the Content Browser and Persona.
+
+    * Ctrl+B in the Content Browser will now invoke the "Show in Folder View" command to allow you to easily sync to the actual location of an asset in the cases where you're viewing it via a search or a collection.
+
+    * Ctrl+B in the Persona animation sequence browser will now invoke the "Find in Content Browser" command.
+
+    * Ctrl+B will no longer do anything when Persona has no asset available for a given edit mode.
+
+* New: The "Move Folder Here" operation can now be aborted if there are a large number of assets within.
+
+* Auto reimporter now uses MD5 hashing to ensure that benign reimporting does not happen when only the timestamp has changed.
+
+* Can now use the "Show Developers Folder" option in the save/load dialogs.
+
+* Fixed an issue in which cancelling the attempted delete of read only files would still remove the asset from the Content Browser.
+
+* Fixed being able to rename or delete some read-only types in the Content Browser.
+
+* Fixed Content Browser drag-and-drop so that it works in all views (it only used to work in tile view).
+
+* Fixed Size Map incorrectly displaying assets with back-references in the "shared" node.
+
+* Fixed the "Show Collections" option appearing multiple times if multiple Content Browsers were open.
+
+* Fixed the Content Browser assets not updating their SCC state correctly when the SCC provider is changed while the editor is running.
+
+* Local collections are now preserved when creating a copy of your project.
+
+## Landscape
+
+* New: Added a "Mesh Vertical Offset" to landscape spline control points:
+
+    * This new property is blendable along the spline and independent of the spline width, unlike the old "Offset" in the spline segment mesh settings.
+
+    * The old "Offset" in the spline segment settings has been renamed to "Center Adjust", to more accurately describe what it does.
+
+* Added the p.bUseUnifiedHeightfield console variable to allow disabling the use of PhysX unified heightfields. (These are enabled by default):
+
+    * When this feature of PhysX is enabled, landscape collision thickness does not function and CCD is the only option to prevent fast moving rigid body penetration.
+
+    * You can set p.bUseUnifiedHeightfield=0 in your ConsoleVariables.ini file if you wish to disable the PhysX unified heightfield.
+
+* Fixed a crash generating landscape rendering data during mobile preview rendering level.
+
+* Fixed a crash when assigning a material to a LandscapeProxy.
+
+* Fixed crash when using the landscape resize tool.
+
+* Fixed a crash when colliding with landscapes with physical materials due to mismatch between the saved level and and data in the derived data cache.
+
+* Fixed external reference errors caused by moving components with an XY-offset map texture to another level.
+
+* Fixed Landscape Grass not updating after undoing a landscape painting operation.
+
+* Fixed landscape spline control point meshes not working.
+
+* Fixed problems with landscape cooking where launching on a mobile device or in the mobile previewer multiple times from the editor would not correctly cook changes made in the editor session.
+
+* Fixed rare crash caused by various UObjects used by the landscape tools getting garbage collected.
+
+* Fixed smooth tool being broken for strength > 1.
+
+* Fixed spline meshes stored in foreign levels disappearing in cooked builds.
+
+* Fixed the "Smooth" tool also causing a flattening effect.
+
+* Fixed the Landscape Editor keyboard shortcuts not working.
+
+    * Note: The number keys are bound to bookmarks by default, and so these won't work for switching tools in the viewport unless you unbind or rebind them.
+
+* Landscape collision thickness value now correctly propagated to underlying PhysX heightfield when changed.
+
+    * Note: Heightfield thickness requires the p.bUseUnifiedHeightfield=0 console variable setting in order to have any effect in PhysX.
+
+* Landscape tool now correctly shows layer names when new layer has been added.
+
+* Moving landscape components to a sub-level will now correctly move any attached foliage as well.
+
+* Precomputed visibility should now be correctly calculated for landscape actors.
+
+* Properly deform navigable geometry for spline mesh. (Contributed by [REDACTED_PERSON], pull request 1295.)
+
+* The user is now prevent from attempting to use the World Grid Material as a landscape hole material. Overriding blend mode on the default material will cause a crash.
+
+* Undo now works correctly with "Clear Region Selection".
+
+## Material Editor
+
+* New: Added Map Check warning after loading a level in the editor that uses a material that needs one of its Usage Flags set. Unless resaved, such materials will render as the default material in cooked builds.
+
+* New: Added World<->View transforms to the TransformPosition node in the material editor.
+
+* New: The Material Editor now saves the toggle state of nodes' comment bubbles, so they are preserved when reopening the material.
+
+* Fixed a crash placing a Custom Output material node in the material editor. That node type should have been excluded from the material node list.
+
+* Fixed a material compile error generating when using a Custom material node together with a TextureObject on OpenGL platforms.
+
+* Fixed a small error which was causing some Material Functions in the library to not appear in the right click menu, despite being exposed.
+
+* Fixed issue where editing a Material Parameter Collection results in marking an unnecessarily large number of materials as dirty.
+
+* Fixed titles on math nodes to update when the pin connections change.
+
+* TextureSample material nodes now fetch the correct tooltips.
+
+## Matinee
+
+* Fix bug that could cause sounds to start not quite at the beginning during playback.
+
+* Fix matinee crash when running without audio device.
+
+* Fixed matinee sound preview playback so that it respects the loop section and correctly restarts. Improved the quality of audio playback when scrubbing.
+
+* Launching movie capture once again works when using a DebugGame binary.
+
+* Matinee can no longer be opened while there is an active Play in Editor or Simulate session as it is unsafe to change matinees during play sessions.
+
+* Switching editor modes will now warn before closing Matinee mode.
+
+## Cascade
+
+* All particle modules now reset to their distribution properties to their correct defaults instead of "None".
+
+* Fixed crash when setting the particle system Initial Velocity to defaults when configured as a Distribution Vector Uniform Curve.
+
+* Mesh particles now properly populate dynamic parameter names in Cascade.
+
+## Persona
+
+* New: Added key handling to Persona segment panel for montages so delete key now works correctly.
+
+* Fixed delete context menu option for animation notifies not correctly deleting notifies on some platforms.
+
+* Fixed a crash in reinstancing of Persona preview instances after compiling blueprints that are referenced in the Animation Blueprint event graph.
+
+* Fixed a crash when switching between edit modes with wind source actors active.
+
+* Fixed an issue in the type-in UI for blendspace sample coordinates where the Y coordinate could not be set to the same as the X coordinate.
+
+* Fixed APEX clothing files failing to import when non-ANSI characters were present in the file path.
+
+* Fixed crash in anim compression when retargeting an animation with a single key.
+
+* Fixed crash in Persona accessing stale preview instances. Now we reinitialize the animation correctly when the blueprint changes.
+
+* Fixed crash in Persona when undoing deletion or duplication of state nodes.
+
+* Fixed crash when undoing a change in custom blending graphs for state machine transitions.
+
+* Fixed deleted names in smart name managers remaining after deletion.
+
+* Fixed ordering of source and target bones in the details panel for the copy bone node to make them more intuitive.
+
+* Fixed UI refresh being triggered for any property edit event fired from any UObject. Now filtered to refresh only when notify data is changed.
+
+## PhAT
+
+* Fixed issue where changing the mesh rendering mode in PhAT would not update in the viewport until the second click.
+
+* Fixed PhAT Editor crashing due to being able to select nonexistent bones.
+
+## World Outliner
+
+* New: Changed the level name tooltip in the Levels Browser to display the full package path name.
+
+## Static Mesh Editor
+
+* New: Added the ability to push the bounds for StaticMesh objects manually.
+
+    * The bounds can be altered by changing the value of the PositiveBoundsExtension and NegativeBoundsExtension properties for a StaticMesh.
+
+* New: Static mesh adjacency buffer generation is now set per-mesh. Controlled either through the mesh editor or on import.
+
+    * This was previously controlled by an internal fixed triangle limit but needs to be optional to allow PN tessellation on high-poly meshes.
+
+    * The flag will be auto-set for old meshes preserving existing behavior.
+
+## Source Control
+
+* New: Added make file writable button for advanced project setting pages, which allows files to be edited when non-writable and not connected to source control.
+
+* New: Once an asset is made writable (if it cannot be checked out), it will no longer continue prompting the user for check out when it is modified.
+
+* Fixed not being able to use collections with SVN or Git set as your SCC provider.
+
+* Fixed Subversion source control plugin to recognize the correct line terminator sequence depending on the platform. This was previously breaking the log output on Mac.
+
+* Subversion source control provider can now use the system default svn binary where it exists.
+
+    * This addresses issues where users have repositories and checkouts that were made with different versions of SVN to the client binary we ship. We now attempt to detect a system-level svn binary, and only use the shipped binary as a fallback.
+
+## World Browser
+
+* New: LOD sub-level generation has more options for material simplification (texture size, constant material parameters).
+
+* Fixed a case where world composition would scroll the minimap view to an undefined location if the user selected a hidden tile.
+
+* Fixed X-axis orientation when importing tiled landscapes from World Machine.
+
+* The 'Bake Grass' option for LOD sub-level generation will now correctly bake grass across all landscape surfaces.
+
+* World composition now incrementally updates sub-level tile images, preventing large editor hitches when a large number of tiles need update at the same time.
+
+## Foliage
+
+* New: Added a bulk foliage instance removal Blueprint function to HierarchicalInstancedStaticMeshComponent
+
+    * RemoveInstances takes an array of instance indices.
+
+    * It can directly use the Instances array passed into the new FoliageInstancedStaticMeshComponent OnInstanceTakeRadialDamage event to remove all damaged instances.
+
+* New: Added Blueprint support for Foliage instances.
+
+    * The foliage component class can now be overridden in the Foliage Type.
+
+    * Users can create a Blueprint subclass of FoliageInstancedStaticMeshComponent to respond to Hit or Damage events and hide or remove instances, spawn effects, spawn real Actors in place of an instance etc.
+
+    * The default class is FoliageInstancedStaticMeshComponent which is a subclass of HierarchicalInstancedStaticMeshComponent but has added support for Point and Radial Damage.
+
+* New: Added support for instanced mesh components (e.g. foliage) to FBX export.
+
+* New: Foliage painting now always trying to paint through filtered geometry. For example, when static meshes are filtered out foliage still can be painted on landscape bellow static mesh.
+
+* New: We no longer try to create physics data for InstancedStaticMeshComponent (eg Foliage) instances with a zero scale, which removes some PhysX warnings and improves performance.
+
+* Added a button to the Map Check warning about duplicate Instanced Foliage Actors to allow you to consolidate them.
+
+* Foliage can now be correctly spawned on top of Actors composed from Hierarchical Instanced Mesh components.
+
+* Running procedural foliage simulation will now mark only sub-levels that were actually affected by the simulation as dirty
+
+* Thumbnails in the foliage editor now update properly when hovered over.
+
+* We no longer try to create physics data for InstancedStaticMeshComponent (eg Foliage) instances with a zero scale, which removes some PhysX warnings and improves performance.
+
+## Cooker
+
+* New: General cooking performance optimizations.
+
+* New: Multiprocess cooking is in UnrealFrontend custom profiles.
+
+* New: New cooking stats system for more details ddc / package stats. Are saved into the <project>\Saved\Stats\ directory.
+
+* New: Optimization to Save package improves cook times.
+
+* New: Performance improvements of saving and caching packages in the cooker.
+
+* Cooking dlc no longer includes assets from base game in the asset registry.
+
+* Fixed an occasional crash when cooking, where a material with a mismatching texture sampler type binding would try to looking an UEnum during cooking and hit a check().
+
+* Fixed asset registry getting out of date when using -iterate.
+
+* Fixed crash in cook on the fly when recompiling shaders.
+
+* Fixed failure to package a build when also playing via LaunchOn.
+
+* Fixed missing PS4 x64 dlls (pad/morpheus) and .uplugins during staging.
+
+* Fixed various issues that were only showing up during the process of packaging a project from the editor.
+
+## Gameplay Framework
+
+* New: Actor transforms when spawning.
+
+    * The process of spawning actors has changed slightly. If an actor is spawned and its template's root component has a transform (for example, a RelativeRotation), that transform, including scale, will now be respected by the spawning code.
+
+    * For certain cases where you are sure you want to spawn at the exact transform given and ignore any transform on the template's root component, we have added the SpawnActorAbsolute() function.
+
+    * Check out DecalActor for an example of how this works in practice.
+
+    * *Notes when upgrading*:
+
+        * - If you have custom actors or blueprints with transform data on the root component that was previously ignored, you will see different results when spawning them with 4.9. In this case the best fix is likely to be simply zeroing out the transform data, since it was ignored before anyway.
+
+        * - If you have added code somewhere to take this transform data into account before spawning, you will likely want to remove it, or else you will end up applying the transform twice.
+
+* New: A Remove Player analogue to Create Player is now available for use in Blueprints.
+
+* New: AActor::bPendingKillPending renamed to bActorIsBeingDestroyed and made private.
+
+* New: Added "BrakingFriction" for CharacterMovementComponent, optionally enabled by "bUseSeparateBrakingFriction" setting.
+
+    * Previously GroundFriction was also used as the braking friction. This new setting allows separate control from the GroundFriction that affects turning control.
+
+    * Added "BrakingFrictionFactor" which controls how friction is used in ApplyVelocityBraking(). (Note: this is 2 by default for historical reasons, but a value of 1 more closely matches the actual drag equation).
+
+    * Different values are easily used in different movement modes by hooking in to OnMovementModeChanged and modifying BrakingFriction.
+
+* New: Added AbilityEndedCallbacks delegate to AbilitySystemComponent that broadcasts when an ability ends.
+
+* New: Added an automated test that will attempt to all variations of a new project. (IE C++ and blueprint and target platform for both mobile and desktop).
+
+* New: Added an automated test will parse and validate all available feature packs.
+
+* New: Added debug command ("p.VisualizeMovement") for showing in-world debug drawing of CharacterMovementComponent state including current velocity and acceleration.
+
+    * virtual UCharacterMovementComponent::VisualizeMovement() can be overridden in child classes to add project-specific debug drawing.
+
+* New: Added GetServerWorldTimeSeconds() to return the simulated time in seconds of the server. This is safe to call on either the server or the client.
+
+* New: Added MaxRange and a Tag (FName) to noise parameters for AI hearing.
+
+    * MaxRange is only used if > 0. It is modified (multiplied) by the loudness. Sensors must still pass their own range checks as well.
+
+* New: Added new Math functions to KismetMathLibrary:
+
+    * RandomUnitVectorInConeWithYawAndPitch, creates an uniformly distributed random Vector within a cone defined by the given Yaw and Pitch values.
+
+    * CreateVectorFromYawPitch, creates a Vector from the given Yaw and Pitch rotation values.
+
+    * GetYawPitchFromVector, retrieves the Yaw and Pitch rotation values in degrees from the given vector.
+
+* New: Added optional interpolation curve to the MoveToLocation AbilityTask to support per-axis interpolation.
+
+* New: Added ShouldAcceptGameplayCue() function to GameplayCueInterface, to allow actors to opt-out of receiving gameplay cues under certain circumstances.
+
+* New: Adding the blueprint function GetSupportedFullscreenResolutions, which gets the list of supported full screen resolutions.
+
+* New: Can now bind mouse wheel not just as scroll up or scroll down, but also as an axis which returns the delta change in mouse wheel position.
+
+* New: Collision now supports a new "PhysicsOnly" option. These objects will only be used in the simulated physics scene, not the query scene.
+
+    * Added ECollisionEnabled::PhysicsOnly option to round out all options of NoCollision, QueryOnly, PhysicsOnly, QueryAndPhysics.
+
+    * Can be used for things like capes and hair that users may not want to be considered by scene queries.
+
+* New: Convenience function to get level name in BPs.
+
+* New: Exposed UPrimitiveComponent::CharacterCanStepUpOn variable and UPrimitiveComponent::CharacterCanStepUp() function to blueprints.
+
+    * Also improved comments/docs on these to indicate that it only prevents step ups when bumping in to objects, not walkability on them (use walkable slope override for that).
+
+* New: Exposed some Vector : Integer arithmetic methods to Blueprints.
+
+* New: Gameplay Cues can now be specified via GameplayCueSets. Classes can implement IGameplayCueInterface to return a GameplayCueSet to use instead of the global set.
+
+* New: Performance: Optimized overlap detection during swept movement. Approximate 10% or greater reduction to swept movement cost when in the presence of overlaps (depending on scene complexity).
+
+    * Overlaps at the end location are now computed based on tests against the individual candidates found during the sweep, rather than using a full scene query from UpdateOverlaps().
+
+    * Changes in rotation still allow cached overlaps as long as the rotation is symmetric (ie couldn't have changed overlaps).
+
+    * UpdateOverlaps() now avoids any processing for child components when done for the root component. This makes optimization simpler and children will test the parent anyway.
+
+* New: Revised feature pack parsing code. Feature packs can now be parsed via the manifest file eliminating the need to pack all feature packs.
+
+    * This was used to implement the feature packs used for the shared resources now used in templates.
+
+* New: Smooth rotation for network corrections of simulated proxies in character movement.
+
+* New: Spawn Actor now accepts Scale properly.
+
+    * Because of this, we added SpawnActor to accept FTransform as an input.
+
+* New: The class of the default physics volume can now be set via the World Settings properites.
+
+* New: When unloading a streaming level all objects are now marked pending kill to avoid any stray references to the objects that could prevent the level from being garbage collected.
+
+* New: You can now bind the Any key. This will respond to any key press/release event and follows the same rules for consumption that other keys follow.
+
+* "Always-loaded" sub-levels will now be fully loaded before gameplay begins when using seamless travel.
+
+* A graceful message for failure to load the default map is now shown followed by closing the application when travelling instead of crashing.
+
+* AActor::bPendingKillPending renamed to bActorIsBeingDestroyed and made private.
+
+* BeginPlay will no longer be called on a component that has seamlessly travelled from another level.
+
+* Broadcast PhysicsVolumeChangedDelegate even if the previous physics volume was null.
+
+* Bugfix to timers that are paused immediately after being set.
+
+* Calling TryActiveAbility will now call CanActivateAbility on the Primary Instance of the Ability Spec, if it exists. (If not we continue to do the previous implementation which will call CanActivateAbility on the CDO.)
+
+* Change FQuat tolerance within SceneComponent and PrimitiveComponent to allow smaller rotation changes to get through.
+
+* Changed KismetMathLibrary functions PointsAreCoplanar, LinePlaneIntersection, and LinePlaneIntersection_OriginNormal to be BlueprintPure.
+
+* Component attachment is now replicated correctly to clients.
+
+* Console command ServerTravel can now pass parameters to the map.
+
+* Fixed characters not rotating with a base object that is moving at a very low rotation rate in games with a high frame rate. Quat comparison tolerance was too strict.
+
+* Fixed Editor crashing due to SetMaterial getting called on a DestructibleComponent during the construction script.
+
+* Fixed crash exiting Play in Editor if you had travelled to or restarted the level being edited during the PIE session.
+
+* Fixed crash when the component specified as the player's audio listener component has been destroyed.
+
+* Fixed Sound Mix failing to be pushed if it had been popped in the same frame.
+
+* Fixed using the PIE prefixed map name when trying to seamless travel in multiplayer PIE.
+
+* Fixed a crash which occurred on PIE exit after enabling Show Collision in game.
+
+* Fixed an issue with feature packs registering for super search which could result in multiple registrations which then caused an assert.
+
+* Fixed camera jumping when trying to interpolate back to original view target. (Contributed by [REDACTED_PERSON], pull request 1215.)
+
+* Fixed CameraModifier blending issue that occurred with blends longer than 1 second.
+
+* Fixed Character::CanJump() event no longer showing up as overridable in the BP event menu.
+
+* Fixed crash when adding an instanced component with a native sub-component to an actor with a blueprint created root component.
+
+* Fixed crash when loading a map that was saved with a WorldSettings object of a class that no longer exists.
+
+* Fixed crash when overlap events (triggered by toggling bGenerateOverlaps or changing collision setting) destroy components or cause other overlaps to end.
+
+* Fixed issue where a Matinee event which stops the Matinee sequence could cause a crash when updating subsequent tracks.
+
+* Fixed issue where EndAbility was not reentrant-safe due to EndAbility being called again from inside of the call to K2_OnEndAbility.
+
+* Fixed issue with camera shake oscillations not working in non-local play spaces.
+
+* Fixed GameMode potentially outputting misleading log message with incorrect class name. (Contributed by [REDACTED_PERSON], pull request 1319.)
+
+* Fixed Navigation Mesh not getting created correctly for a ProceduralMeshComponent's mesh due to not having its bounds updated.
+
+* Fixed NavWalking movement mode replicating to network proxies as Custom movement.
+
+* Fixed one shot reverse behavior in interp to component.
+
+* Fixed physics volume change notifications not working for the first change in physics volume.
+
+* GameplayEffects now have a higher granularity on the client prediction of time remaining on active effects.
+
+* Improved comments on actor set location/rotation/transform functions, especially regarding bSweep and bTeleport.
+
+* Interpolate simulated position faster when network character proxy is coming to a stop. Reduces visual overshoot past end location.
+
+* Keys bound via key or action events will now count as having consumed the event even if there wasn't a pressed or released event this frame preventing axes using that key in lower priority input components from being active on other frame.
+
+* Made UCharacterMovementComponent::GroundMovementMode private, and provided Get/Set methods.
+
+    * Enforces that GroundMovementMode is only ever Walking or NavWalking.
+
+    * Set method will switch movement modes if currently a different ground walking mode.
+
+    * Cleaned up constants in Pack/UnpackNetworkMovementMode() and made sure we only have a valid ground mode (Walking or NavWalking).
+
+* Made it clear that MoveIgnoreActors does not apply to movement from physics simulation.
+
+* Made some useful CameraShake instance variables public and BP-accessible. Also Added UCameraAnimInst::SetScale().
+
+* Make sure client and server perform movement with the same acceleration value after network rounding. Reduces accumulated client position error.
+
+* Marked bTeleport' flag as Advanced for blueprints wherever we do so for the bSweep flag.
+
+* Negative values are no longer allowed in the editor for CapsuleComponent radius or half-height.
+
+* Pawns whose Tick functions have bStartWithTickEnabled set to false will no longer start ticking when possessed.
+
+* Performance: Controller::SetControlRotation() now avoids doing work if nothing changes.
+
+* PlayerCameraManager crashing on empty entries in DefaultModifiers array.
+
+* PlayerController uses a reliable RPC to request that the pawn gets resent from the server if it has not yet replicated. Previously this was an unreliable call so it didn't always work.
+
+* Revised comment on interp to component bPauseOnImpact to better reflect functionality.
+
+* SceneComponent AttachChildren array is now replicated so that it is consistent with AttachParents.
+
+* SetCullDistance on UPrimitiveComponent now accepts 0 as a legal distance. A distance of 0 means the primitive is not culled at any distance.
+
+* The data for automated line test was incorrect causing false failures.
+
+* Timeline components will no longer execute their tick function when not active.
+
+* We no longer ignore hit testing in APlayerController::GetHitResultAtScreenPosition if the mouse isn't directly over the viewport. That function can be used regardless of the state of the mouse, so that logic has been moved elsewhere.
+
+* When transitioning levels EndPlay will no longer attempt to be executed for an Actor that was destroyed by another Actor's EndPlay call.
+
+* You can now set the GameMode HUD class to be none.
+
+# Learning Resources
+
+## Documentation
+
+* Improved docs for Add Yaw Input and Add Controller Yaw Input etc., to point to the scaling values that are applied.
+
+* Improved docs on Movement Component and Floating Pawn Movement to indicate that collision is wholly dependent on the Updated Component (normally the root component).
+
+## Sample Content
+
+* Fixed UI icons not being present in sample content.
+
+## Localization
+
+* New: Added game native culture setting to control which culture's localization data the editor loads for game assets. Game assets now display the specified culture's localization data when in the editor.
+
+* New: Added initial support for working directly with Localization Services in the editor. Allows you to upload and download translations to online services for localization. Allows plugins for specific localization service providers (OneSky, etc). Also created first localization service plugin for OneSky.
+
+* New: Added the source location information of a localization archive entry as an "extracted comment" in PO files when exporting.
+
+* New: Culture names now consistently use hyphens rather than underscores. Culture names with underscores are still accepted.
+
+* New: First pass of the OneSky plugin: Extend Localization Dashboard toolbar with OneSky buttons for "Import all cultures/targets from OneSky" and "Export all cultures/targets to OneSky".
+
+* New: Fixed live culture swapping when in certain configurations such as shipping. Enabled use of "-culture=" switch on command-line in all configurations.
+
+* New: Foreign archives now generate entries for the native translation of the native source as though it were the native source - the result is that the native translation will also be translated in the foreign culture. Localization resources now use the native translation of the native source as though it were the native source, specifically when looking up the foreign translation to use from archives - the result is that the native source will use the foreign translation of its native translation.
+
+* Fixed broken legacy support for SourceFileSearchFilters setting in GatherTextFromSource steps in localization configuration scripts.
+
+* Fixed dashboard handling of paths rooted in drives other than the current drive of the editor.
+
+* Fixed issues selecting language and region combinations for Chinese and Portuguese in the Region & Language page of the Editor Preferences.
+
+* Fixed non-functional PO import/export buttons from the translation editor.
+
+* Localization commandlet JSON file writing no longer fails if source control is enabled but the files don't already exist on disk.
+
+* Localization dashboard now appropriately adds the native culture setting to generated localization configuration scripts.
+
+* Retrieving localized metadata will prefer local changes to the metadata if they do not match what is expected by localization.
+
+* The Category Names of Components in blueprints are no longer gathered for localization when the "ShouldGatherEditorOnlyData" setting is false.
+
+* Updated logic in commandlet for gathering text from assets which was falsely reporting that nothing was gathered because no files were loaded. Now properly reports whether there were any packages found to be processed, rather than if any need to be loaded and processed in a deferred manner.
+
+## Marketplace Content
+
+* Infiltrator Demo - Various maps have been cleaned up to remove NULL referenced scene assets.
+
+## Networking
+
+* New: A warning will now be printed to the output log if an actor tries to send a non-multicast remote function but has no owning connection.
+
+* New: Added a check to see if the character is pending kill before calling TickCharacterPose() as the character can be marked for destroy after PerformMovement.
+
+* New: Added a file NetConstants.h which holds MAX_PACKET_SIZE which can be changed to update the MTU. Changed the MTU from 512 to 1000.
+
+* New: Added a macro CLOSE_CONNECTION_DUE_TO_SECURITY_VIOLATION to AssertionMacros.h which closes a connection, logs it, and provides a debug string to the client.
+
+* New: Added a scale component to PackageMapClient and now use a 1 bit check whether Location, Rotation, Scale, and Velocity have been serialized.
+
+* New: Added ability to save json in condensed format.
+
+* New: Added ability to search for replays by meta tag.
+
+* New: Added ability to track which users were in replays.
+
+* New: Added ability to view recently viewed replays.
+
+* New: Added LogRep and LogRepTraffic to EngineLogs.h and changed over some UE_LOG calls to use the new logs.
+
+* New: Added NMT_SecurityViolation where servers pass a debug string to clients to inform them of a disconnect.
+
+* New: Added support for the engine to notify the replay streaming service when players join in progress.
+
+* New: Added the ability to find replays that you participated in.
+
+* New: Added the Packet Handler Framework which allows for transformation of packets, including Encryption, Compression, Reliability.
+
+* New: Added Third Party Library CryptoPP to the Third Party Source directory.
+
+* New: Added UE_SECURITY_LOG to AssertionMacros.h which logs security events.
+
+* New: Encryption Handler components that use the CryptoPP library and Packet Handler system to encrypt packets.
+
+* New: Reduced bandwidth usage by up to 15% or more by compressing replicated location, velocity, and rotation values more aggressively.
+
+    * The amount of bandwidth savings you see in your project will vary depending on how much of it consisted of actor movement replication.
+
+    * This does reduce the precision of these values slightly, but it should not cause noticeable artifacts in most cases. The level of compression is now configurable in the "Replication" section of the Class Defaults, in case higher precision is needed.
+
+    * Using "Round Two Decimals" for the location and velocity, and "Short Components" for rotation will be equivalent to the precision that was used in 4.8.
+
+* New: UNetConnection has a new GetWorld function. It will try to get the world from the owning net driver and failing that the owning actor.
+
+* Added #IF_WITH_EDITOR code guard to remove code that derives GetNetMode if in PIE.
+
+* Added voice-related exports to NetConnection.h. (Contributed by [REDACTED_PERSON], pull request 1234.)
+
+* Added check in NetConnection for the MaxPacketSize being passed into InitConnection.
+
+* Added numerous changes to stop replicating client on the server when bReplicateMovement is false.
+
+* Added PreReplication of components when PreReplication of actors is called.
+
+* Added two functions which will attempt to derive the NetMode from other data if the NetDriver is null.
+
+* Added code guards around AttemptDeriveFromPlayInSettings instead of within the function so it won't be compiled in release builds.
+
+* Changed MTU back to 512 from 1000 as discussions of MTU changes are still happening.
+
+* Changed the serialization code in PackageMapClient.cpp to follow more consistent serialization conventions.
+
+* Changed FOnActorChannelSpawn from one param to two params.
+
+* Changed FOnReplicatedActorSpawn from two params to one param.
+
+* Disabled single process PIE seamless travel while we work on a long term fix.
+
+* Fixed crash that could occur while scrubbing in replays.
+
+* Fixed a bug where time dilation was incorrectly affecting gamepad during replay playback.
+
+* Fixed a crash that could occur if replay streaming failed after a "stop streaming" command was issued.
+
+* Fixed a crash that could occur when viewing live replays.
+
+* Fixed a crash when starting to record a network profile in some cases.
+
+* Fixed a rare bug that could cause network game clients to fail to load some assets in cooked builds.
+
+* Fixed a rare crash that could occur when a network client joins a game.
+
+* Fixed issue where the last checkpoint wouldn't load if there wasn't any stream info after it.
+
+* Fixed APlayerController::SpawnLocation not being replicated. (Contributed by [REDACTED_PERSON], pull request 633.)
+
+* Fixed an issue where the engine was generating corrupted network profiler files if the URL of the loaded map contained more than about 255 characters.
+
+* Fixed potential crash if a demo net driver is left idling for too long.
+
+* Fixed the reported size of replicated functions in network profiles. Some bits that were incorrectly reported as being in the parameter section are now correctly reported as being in the header section.
+
+* Live replay timeouts when first connecting are handled correctly.
+
+* Multiple running instances of the same project can now record network profiles simultaneously.
+
+* Replay streamer now generates error messaging and stops demo recording when it fails.
+
+* Special cased PIE dedicated server or client without server in World.cpp.
+
+* The demo net driver for replays is destroyed when a seamless travel copies the world data to ensure it always gets destroyed.
+
+* We no longer spawn particle system components on dedicated servers, reducing both memory and performance overhead.
+
+* When viewing replays, character smoothing between updates should be much smoother.
+
+* When scrubbing replays, camera will continue to follow the original view target.
+
+## Online
+
+* New in BuildPatchServices:
+
+    * Improved logic for detecting existing installations when running an installer.
+
+    * Added support for SHA1 verification of chunk match so that generation does not need to load existing chunks.
+
+    * Exposing list of files in the build via the manifest interface.
+
+    * Patch generation now makes use of available CPU cores to improve processing speed.
+
+* Fixed a variety of issues with BuildPatchServices:
+
+    * Detecting previous move when resuming a patch.
+
+    * Resuming an installation on some platforms where the file API would not allow seeking on the file handle.
+
+    * Verification stage not being capable of cancelling.
+
+    * Incorrect error string generation.
+
+    * Potential cloud enumeration deadlock for patch generation.
+
+    * Optimising installer attribute setting logic, as well as improving reported progress information for this stage.
+
+    * Removed uses of helper function FPaths::FileExists which is not robust on all platforms.
+
+* Fixed XMPP MUC room members populated correctly.
+
+* Fixed player's presence status on Steam so that it now shows their actual status instead of the string needed to connect to their game.
+
+* Fixed a rare issue where attempting to save out Steam achievements multiple times in the same frame could prevent further Steam tasks from executing.
+
+## Other
+
+* New: Added a generic motion controller key abstraction
+
+    * This allows multiple physical devices (e.g. left and and right hand controllers) to be routed to a single player in the engine.
+
+    * New motion controller inputs specify which hand they're for. E.g. Motion Controller (L) Grip 1.
+
+* New: Added a render array to texture maxscript which can be used to bake textures out of unreal's blueprints given a certain data format.
+
+* New: Added entitlement checking for Gear VR, which lets developers verify that their app has been purchased through the Oculus store.
+
+* New: Added functions to get and set the World to Meters scale from Blueprints.
+
+* New: Added in Motion Controller Component, which you can add to any actor, and will set its location and orientation to that of the player's motion controller.
+
+* New: Created a Flurry analytics plugin for iOS.
+
+* New: Improvements to the console:
+
+    * Console commands are now searchable. To search, begin with a "?" followed by your search terms separated by spaces. For example, to search for commands that contain both foo and bar you would type "? foo bar".
+
+    * The autocomplete feature has been updated so it will continue to display an exact match while you enter parameters to the command.
+
+* New: Leap Motion Controller now supports HMD mode, along with image passthrough.
+
+* New: The multicast (broadcast) analytics module is now a plugin to make it easier to use.
+
+* New: Updated engine code to support the Oculus 0.7 SDK. Engine will be hotfixed with 0.7 revisions when the runtimes are made public.
+
+* New: Updated Gear VR support to use the 0.6.0 mobile SDK.
+
+* Fixed Oculus plug-in forcing frame synchronization every frame. Engine will now allow one frame of buffering.
+
+* Fixed VR PIE not showing up when editor tabs with viewports are visible.
+
+* Fixed VR preview occasionally going black when other tabs are open.
+
+* Fixed a bug preventing the contents of web browser widgets from displaying properly when stretched across multiple screens.
+
+* Fixed a crash that occurred when the Oculus service was running, but no device was plugged in.
+
+* Fixed comments on linkable element types so the editor shows correct tool tips when selecting types.
+
+* Fixed example environment collision to ensure consistent results within level.
+
+* Fixed issues with post-processing not rendering correctly in stereoscopic when in ES2 / forward rendering.
+
+* Fixed jittery movement with Leap Motion hands.
+
+* Exported SteamVR function library, for use in projects that want to interface with SteamVR through C++.
+
+* Integrating fixes for Mac / Open GL virtual reality pipeline with Oculus.
+
+    * Custom Present now works with the Open GL RHI.
+
+    * Devices now update when paused.
+
+* Renamed the "Desktop" Target Platform to "AllDesktop" to stop conflicts with the Desktop directory.
+
+* Set the default frame buffering for HMDs to be one frame, to support compositor-based rendering. This fixes an issue in 4.8 where non-Oculus devices were feeling laggy.
+
+* Added a dialog box to alert users when there is a problem loading a plugin descriptor file.
+
+## Paper2D
+
+* New: Exposed a way to set collision thickness / offset for tile maps and tile layers, and added a RebuildCollision method for use after SetTile calls.
+
+* New: Tile Maps: Changed the default collision domain for Tile Maps from None to 3D.
+
+    * Note: Collision must still be authored in the tile set editor per-tile in order for any actual collision geometry to be generated
+
+* Changed asset set methods on Paper2D components to use AreDynamicDataChangesAllowed, allowing asset updates on components with Movable=Static in a construction script.
+
+* Fixed a crash when building lighting for a map that contains a placed tile map component or grouped sprite component.
+
+* Fixed an issue where body elements of a 2D terrain spline do not have collision.
+
+* Fixed an issue where sprite sheets could not be reimported in a project that was converted from an older version.
+
+* Fixed an issue where sprite trim/rotate settings were not reimported properly. Also increased speed of reimport process.
+
+* Fixed inconsistency between vertex colors on Paper2D assets and static meshes.
+
+    * All sprite colors are now treated as being defined in linear space, and existing assets will be converted on load.
+
+## Physics
+
+* New: Added GetPhysicsLinearVelocityAtPoint, SetCenterOfMass, SetMassOverrideInKg, GetClosestPointOnCollision to PrimitiveComponent.
+
+    * Also allow specifying duration for PrintString & PrintText functions.
+
+    * Both contributed by [REDACTED_PERSON], pull request 1335.
+
+* New: Added a project wide setting for turning off the default per poly complex collision.
+
+* New: Added GetMassScale function to PrimitiveComponent and exposed it to Blueprints.
+
+* New: Added SetIgnoreMoveInput and SetIgnoreLookInput functions to PlayerController.
+
+    * These functions reset the Move or Look ignore flag to its default value.
+
+* New: Added SetPlaneConstraintEnabled function to MovementComponent and exposed it to Blueprints.
+
+* New: Added support for parallel blend physics.
+
+* New: Added support for welding complex collision. This means we now support multiple trimeshes per BodySetup.
+
+* New: Added the ability to create Blueprint variables of the Physical material type.
+
+* New: Added time dilation in PhAT.
+
+* New: Exposed GetCollisionObjectType and GetCollisionResponseToChannel to blueprints.
+
+* New: Exposed skeletal mesh motor controls to blueprints.
+
+* New: Exposed the teleport flag for movement nodes and functions. This makes it easier to control velocity and collision of objects being moved.
+
+* New: Exposed UpdateJointsFromAnimation to PhAT. This makes it easier to tune animation driven motors.
+
+* New: Increased the default max angular velocity to 10 revolutions per second.
+
+* New: Introduce sleep and wake events that user can hook into. This allows you to get callbacks whenever physics simulation wakes or puts an object to sleep.
+
+* New: Sub-stepping physics is now supported on mobile.
+
+* New: Trimeshes with negative scale no longer need a separate cooked mesh. Saves memory and streaming time.
+
+* New: Turn constraint projection on by default and set their values to something reasonable.
+
+* Fixed collision shapes not being drawn for non-uniformly scaled objects.
+
+* Fixed wind forces being applied to APEX cloth even when no sources are present.
+
+* Fixed bIgnoreOwningActor option on URadialForceComponent not working for forces.
+
+* Fixed bReturnPhysMat not being respected for heightfields. Also made sure to init Hit.FaceIndex to INDEX_NONE when it's not filled in.
+
+* Fixed a crash when passing no component to Component Overlap Components.
+
+* Fixed a crash that occurred when using a material with PN Triangle tessellation enabled on Cloth.
+
+* Fixed cloth crash when transitioning between worlds containing the same cloth asset.
+
+* Fixed crash in auto-convex tool when no sections have collision, and simplify status message.
+
+* Fixed destructible component crash when transitioning between worlds with substepping enabled.
+
+* Fixed destructible mesh root chunk transforms not being correctly updated when not visible, leading to the render mesh being offset from the actual location of the chunk.
+
+* Fixed destructible meshes flickering in cooked builds due to incorrect root bone transforms.
+
+* Upgrade to final PhysX 3.3.3 (before we had a preview version) fixes various issues as found in their release notes.
+
+* Auto-convex generation now ignores sections with collision disabled (suggestion from @ManOfLode).
+
+* Changing PhysicalMaterial override on an instance now correctly updates mass properties.
+
+## Platforms
+
+* New: A prominent message is show when a user attempts to package a game that does not have binaries built.
+
+* New: Added support for optional packaging blacklist per platform and build configuration.
+
+    * Create a text file in [Project]/Build/[Platform] called PakBlacklist-[Config].txt where Config is Development or Shipping
+
+    * Each file being staged will be checked to see if the beginning matches any line in the file. If so, it is skipped. Check the log to see a list of excluded files.
+
+* New: The new event stats are now enabled for platforms that use the pthreads implementation.
+
+* New: When packaging a game for a platform, errors and warnings are now presented to the Message Log.
+
+    * This means that errors and warnings can be identified quicker.
+
+* Config Hierarchy edit button is now enabled even when some files in the ini file hierarchy are read-only.
+
+* Fixed a bug during staging where plugins may not be discovered properly.
+
+* Fixed a cookonthefly bug where plugins would not be detected properly.
+
+* Fixed a mutexattr leak for pthread platforms.
+
+* Fixed issue where ChunkID assignments set by the editor are ignored.
+
+* Fixed Slate style warnings when loading a project on some platforms.
+
+* Remote Mac or iOS building from Windows now return clearer set-up errors.
+
+* Virtual joysticks now deal with screen resolution or orientation changes properly (repositioning and sizing the controls).
+
+## Android
+
+* New: NVIDIA AndroidWorks installer.
+
+![image alt text](image_76.png)
+
+    * TADP replaced with NVIDIA AndroidWorks for Windows and Mac OSX.
+
+    * Please remove any previous TADP installation before running the new AndroidWorks installer.  An EXE for windows and a DMG for OSX are located in UE4/Engine/Extras/Android.  
+
+    * Included in new AndroidWorks:
+
+        * NDK r10e (64-bit toolchain).
+
+        * SDK (24.2.0).
+
+        * Nsight 3.1 debugger Visual Studio extension with support for all Android ABIs.
+
+* New: Added support for Immersive Mode on Android devices of version 4.4 and up.
+
+* New: Project generation now compatible with AndroidWorks (Nsight 3.0+).
+
+* New: Supported added for ES3.1 builds. These can be built at the same time as ES2.0 builds of an application..
+
+* New: Enabled ld.gold on ndk r10e (clang 3.6), and added expermental arm64 compiler arguments.
+
+* Added long suffix to file length in generated OBBData.java. (Contributed by [REDACTED_PERSON], pull request 1405.)
+
+* Added screenSize to android:configChanges in AndroidManifest.xml to prevent onDestroy issues on suspend.
+
+* Better handling of multiple Android build targets.
+
+* Corrections for ES 3.1 + AEP shaders.
+
+* Custom rules added to deal with newest SDK releases of build.xml not defining the paths to some tools.
+
+* Disable GearVR plugin for unsupported Android platforms.
+
+* Enabled shadow warning for Android targets and cleaned up the ones found.
+
+* Fixed a problem where pressing 'back' while downloading or verifying an OBB on Android could cause the application to crash.
+
+* Fixed Android ES 3.1 + AEP so it runs the high-end Rivalry demo again (regression from 4.8).
+
+* Fixed black-screen for ES 3.1 when not using MobileHDR (e.g. Tappy Chicken).
+
+* Fixed case where some Android devices did not render BSP meshes, e.g. Galaxy S3 (Mali) or Galaxy Note 2 (Mali).
+
+* Fixed releasing queries on Android ES 3.1. Before, it could leak query memory/resources.
+
+* Fixed: Rendering corruption when using a fully rough material on some Android devices (eg: Nexus7, Nexus4).
+
+* Fixed compiling problems with non-unity builds for Android.
+
+* Improved Android ARMv7 debug flags.
+
+* Invalidates cached resolution when application returns from background to handle receiving new window.
+
+* Keep the last export of each Android NDK/SDK path from .bash_profile instead of of the first one if not set in ini (Mac).
+
+* Make sure the UE4Game and UE4Game/[GameName] directories are created so save games and logs work properly.
+
+* Only use -fno-function-sections on Clang 3.6 or later on Android to correct issue with backward compatibility to older toolchains.
+
+* OpenWrite does not use O_APPEND but preserves seek to end of file behavior if append requested on Android.
+
+* Pass correct event for onDestroy so application terminating delegate called, but do NOT depend on this; there is no guarantee onDestroy will be called by the system. Any state you want to save should instead be handled by ApplicationWillEnterBackgroundDelegate.
+
+* Removed libGLES3 dependency from libvrapi.so (part of Oculus mobile VR library for GearVR) so older Android versions supported.
+
+* Resolved issues with incorrect screen resolution being using when returning back to the application from sleep mode or the home screen.
+
+* Sensor device orientation options are now supported.
+
+* Strips debug symbols from libgnustl_shared.so for distribution builds on Android.
+
+* Updated Android x86 PhysX libraries.
+
+* Updated FAndroidPlatformMemory::BinnedAllocFromOS for compatibility with SDK 21.
+
+* Use gnu-libstdc++ 4.8 for NDK level 19, and allow 4.9 for higher NDK levels.
+
+* Cleaned up some tooltips.
+
+## iOS
+
+* New: Added initial support for the Cloud Kit framework. Cloud-based databases can be accessed via the Online User Cloud Interface API and the Online Shared Cloud Interface API.
+
+* New: Added Remote Push Notifications courtesy of Get Set Games. All features of remote notifications can be accessed from blueprints and code.
+
+* New: Added some convenience functions to the In-App Purchase feature. These functions provide access to various localization data for the products. (Courtesy of Get Set Games.)
+
+* New: Added support for building games against the iOS 9 SDK.
+
+* New: The ability to restore In-App purchases has been added. This feature can be accessed via the Restore In-App Purchase blueprint or through code utilizing the Online Store Interface.
+
+* New: Updgraded Facebook SDK on IOS to version 4.01.
+
+* Fixed a crash on game startup when running on iOS 9.
+
+* Fixed a crash when entering text into a text edit field on IOS.
+
+* Fixed a hang when popping up the message box on iOS.
+
+* Fixed an intermittent crash on game startup when utilizing Metal on iOS.
+
+* Fixed an issue with initializing remote building for iOS on OS X 10.9.
+
+* Fixed an issue with non-existent dSYMs being added to the build receipt when building for iOS with dSYM generation enabled.
+
+* Fixed an issue with the xcode-select command when building using a remote Mac.
+
+* Fixed loading the project ini files for every file in a game project when building for iOS. This results in faster build times.
+
+* Fixed the date check when looking for valid signing certificates on the Mac.
+
+* Increased the heap size of the game thread on iOS. This prevents out of stack space crashes.
+
+* The iOS deployment server now correctly handles non-ASCII characters in the path or filename of files.
+
+* The splash screen will now properly be removed if a message box is displayed.
+
+## Linux
+
+* New: Added CodeLite project generator and source accessor. (Contributed by [REDACTED_PERSON], pull requests 1378 and 1379.)
+
+* New: Added Linux platform analytics.
+
+* New: Dedicated servers will now flush logs at regular intervals (30 seconds by default, configurable via ServerFlushLogInterval).
+
+* New: Setup scripts will refuse to run on Mono older than 3.x; old Mono versions are known to be unusable for building UE.
+
+* New: Significantly increased length limit for DDC (derived data cache) paths on Linux; previous numbers were based on smaller Windows MAX_PATH limit.
+
+* New: Switched file dialogs from Qt/GTK to Slate; engine no longer depends on presence of said libraries in the system.
+
+* New: The engine and tools will refuse to run with root privileges.
+
+* New: Upgraded Linux cross-toolchain to clang 3.6.0.
+
+* Engine libraries should no longer use initial-exec TLS model (this prevented loading libraries under certain circumstances).
+
+* Fixed a crash when saving assets with localized names.
+
+* Fixed building *Editor targets using QtCreator. (Contributed by [REDACTED_PERSON], pull request 1171.)
+
+* Fixed crash on a dedicated server when the higher level code tries to access display metrics. (Contributed by [REDACTED_PERSON], pull request 1232.)
+
+* Fixed file handles leak on Linux on program restart.
+
+* Fixed hot reload on Linux. (Contributed by [REDACTED_PERSON], pull request 1178.)
+
+* Fixed inability to compile Linux ARM binary (regression introduced late in 4.8.x cycle).
+
+* Fixed inability to locate bundled Steam libraries.
+
+* Fixed Linux editor not honoring "Use less CPU when in Background" setting. (Contributed by [REDACTED_PERSON], pull request 1207.)
+
+* Fixed Linux toolchain not honoring "enable exceptions" build configuration setting. (Contributed by [REDACTED_PERSON], pull request 1135.)
+
+* Fixed minimize/maximize window events not being fired in certain circumstances. (Contributed by [REDACTED_PERSON], pull request 1115.)
+
+* Fixed multiple problems with localized project names and directories.
+
+* Fixed packaging failure on Linux when the project is placed next to the engine directory.
+
+* Fixed popups not appearing under certain circumstances or being non-interactive; mouse relative movement being wrong with multiple windows; multiple other improvements. (Contributed by [REDACTED_PERSON],pull request 1375.)
+
+* Fixed ungraceful handling of OpenAL initialization failure (Contributed by [REDACTED_PERSON], pull request 1206.)
+
+* Fixed Unicode keyboard input.
+
+* ICU libraries are now properly accounted for in UBT/UAT as a run-time dependency.
+
+* Improved mouse position reporting when the frame rate is low. (Contributed by [REDACTED_PERSON], pull request 1369.)
+
+* Including a libpng module in Linux dedicated server no longer breaks the build (that was an artificial limitation that is not relevant anymore).
+
+* Superfluous axis events from controllers are no longer sent multiple times per frame; controller IDs can be disjoint. (Contributed by [REDACTED_PERSON], pull request 1122.)
+
+* Switched to using SDL key codes instead of raw keyboard scan codes. This fixes various improperly mapped keys and removes US QWERTY assumption in general.
+
+* UBT will properly recognize when UHT exits abnormally on Unix platforms (Linux, Mac) and report it as a crash.
+
+## Mac
+
+* New: Added support for Xcode 7 (for iOS 9 and MacOS 10.11).
+
+    * Improved how the Xcode path is discovered in UBT. (No more hardcoding, unless desired - it uses the current Xcode, or xcode-select when building from command line.)
+
+    * Added some UBT utility functions. (Log.TraceInformationOnce to print out a message only one time, and Utils.RunLocalProcessAndReturnStdOut to easily capture output.)
+
+* New: Bundled Mono for Mac updated to version 4.0.2.
+
+* New: CrashReportClient can now generate the source context showing where a crash has occurred in the diagnostic report on OS X.
+
+* New: Generate debug symbol files on Mac OS X for Development & Shipping build just as we do for Windows.
+
+* New: Mac users can now emulate right-click using Control key + left-click.
+
+* Changed the way the CrashReportClient is launched on Mac OS X to avoid a potential deadlock.
+
+* Disable OpenGL GPU timing on Nvidia GPUs and entirely on OS X 10.9 (Mavericks) to avoid crashes in the OpenGL drivers.
+
+* Fix a crash when launching an external process on Mac OS X caused by over-releasing the path to the executable.
+
+* Fix reflections on translucent surfaces for OpenGL by moving the ReverseBits32 function from MonteCarlo.usf to Common.usf and using it on platforms that don't natively support the reverse bits intrinsic.
+
+* Fix transition from WindowedFullscreen to regular FullScreen and back again on Mac OS X.
+
+* Fixed a problem with windows not rendering on OS X 10.11 beta.
+
+* Fixed emitting the callstack of the crashed thread to the log on Mac OS X.
+
+* Fixed a crash on auto-importing movies.
+
+* Fixed issues with playing multiple movies in a row.
+
+* Fixed Rename option in various context-menus on Mac OS X by ensuring that the correct window deactivation message is sent when closing the active window.
+
+* Fixed setting of the current working dir on Mac for paths that contain non-ANSI characters.
+
+* Fixed splash screen rendering on OS X 10.11 beta.
+
+* Fixed tooltips rendering on OS X 10.11 beta.
+
+* Fixed window activation in a case where mouse-down event changes focus to a new widget.
+
+* Fixed mouse position setting/reporting in fullscreen mode.
+
+* Fixed the Mac input event code generating incorrect key-up or key-down states for modifier keys.
+
+* Increased the stack size for all threads to work around some crashes.
+
+* Made the Xcode source accessor more robust on OS X by ensuring that the Xcode project required exists and allowing the user to switch to a different application while waiting for Xcode.
+
+## Playstation 4
+
+* New: Added RE_Z feature to shader definitions.
+
+* New: Added saving of Crash Context during coredump.
+
+* New: Added support for conservative depth optimization to PS4.
+
+* New: Added support for the FirstInstance draw optimization. Reduces the number of RHI commands generated for instanced drawing.
+
+* New: Added support for the new Wave shader compiler. Enabled by default. Controlled by console variable r.PS4UseWaveShaderCompiler.
+
+* New: Added support for the shader complexity view mode to PS4.
+
+* New: Added the option to strip PS4 shaders of optional debug data. Enabled by default. Saves approximately 75% of the memory cost of shaders. CVAR is r.PS4StripExtraShaderBinaryData and is enabled by default.
+
+* New: Always create a pak file when packaging for distribution and allow distribution builds of non-shipping configs.
+
+* New: Disabled DispatchDraw support form LCUE to reduce CPU overhead as this feature is not currently used.
+
+* New: Enabled compiler warning for shadowed variables and fixed known instances of the warning.
+
+* New: GPU memory allocator for PS4 has been changed to pool small allocations. Fixes problems some licensees were seeing with allocator performance gradually degrading over time.
+
+* New: MediaFramework movie playback on PS4 now uses the new AvPlayer decoder. Conversion to RGBA is now performed on the GPU.
+
+* New: PS4 ContinuousSubmit toggle is now available in shipping.
+
+* New: Reduced GPU Context Rolls by only setting setCbControl, setActiveShaderStages and changing ring buffers when needed. Saves about 100 context rolls and 0.1ms of GPU in RenderTestMap.
+
+* New: Removed support for the older 60hz Morpheus kits.
+
+* New: Support for SDK 2.0 has been removed.
+
+* New: Unused files are now deleted from the device during iterative deployment.
+
+* New: Updated the Windows DualShock plugin to work with PS4 SDK2.500.
+
+* Fixed a bug where leaving a PSN session would not leave the session properly.
+
+* Fixed a bug where MediaFramework audio playback was not working on PS4.
+
+* Fixed a crash when using the experimental LightGrid.
+
+* Fixed a race condition in the ps4 temporary block allocator that could cause intermittent crashes.
+
+* Fixed an intermittent crash in reflection captures.
+
+* Fixed audio not resuming after being paused.
+
+* Fixed crash in DrawPrimitiveUP when using parallel rendering on PS4.
+
+* Fixed Cursor CurrentPosition not getting initialized.
+
+* Fixed DevKitUtil project to use SCE env var for ORTMAPIlib path instead of hardcoded path.
+
+* Fixed failing to close the current running UE4 app when packaging which can cause the packaging to fail due to locked files.
+
+* Fixed handling of failure to launch elf when using PS4DevKitUtil.
+
+* Fixed huge timeout when automatically trying to turn on a devkit that is unreachable. Devkits must now be turned on manually.
+
+* Fixed memreport -full crashing inside of PS4 audio buffers.
+
+* Fixed movie streamer leaking audio ports.
+
+* Fixed permanently overriding the cache type of a texture when setting a UAV.
+
+* Fixed PS4 rendering artifacts when dbuffers are enabled, but there are no dbuffer decals rendering.
+
+* Fixed recenter command for Morpheus HMDServer builds.
+
+* Fixed Scissor GNM Validation errors.
+
+* Fixed tearing during startup movie playback. Also we now tick the startup movieplayer at a higher rate for smoother playback..
+
+* Fixed texture updating for certain texture formats.
+
+* Fixed the reported 'Physical' memory pool size for PS4 to correctly match the size of the main CPU memory heap.
+
+* Fixed TMap visualizer and added missing visualizers to the .natvis.
+
+* MediaFramework no longer waits a full 10 seconds to start movies on PS4 regardless of the time to actually start up the movie.
+
+* Morpheus projects will now fall back to normal 2D rendering rather than crashing if Morpheus is not detected on boot.
+
+* Optimized vertex and index buffer write locks on PS4 to no longer require an extra memcopy.
+
+* Path shortening is no longer applied when reading from /hostapp to support custom workflows.
+
+* PS4 draw call stats are now in a verbose stat group to avoid stat spam.
+
+* PS4 now always attempts to pan for surround sound.
+
+## VR
+
+* New: MSAA is now enabled by default when running on GearVR on Galaxy Note 4 and Galaxy S6 devices. You can disable this by setting gearvr.EnableMSAA=0 in your ConsoleVariables.ini file.
+
+* Fixed SteamVR public module interface class not hooked up to actual module class.
+
+* Fixed stereo rendering not working when Playing in VR while Blueprint editor was open.
+
+## HTML5
+
+* New: HTML5 networking plug-in is now enabled by default.
+
+* New: HTML5 SDK is now included with the Engine.
+
+* New: Packaging HTML5 on Mac deploys a double-click script for HTML5LaunchHelper to avoid the need to use the command line.
+
+* New: Reduced emitted build size by using compression.
+
+* New: upload to S3 for serving HTML5 builds.
+
+* Added the ability for HTML5LaunchHelper to server to remote machines not just localhost.
+
+* Changed the location of .emscripten SDK config file to no interact with the any other pre-installed SDK on the machine.
+
+* Fix issue with HTML5 web server not flushing sockets. Fixes issues with clients not disconnecting correctly.
+
+* Fixed 3d audio spatialization on HTML5 and Linux.
+
+    * Unified HTML5 and Linux OpenAL implementation by having the HTML5 platform use the Linux implementation.
+
+* Fixed multiple HTML5 fullscreen issues across browsers.
+
+* Fixed texture read issue that caused landscapes to appear mostly black in HTML5.
+
+* We now properly copy the web server as part of the packaging process.
+
+## Windows
+
+* Changing OS Version for minimum supported OS on windows, no longer requires an editor restart.
+
+## All Mobile
+
+* New: Changed how non-power-of-two textures work on platforms that don't support them. Instead of silently failing, it will replace contents with an "error texture" so you can see on device what to do, instead of seeing nothing.
+
+* New: Games now have more control over the login process in Game Center on iOS and Google Play on Android.
+
+    * The login UI no longer appears automatically during engine startup if the iOS or Google Play online subsystems are enabled. This was causing various timing issues, and games had no way to control when it appeared.
+
+    * Games must now explicitly invoke the system's login UI by calling either "IOnlineExternalUI::ShowLoginUI" from C++ or "Show External Login UI" from Blueprints. These are asynchronous operations that allow the game to defer any other online operations, such as achievement or leaderboard queries, until after a user has successfully logged in.
+
+    * Added a new Blueprint node "Is Logged In" to get the login status of a user.
+
+    * Implemented the C++ function IOnlineIdentity::Logout and added a new Blueprint node "Logout" for the Google Play online subsystem.
+
+        * This feature, along with the external login UI support, should allow games to implement a login/logout button.
+
+        * Note that this functionality is not supported on iOS as Game Center itself does not support explicit logout from within a game.
+
+* Remove redundant checks for Pixel-Buffer-Object support from the OpenGL RHIUpdateTexture2D & RHIUpdateTexture3D functions to allow them to be called on mobile. (Contributed by [REDACTED_PERSON], pull request 1086.)
+
+## Rendering
+
+* New: Added a material function that can convert RGB colors to HSV.
+
+* New: Added ability to create a texture UAV to a specific mip and create UAVs for additional texture types.
+
+* New: Added ability to create SRVs for additional texture types (3D, 2DArray and Cube).
+
+* New: Added experimental volumetric decals defined by a distance function. Can be useful, but might not be fast enough for some applications.
+
+* New: Added max displacement to materials to fix tessellation frustum culling.
+
+* New: Added View Property material expression with multiple values to expose more things to the material editor without using Custom.
+
+* New: Added wider Gaussian filter kernel to allow High Res Screenshot having larger bloom and Gaussian DOF. Added console variable "r.Filter.LoopMode" for experiments.
+
+* Improved Visualize DOF, especially for Circle DOF.
+
+* New: Updated Light Propagation Volume from Lionhead Studios.
+
+    * Lower order SH is faster but a more blurry.
+
+    * Directional Occlusion can be enabled in Settings. Affects Sky Light and Reflection Environments but not Ambient Cube.
+
+    * Material "Block GI" feature.
+
+    * Spotlight support. Uses AsyncCompute on XboxOne, but not order optimized yet.
+
+* New: Procedural Mesh Component improvements and fixes.
+
+    * Add 'fast path' Update Mesh Section function to Procedural Mesh Component. Updates collision as well.
+
+    * Add Set Mesh Section Visible & IsMeshSectionVisible functions.
+
+    * Fix crash in Procedural Mesh Component when Triangles array was not a multiple of 3.
+
+* Altered the sprite material function. Removed the expansion vector normalization to allow for non quad meshes.
+
+* Exit gracefully when running out of video memory and notify the user.
+
+* Fixed hierarchical instanced static meshes. The cluster tree now gets rebuilt whenever there are unbuilt instances.
+
+* Fixed 16-bit float normal map detection for texture import.
+
+* Fixed a bug that would cause parallel rendering to be disabled in shipping builds.
+
+* Fixed a bug with the depth-write performance warning. Performance warning is also now more informative.
+
+* Fixed a crash when destroying Media Framework assets.
+
+* Fixed an ensure when using Responsive AA on Canvas tiles.
+
+* Fixed black viewport in split screen with full screen enabled.
+
+* Fixed Blendables in Post Process Component losing the array elements on load.
+
+* Fixed Crash when using of Circle Depth Of Field in Shader Model 4.
+
+* Fixed crash when using Post Process Material Blendables with weight of 0.
+
+* Fixed Depth Of Field material function for Circle DOF.
+
+* Fixed ElementalDemo crash due to nil constant buffer.
+
+* Fixed environment reflections not showing up on SM4.
+
+* Fixed issue in which alpha-only visualization of textures would display incorrect results with sRGB set to true.
+
+* Fixed issue preventing custom and procedural mesh components writing to custom depth pass.
+
+* Fixed missing rendering of some translucent geometry, editor/debug material had SeparateTranslucency on, which often it shouldn't be.
+
+* Fixed multiple issues and extended transform position especially for velocity computations (affecting Motion Blur and Temporal AA).
+
+* Fixed rendering artifacts in some editor viewports caused by wrong matrix computations.
+
+* Fixed Responsive AA feature which got broken by recent Temporal AA optimization.
+
+* Fixed scene captures showing incorrectly when switching to fullscreen.
+
+* Fixed shader complexity not working with decals.
+
+* Fixed SlaveComponent to support LOD change.
+
+* Fixed some issues in which the landscape gizmo would not render correctly until moved.
+
+* Fixed skeletal meshes with multiple translucent batches to render correct translucency types.
+
+* Fixed Scene Color node to render correctly in all modes.
+
+* Fixed Motion Blur exaggeration when Paused in Simulate mode.
+
+* Fixed shadows being occluded inappropriately for players who are not on-screen in split-screen mode.
+
+* Fixed crash when switching to Shader Model 4 with Light Propagation Volumes active.
+
+* Fixed Circle Depth of Field to be independent of Resolution.
+
+* Fixed wrong colors and artifacts when using color grading (extreme red in darks or more subtle ones).
+
+* Fixed wrong texture format when initializing an SRV.
+
+* Improved the error message logging in FScene::Release() when a component still has a reference to the dying scene to list the associated asset for any kind of component, not just UStaticMeshComponent.
+
+* Remove unused ShapeMaterial from ShapeComponent.
+
+* Reverting diffuse model back to Lambert for simplicity with capturing materials.
+
+## FX
+
+* New: Added DefaultValue property to Dynamic Parameter material node. Minor workflow improvement for VFX artists.
+
+* Fix for bug in SkelVert/Surf location cascade module in which skeletal mesh vertices would not be accessed correctly for some meshes.
+
+* Fix for rare crash in SkelVert/Surf Location Cascade module.
+
+* Fixed a bug that caused exporting emitters to a particle system with a higher number of LODs in the first emitter to crash the editor if the emitter to export was set to solo mode.
+
+* Fixed a problem with unclamped alpha during building of lookup tables for distribution values, that could cause unpredictable behavior.
+
+* Mesh particle collisions now consider the particle size by default. This should produce more accurate collisions for mesh particles with regular shapes. This can be disabled via the MeshTypeData property bCollisionsConsiderParticleSize.
+
+* Removed "Particle" options from Beam Source and Target options as these features are not implemented in UE4.
+
+## Lighting
+
+* New: Added bUseVertexNormalForHemisphereGather Lightmass primitive setting for foliage static lighting, in cases where artists want foliage lighting to match the underlying terrain.
+
+* New: Volume lighting samples are no longer placed inside geometry.
+
+    * Volume lighting sample computation is now multithreaded in Lightmass so it can't bottleneck the lighting build.
+
+* Added lighting cache invalidation step when changing static mesh to avoid rendering stale data.
+
+* DDC commandlet now blocks on async mesh distance field tasks and actually writes the results to the DDC.
+
+* Fixed a bug where only emissive objects would render into the LPV.
+
+* Fixed an issue where self-shadow-only meshes would only shadow when selected in the editor.
+
+* Fixed bugs with the clear coat shading model.
+
+* Fixed crash building lighting due to zero area BSP nodes creating 0 resolution texture mappings.
+
+* Fixed lightmass exporter being deleted before the swarm connection is closed, causing occasional crashes when canceling a lighting build.
+
+* Fixed static objects rendering as fully lit when they should be fully shadowed.
+
+    * This was reported for the case of having one landscape below another, but also affected static meshes or BSP.
+
+* More graceful handling of overflowing object count for distance field shadowing and DFAO.
+
+* Objects with 0 scale are excluded from the distance field scene, fixes bugs with several lighting methods with 0 scale objects.
+
+* Primitives affected by the indirect lightning cache will now be correctly updated when a precomputed lighting volume is streamed in or out.
+
+* Updated grouped shadow caster bounds calculation to ignore non-shadow casting primitives.
+
+## Materials
+
+* New: A few material functions have been added for generating and combining Shapes. These are mostly useful for the "Volumetric Decal" material type, but they may be useful in other situations as well.
+
+    * Shapes: Sphere, Box, Capsule.
+
+    * Combiners: Intersection, Subtraction, Union.
+
+* New: A few new material functions have been added for creating Shapes in the material editor.
+
+    * DrawLine-2D and Drawline-3D: Specify a line in either 2D or 3D space. Has the option to behave as either capsules or transformed boxes.
+
+    * BoxMask-2D and BoxMask-3D: Allow you to create 2D or 3D gradients similar to the "SphereMask" node but in the shape of a box.
+
+* New: A material function for generating a BullCurve has been added.
+
+    * This is a simple bell curve aka "Normal Distribution" that has been re-scaled so that that inputs of -1 and 1 on X return 0. No control params. Uses the simplest form of the distribution: f(x) = e^(-0.5*x^2)
+
+* New: Added a Blueprint macro to convert Blueprint arrays into HLSL code that can be copied.
+
+    * To use, you simply need to specify what type of input array you want (float through float4), a variable name and then pipe the results to an editable String variable. Then you use the editable text field to copy the result into the place where it is needed (for example inside of a Custom node in a material).
+
+    * Note: In general using large arrays in HLSL will be slow. Often this is something more useful when you are trying to do something 'offline rendering-ish' inside of UE4, such as pre-computing a bunch of random ray directions for some kind of raytracing experiment.
+
+* New: Many of our VFX materials require HDR (high dynamic range) pixel values (Values above 1) in order to properly generate bloom and other engine effects, however multiplying an LDR (low dynamic range) texture into a high range often looks awful. This material function takes an LDR texture as an input, and remaps it to an HDR range defined by the user, This maintains the look and preserves the chroma of the input image and remaps it into HDR.
+
+* New: New material function 'CalcLightSourceAngle' can be useful to figure out the light source angle of a light as seen from any position in the world.
+
+* New: New material function 'Sphere_AO' returns an ambient occlusion value for any point in the world based on a sphere occluder of a given position and radius. Useful for blob shadows that fade in and out realistically.
+
+* New: New material functions for Sampling the Shadowing of a sphere based on lookup textures.
+
+    * Sphere-ConeShadow-Texture: Returns the shadow response from a sphere in the world given a position and radius.
+
+    * Ellipsoid-ConeShadow-Texture: Same as the above Sphere shadow but stretched along a specified axis. NOTE: This version is currently experimental and has a few issues. These issues will be fixed in the next release.
+
+    * Both functions use the following default lookup texture: Texture2D'/Engine/EngineMaterials/T_SphereShadow_RGBA_5-15-30-60deg.T_SphereShadow_RGBA_5-15-30-60deg'
+
+        * The texture is channel packed such that each of the channels has a different light source angle encoded. R=5, G=15, B=30, A=60 degrees.
+
+* New: Scalar material parameters can specify the slider min and max for the material instance editor.
+
+* New: Spherical-Cap-Intersection. This material function calculates the area of a spherical cap intersection given a light source angle and angle of visibility. Has option to use cheaper 'SmoothStep' calculation. In general this function is useful for pre-baking lookup textures but it can also be used in realtime when using the SmoothStep approximation.
+
+* Custom UV pins on Make/Break material attributes nodes now hide and show under the same conditions as their counterparts on the material result node.
+
+* Fixed rare crash in material editor when using custom UVs.
+
+* Fixed a crash that occurred when overriding a distortion material to be opaque in a material instance.
+
+* Fixed crash sometimes caused when reparenting a material instance.
+
+* Fixed issue in which all materials referencing a texture that is re-imported would turn black.
+
+* Fixed issue in which the tangent space normal property of materials would not function correctly when using material attributes.
+
+* Fixed issue where some parameters would not appear in material instances.
+
+* Fixed many transform nodes being broken in decal materials
+
+* Fixed materials that use Pixel Depth Offset not using the pixel shader during the prepass causing depths not to match up during the base pass.
+
+* Fixed wireframe materials to be considered "not writing every pixel". This fixes a few rare interactions causing corruption in shadow casting.
+
+* Modified the camera offset material function node to reference vertex world position instead of actor position to remove z-depth scaling.
+
+## Postprocessing
+
+* New: Improved Blendables, exposed weight on each blendable entry, better blueprint access.
+
+* New: Replaced experimental cylindrical postprocess projection with better but still experimental Panini projection.
+
+* Fixed blacks of old tonemapper not being black with new log space LUT.
+
+* Fixed light shaft rendering when used with upsampling.
+
+* Fixed pause with Temporal AA.
+
+* Improved SSAO quality with non-power-of-two viewport sizes, as recent optimizations made that worse.
+
+* Improved Temporal AA quality using bicubic filtering (slightly sharper).
+
+## Optimizations
+
+* New: Don't use IndirectLightingCache shaders when there are no precomputed lighting volumes.
+
+* New: Exposed console variable to access faster lower quality SSAO (r.AmbientOcclusionSampleSetQuality).
+
+* New: Implement the ScreenSpaceReflections stencil pre-pass (r.SSR.Stencil).
+
+* New: Indirect lighting cache point samples don't use the texture layout anymore, since they don't use the volume texture anyway.
+
+* New: r.ReflectionEnvironment is now usable in shipping builds.
+
+* New: Reflection environment console variables are now another scalability option.
+
+* Fixed the initialization of the texture streaming system so that it's initialized early enough to stream in textures at startup. Before this change, it would kick in a little bit later which could temporarily overload texture memory at startup.
+
+* Updated distortion pass to use a stencil mask which improves performance with lower screen coverage.
+
+## Mobile Rendering
+
+* New: Enabled limited post process support for devices that do not support full HDR rendering, meaning those which lack float16 render target support (e.g. Galaxy S6).
+
+    * This replaces "mosaic" rendering when the color_buffer_half_float extension is not available but a framebuffer fetch gl extension is.
+
+    * Bloom is supported.
+
+    * Unlike "mosaic" mode, there are no resolution limitations.
+
+* New: Modulated shadow rendering for forward rendering (mobile path).
+
+    * To use enable 'Cast modulated shadow' on a stationary directional light.
+
+    * Dynamic shadows will then be rendered by modulating the light's 'Modulated Shadow Color' against the scene color.
+
+* Fixed stencil buffer rendering when using the Metal RHI.
+
+## UI
+
+* New: In game objects Mouse Enter and Leave Events no longer fire if UI (Slate/UMG) is over it.
+
+* New: Materials designed for UI now have a new material domain that can be set in the material editor. This replaces the "Used with UI" checkbox on materials. All existing materials using this checkbox will be automatically updated.
+
+## Slate
+
+* New: Editable text widgets in Slate now allow the context menu to be override or disabled.
+
+* New: Empty lines in a multi-line editable text now draw a little highlight so that you can tell that the line itself is selected.
+
+* New: Several improvements to the way sRGB is treated in the engine. Trying to standardize on the sRGB standard instead of pow(2.2). The key take away is - your textures should look like they do in Photoshop when rendered by Slate.
+
+    * Now using a more accurate Gamma -> Linear space conversion function for converting FColor to FLinearColor.
+
+    * Changing the way all FColors are converted into FLinearColor by default. Previously, we used pow(color, 1/DisplayGamma) for gamma correct. We now use the IEC 61966-2-1 standard.
+
+        * Fixed a problem with Slate quantizing the lower range of black colors.
+
+        * In Slate, you should now see an image matching Photoshop instead of being slightly darker in the lower ranges. However, because we don't want to invalidate all existing textures that users have authored, all existing UTextures have a UseLegacyGamma flag set to true, all new textures will be set to false.
+
+* New: Simplified Raw Texture Update function by using FSlateTextureData to pass the texture data and size to the render thread.
+
+* Activating a window now sets input focus for all users.
+
+* Clicking on the bottom/right most pixel of a slate widget now registers the click.
+
+* Fixed "Delete" on an empty line in a multi-line editable text being handled like a "Backspace".
+
+* Fixed an intermittent crash when using Ctrl+Arrow to move the cursor in an editable text control.
+
+* Fixed problem with menus not closing when clicking in another window on Mac.
+
+* Fixed SMultiLineEditableText Slate widget's ability to select and copy text from it.
+
+* Fixing a few issues with memory leaking in the Slate RHI resource manager. There was never a pass to remove stale entries in the dynamic resource map, so now we occasionally purge it of stale entries.
+
+* Fixing several issues involving multi-touch. Data stored involving widgets under cursor are now bucketed based on user and cursor index so that multiple pointer devices, as well as index can independently hold separate widget under cursor data.
+
+* Improved functionality for getting first chance access to key down events in SEditableText, SEditableTextBox, SMultiLineEditableText, and SMultiLineEditableTextBox.
+
+* Menus created from menu anchors no longer auto-dismiss when a parent (in the menu stack) receives focus.
+
+* Multiple causes to the often seen early clipping in Slate have been fixed. You might have sometimes noticed a column or two of pixels on the bottom or right of a widget getting cut off, those issues are now fixed.
+
+    * Fixing the incorrect calculation of the clip rectangle which resulted due to pixel snapping along the extents of the box, instead of pixel snapping the corners of the clip box.
+
+    * Disabling the float16 precision version of the slate shader on all platforms. Since the clip values are all in window space, it's terribly easy on modern monitors to exceed single digit precision (2048), causing multiple columns worth of additional pixel clipping erroneously.
+
+* Optimized the Slate Draw element batcher to perform fewer memory allocations per frame and pass off some processing to the render thread earlier to avoid additional memory copies.
+
+* Platforms using a virtual keyboard no longer attempt to show it when focus is received by a read-only text box.
+
+* Slate windows now take into account required border/tiitle bar sizes when resizing the window. This ensures that the client area stays at the requested size.
+
+* Text margins are now consistently calculated at all zoom levels and no longer do an immediate calculation, they instead just invalidate layout like all other operations do.
+
+    * Additionally there is now a step in the UpdateLayout pass after justification and flow have been calculated, not doing it last leads to a lot of bugs and inconsistency.
+
+    * This fixes the issue with zoom negatively affecting margins as well as the double application of margins to the justification of text, additional offsets were being calculated into the offset location of the text.
+
+* The SButton no longer attempts to handle MouseUp, if the button wasn't pressed. Previously pressing outside the button, and then releasing the mouse inside the button would cause the button to attempt to handle it, even though the down never originated on the button.
+
+* The STableViewBase now properly handles touch events when under editor conditions and sets up the callback to handle the inertial scrolling.
+
+* UE-18600 - Custom Draw elements no longer are attempted to be early clipped using the known scissor rect before going through the element batcher, since their clip rect info isn't valid we just ignore doing this step for those elements.
+
+* When we update the RHI data, we now also update what the slate render proxy thinks is the actual size of the texture, which might change due to a reimport.
+
+## UMG
+
+* New: Add features to allow for better in-timeline editing of animations.
+
+    * Add combo buttons to object binding tracks for adding additional property tracks.
+
+    * Add widgets for editing key values to the animation outliner.
+
+    * Add buttons for navigating to previous and next keys, and for adding a key at the current time to the animation outliner.
+
+* New: Added some new methods to help convert positions inside a widget to pixel positions in the viewport. LocalToViewport & AbsoluteToViewport, they take the geometry data for a widget and a given location, and can give you both the Pixel and Viewport (Slate Unit) Position.
+
+* New: Added support for dragging multiple widgets to reparent them in the UMG Hierarchy view.
+
+* New: Adding a few extra scaling modes to the scale box, allowing it to Scale to Fit X, and Scale to Fix Y.
+
+* New: Adding the blueprint function CancelDragDrop, to allow you to programmatically cancel any drag/drop operation.
+
+* New: Adding the Highly New & Experimental SInvalidationPanel in Slate and the corresponding Invalidation Box in UMG.
+
+    * Allows widget geometry to be cached to speed up Slate rendering. Widgets that are cached by an Invalidation Box are not Prepassed, Ticked or Painted. It's up to the user to call Invalidate (C++) or InvalidateLayoutAndVolatility (BP) on a child widget to force invalidation. Some core widgets do this automatically when certain properties are changed, more will do it over time.
+
+    * You can force an invalidation directly on the panel too, by calling InvalidateCache, usually this is not required, but can be in certain situations - such as before child widgets have a chance to learn about their parent invalidation panel next tick.
+
+    * Hit testing has not been tested heavily. You'll find more success caching widgets that are not hit testable, like in a HUD.
+
+    * You may notice odd draw ordering issues in this version in heavy nesting of invalidation panels.
+
+    * Currently caching only on the CPU. Additional improvements coming in the future to include caching on the GPU to make the panel even faster.
+
+    * Any widget can be set to be Volatile. Volatile widgets act like normal Slate widgets pre-invalidation. They're redrawn every frame, including all their children. When combined with the invalidation panel, it allows you to only need to care about redrawing only the most dynamic bits of the UI, as invalidating a whole panel could be far more costly.
+
+    * You can debug your invalidation panels using the widget reflector, with the [Invalidation Debugging] toggle.
+
+        * White/Red Flash = Widget responsible for invalidation
+
+        * Green Border = Invalidation Panel (Cached This Frame)
+
+        * Red Border = Invalidation Panel (Cache Invalidated This Frame)
+
+        * Dashed Yellow Border = Volatile Widget
+
+* New: Animation outliner node selection now respects control (toggle) and shift (select inclusive).
+
+* New: Animation outliner nodes are now sorted alphabetically.
+
+* New: Double click on a section will now focus the timeline range to the bounds of the section. Double click again will return to the previous timeline range.
+
+* New: DrawLines, DrawFormattedText added to the set of low level painting methods for UMG.
+
+* New: Enabled dragging of multiple widgets to move or reparent in the UMG Designer view.
+
+* New: Exposing the Next/Previous widget navigation options on all widgets.
+
+* New: Exposing the SScrollBox's ScrollIntoView function in a much better way.
+
+    * We no longer allow/expect the user to pass in geometry for the scrollbox, instead they simply request the widget to scroll into view and on the next layout pass we make that happen.
+
+    * The scroll into view function is now also exposed to blueprints for use in the UMG version of the scrollbox.
+
+* New: Removed appearance override properties from UMG widgets. Style property is now used.
+
+* New: The WidgetComponent has been reworked to use the new IGameLayer system to create a single shared canvas layer that all the components add themselves to when they're displaying their information in screen space, instead of world space. They now properly sort based on depth/distance from the camera.
+
+    * Additionally, the layer widget now updates the positions for the widgets. Thereby solving the problem of tick groups causing weird update orders for objects touching the WidgetComponent's Actor as a base. The WidgetComponent now has a normal tick group.
+
+* New: Tooltips when moving widgets in the UMG Hierarchy view now match the style of those in the World Outliner.
+
+* New: Widget Components in World Space, now support for 3 blend modes, Opaque, Transparent, Masked (new). Masked allows you to have the transparency you need on edges, but writes to the depth buffer, so that it sorts correctly. There is also now an option for ticking when the in world widget is off-screen to continue to allow animations to play and tick to be called on the widget.
+
+* New: Widget Components now have a parabolic distortion option to give them a sort of curved glass appearance, this only works if they're rendered in World Space.
+
+* New: You can now use the arrow keys in the designer to nudge widgets in the canvas by the grid snap amount. If grid snap is enabled it will actually keep it snapped to the grid, otherwise it will just move by the current snap amount absolutely.
+
+* A WidgetComponent with Space set to World, will no longer crash if the draw size is set to 0.
+
+* AddChildVerticalBox renamed to AddChildToVerticalBox.
+
+* Drag/Drop now properly maps the cursor to a user widget and no longer shows the native cursor in addition to the custom cursor if any.
+
+    * Also fixed a crash on mobile where the platform doesn't have a cursor object but we still attempted to interact with it on drag/drop operations.
+
+* DrawLines, DrawFormattedText added to the set of low level painting methods for UMG.
+
+* Fix several issues with UMG widgets crashing during map transitions.
+
+* Fixed a potential crash with cooked content.
+
+* Fixed incorrect default rotation for widget components. New widget components now face toward +X by default.
+
+* Fixing a bug in SObjectWidget that was preventing handled key presses from being returned as handled to the rest of slate.
+
+* Fixing a bug that was showing the category field in the details panel when nothing is selected.
+
+* Fixing a bug where the modifier keys and additional base class data were never passed to new copies for mouse pointer events. Preventing users from checking modifier key state in mouse events on UserWidgets. Only affected UMG and not Slate, as the VM makes copies of the struct.
+
+* Fixing animation and tick rates for widgets living in the world, they should now correctly be controlled by time dilation as they no longer are getting ticked twice, and the second tick using the wrong delta time. This broke when we introduced tick folding, which ticks again inside of paint, but continued to use the global slate delta time, as well as tick before hand.
+
+* Fixing several issues with brush and material utility functions not working correctly with the pins for input being displayed in the blueprint as output pins.
+
+* Fonts and other appearance properties on editable text widgets are now updated in real-time in the UMG editor.
+
+* If the owner player changes on a widget component, we now remove the widget from the screen if it's currently already on a user's screen to move it to the other users screen.
+
+* Keep animation outliner selection consistent across undos.
+
+* Making some improvements to the way the SScaleBox handles fill justification so that it actually fills the entire contents of the scalebox, and handles additional scale being applied to it.
+
+* Moving widgets in the UMG Designer no longer makes them move to the last slot in their parent.
+
+* Removing Thickness from the line painting function. It wasn't actually hooked up to anything, and there's no support in slate for it currently, only when rendering splines. Will add back once we've added support for it in Slate.
+
+* Rounding the projected location when calculating screen position from world position before applying inverse DPI scale, this produces a much more stable location with no half-pixel offsets due to later applied scale, since we do it before the inverse we more or less end up at the whole number we expect. This makes things like widgets above characters heads much more stable.
+
+* SConstraintCanvas no longer attempts to draw elements outside its view, and now draws them in the correct order.
+
+    * Making a lot fewer calls to the enabled delegate when painting, prefiltering arranging children that are not visible and avoiding the math for them.
+
+    * Now consistently drawing children in the correct order. Previously we only sorted based on ZOrder, and assumed a stable sort for equal elements, but this was not the case. So occasionally the order just appeared wrong, unless you manually set the ZOrder, that's no longer required.
+
+* The cursor property on widgets now works as it should. You can also call SetCursor and ResetCursor to set it at runtime on any UWidget, or clear your overridden cursor.
+
+* The Multiline Editable Text widgets should now properly notify when text is committed or changed.
+
+* The slatebrush's ResourceObject property is no longer directly accessible in a blueprint. You'll need to use the functions, GetBrushResource, SetBrushResourceToTexture, and SetBrushResourceToMaterial to change the value. This is to prevent the wrong types from being associated with the resource object.
+
+* The STooltipPresenter now calculates a post transformed, rounded layout position to ensure that tooltips do not start on a half pixel. This reduces a lot of the inner jittering of a tooltip as you move your mouse across the surface of a widget.
+
+* The Widget compiler now warns if you have events that return EventReplys that are unconnected to anything.
+
+* UE-17226 - The assignment operator for the FPaintContext wasn't copying the MaxLayerId, painting into layers and not telling slate they exist causes the graphics corruptions seen by several people.
+
+* UE-18029 - Making the SlotAs functions more robust against crashing when called at the wrong time from blueprints.
+
+* Updating the 3D rendering path in Slate the properly set the RT and clear at the same time it to fix a bug on iOS when using a transparent Widget Component.
+
+* Widgets removed by replacing them in the UMG editor hierarchy are moved to the transient package to stop name clashes.
+
+* Working on several general fixes to the navigation system.
+
+    * Widgets other than the focused widget can respond to navigation requests with explicit or delegate based handling. Adding a new function to allow querying the if descendants of a widget are focused by a specific user.
+
+    * Implementing controller support for the SScrollBox; it has custom logic for handling navigation inside of it scrolling the next element into view for the focuser and moving focus to the next logical child.
+
+# Programming Release Notes
+
+## AI
+
+* New: Added InvalidOrientation constant to AISystem.
+
+* New: AIController class now derives from IGenericTeamAgentInterface which will make team-handling easier in the near future
+
+* New: FAIResourceLocks now have the ability to count the number of locks for each LockPriority. (Default behavior remains the same, this functionality is an additional feature.)
+
+* FEnvQueryInstance::ItemIterator's "operator bool()" is now explicit. That prevents accidental (implicit) conversion to an int32 when you intend to get the index.
+
+* Fixed pawn action sequence not responding correctly to abort event.
+
+* Replaced EQS skip item mechanic with forcing item status.
+
+### Behavior Tree
+
+* Added support for behavior tree's action task losing its PawnAction due to some external events.
+
+* Fixed multiple default subobjects with the same names created for behavior tree's blackboard filters.
+
+* Unified behavior tree's StopTree & StopLogic functions.
+
+### Navigation
+
+* New: Added pruning for crowd agent's corridor when close enough to next intermediate point (corner vertex). Depends on angle between first 2 segments of path, should reduce twitching when large group is moving through narrow corridor.
+
+* Changed FNavMeshPath::DoesIntersectBox to ignore agent location, as it's not relevant to function's result.
+
+* Fixed Navmesh generation failing after hot-reloading
+
+* Fixed early out condition in navigation tile generator using async data gathering.
+
+* Fixed invalid pointer to owner of crowd agent's custom filter.
+
+* Fixed rare case of infinite loop in navmesh pathfinding.
+
+* Navigation system will store weak object pointers to all custom navlinks to ensure their validity.
+
+## Animation
+
+* New: Added DLL export to some skeletal control classes and methods, making them extensible in another module.
+
+* New: Allowed skeletal controls to specify the coordinate space that they would prefer the gizmo to be displayed in by implementing GetWidgetCoordinateSystem
+
+* Additive blending optimizations in FAnimationRuntime.
+
+* Replaced ApplyAdditiveSequence with AccumulateAdditivePose. No pose copy and fast path for weight == 1.
+
+* Replaced BlendRotationOffset with AccumulateMeshSpaceRotationAdditiveToLocalPose. No pose copy, removed one pose allocation, fast path for weight == 1, and mesh space transform done only on rotation part and not full FTransform.
+
+* Added ConvertMeshRotationPoseToLocalSpace.
+
+## Blueprints
+
+* New: BlueprintCallable functions can now specify the metadata HideSelfPin to indicate that the node should not show the Target pin, nor should it be connectable, thus filtering it from the context menu when dragging off a pin and ensuring it can only be used in the self context.
+
+* New: External modules can now append toolbar buttons to the Blueprint editor, using the BlueprintEditorModule's MenuExtensibilityManager.
+
+* New: External modules can now extend the Blueprint context menu by adding their own filtering through FBlueprintGraphModule::GetExtendedActionMenuFilters().
+
+* New: Graph Editor: Allowed graph nodes a chance to create a custom SGraphNode visual representation without having to manually register a factory (override the new method CreateVisualWidget).
+
+* Fixed a bug that caused instanced properties to be cleared on Blueprint compile.
+
+* Fixed a bug where functions with inaccurate WorldContext metadata were not accessible.
+
+* Fixed a bug where globally declared delegates could be mismatched with Blueprint nodes originally using a like named delegate, belonging to a specific class.
+
+* Fixed runtime blueprint use of atmospheric component.
+
+## Core
+
+* New: Add ForEachObjectWithOuter and ForEachObjectOfClass utility functions that will call a TFunctionRef for each object instead of using GetObjectsWithOuter/OfClass which requires copying each object in to a (potentially very large) array.
+
+* New: Added a console command, "abtest" to asses the impact of changing a console variable on performance in a dynamic, noisy gameplay session.
+
+* New: Added a console variable to control TimeBetweenPurgingPendingKillObjects at runtime.
+
+* New: Added basic support for flushing async loading of a specific package only.
+
+* New: Added command line param to dump commandlet stats to log file.
+
+* -LoadTimeStatsForCommandlet dumps all load time stats (LinkerLoad and AsyncLoad groups).
+
+* New: Added command line parameter ("-fullcrashdump") for forcing any Windows crash minidumps to contain full memory info.
+
+* New: Added logging to builds about the command line used to invoke UnrealHeaderTool and the duration.
+
+* New: Added a few more stats for async loading.
+
+* New: Added an extra parameter to FPlatformProcess::ConditionalSleep to be able to specify sleep timeout.
+
+* New: Added checks to ensure that two different custom versions aren't added with the same GUID.
+
+* New: Added checks to UnrealBuildTool for modules which are both a dependency and dynamically loaded.
+
+* New: Added debug.EnableLeakTest console command for stats based memory profiler testing.
+
+* New: Added FChar::ConvertCharDigitToInt function to convert digit characters to integer type.
+
+* New: Added FPaths::LaunchDir(), which is the directory the application was launched from. Useful for finding files relative to working directory for commandlets.
+
+    * Improved UnrealPak to use LaunchDir to find .pak files in any location. (UnrealPak, to extract .pak files, is often run from non-standard locations.)
+
+* New: Added FString::Reset() which allows emptying a string without reallocating memory.
+
+* New: Added functions to access low-level file stat data as a single file-system request.
+
+    * You can now use GetStatData to get all of the stat data for a file (its timestamps, size, read-only state, and whether it is a file or a directory) in a single file-system request.
+
+        * Depending on the platform, this can be much more efficient than making multiple requests
+
+    * This change also adds alternate directory iteration functions (IterateDirectoryStat and IterateDirectoryStatRecursive) for when you need the stat data along with the files/directories on disk.
+
+        * For platforms that provide this stat data as part of the low-level directory iteration (such as Windows), this can yield much better performance than making two separate file-system requests.
+
+        * Platforms that do not provide this information as part of the directory iteration will just perform a stat request while iterating.
+
+* New: Added guards to make sure linkers are only being deleted by linker manager.
+
+* New: Added guards to make sure two identical default subobjects can't be created in a class constructor.
+
+* New: Added implementations of VectorCeil, VectorFloor, VectorTruncate, VectorFractional when using UnrealMathDirectX.
+
+* New: Added more stats to Garbage Collector.
+
+* New: Added new stats command which adds a custom marker to the stats stream.
+
+    * Can be used by typing stat namedmarker #markername# in console.
+
+* New: Added static_asserts to TSharedPtr<T> to ensure that it is not used in the opposite thread-safety mode of T's TSharedFromThis.
+
+* New: Added stats for wait/trigger in events.
+
+* New: Added stats to detect BulkData async buffer leaks.
+
+* New: Added stats to triggering async event.
+
+* New: Added support for binding Lambdas to multicast delegates.
+
+    * Note: lambdas bound to multicast delegates can *only* be removed via the handle returned by Add(), not the deprecated Remove(Delegate) API.
+
+* New: Added support for UObjectProperty, ULazyObjectProperty and UAssetObjectProperty as UMapProperty keys.
+
+* New: Added TMap property support in the reference chain.
+
+* New: Added TSet::Reset.
+
+* New: Added utility functions to TAssetPtr and FStringAssetReference to get package and asset name.
+
+* New: AVX support can now be enabled on a per-module basis by settings bUseAVX = true in a module's .Build.cs rules file.
+
+* New: Changed the default return value for GetResourceSize() to 0.
+
+* New: Declaring a BlueprintImplementableEvent function as virtual will now output a warning.
+
+* New: Deprecated FRotator::operator-() and created a new method FRotator::GetInverse(). Componentwise negation is a meaningless operation on a Rotator.
+
+* New: Enabled editability of subobjects in UnrealHeaderTool.
+
+* New: ensure() and ensureMsgf() macros will now only log/break the first time they fail.
+
+    * Added ensureAlways() and ensureAlwaysMsgf() macros which act like the old ensure() macros, and added a 'debug ensureAlways' console command.
+
+    * Removed the ensureMsg(), ensureOnce() and ensureOnceMsgf() macros.
+
+* New: Expanded TDoubleLinkedList API.
+
+    * Nodes can now be removed and inserted without deleting and reallocating the node memory.
+
+    * Added a Contains helper method.
+
+* New: FAssetPtr and TAssetPtr now provide a LoadSynchronous utility function.
+
+* New: FTickFunction::EnableParent has been removed.
+
+* New: Inlined fast UObject type checks and enabled them by default.
+
+* New: Made object delete listener more thread safe.
+
+* New: Made UnrealHeaderTool's error messages more helpful when it cannot find a type reference.
+
+* New: Non-monolithic targets now only build the plugins which they have enabled, unless bBuildAllPlugins is set to true in their target.cs file.
+
+* New: OnControllerAnalog, OnControllerButtonPressed, and OnControllerButtonReleased can now be used from plugins defining new input bindings.
+
+    * Functions now take an FName (aliased as FGamepadKeyNames::Type).
+
+    * Removed EControllerButtons and replaced with FGamepadKeyNames which provides FName mappings for all non-platform specific gamepad keys.
+
+* New: Optimized UnrealHeaderTool parsing.
+
+    * Added new IsSorted algorithm.
+
+    * Added new FindSortedStringCaseInsensitive algorithm for binary searching across multiple known strings.
+
+* New: PanicFlushThreadedLogs now also flushes the related output devices.
+
+* New: Performance: Added vectorized VectorSinCos for SSE platforms.
+
+    * Optimized FRotator->FQuat conversion to use the vectorized version, and also vectorize the quaternion construction operations.
+
+* New: Added Polar/Cartesian conversion functions with FVector2D parameters. (Contributed by [REDACTED_PERSON], pull request 1108.)
+
+* New: PreUnloadCallback will be called on a module interface before shutdown. (Contributed by [REDACTED_PERSON], pull request 1243.)
+
+* New: Only dump the script stack on assertion failure in the game thread. (Contributed by [REDACTED_PERSON], pull request 1281.)
+
+* New: Removed package type parameter from LoadPackageAsync as it was not used for anything.
+
+* New: Serialized FStringAssetReference properties can now be re-typed as TAssetPtrs and handled automatically during deserialization.
+
+* New: Stats can now be used in commandlets.
+
+* New: UEnums will now be tightly packed.
+
+    * UnrealHeaderTool will no longer attempt to generate fake entries for missing enum values.
+
+* New: UEnums will now store values as uint8 instead of int8 (many enums defined values greater than 128).
+
+* New: UnrealBuildTool Makefiles are now also enabled on the Mac to speed up UnrealBuildTool startup times on this platform.
+
+    * Fixed Makefile serialization when using Mono.
+
+* Fixed inaccurate calculation of non-frame stats on pooled threads.
+
+* Added UEngine::bForceDisableFrameRateSmoothing to enable external systems to force smoothing off without overwriting the desired base engine setting. Various HMD modules now use this to disable smoothing during stereo rendering. Fixes issue where framerate smoothing was disabled whenever a HMD was detected.
+
+* Added code to fixup command line if project path is incorrect.
+
+* Added code to make sure all objects that are being streamed in are referenced for Garbage Collector until loading has finished.
+
+* Added IsValidIndex to FString and had the [] operators use this to verify the index they were given was valid.
+
+* Added missing implementation for FLinkerLoad::FlushCache.
+
+* Changed UnrealBuildTool to generate UnrealHeaderTool's manifest using the same modules' order during hot-reload as during normal build.
+
+* Corrected config section used to save config values for PerObjectConfig objects that are contained in the transient package.
+
+    * This brings the code in line with what is used in UObject::LoadConfig.
+
+* Directory watcher no longer reports reads when the watch is being shutdown
+
+* Disabled compiler warning C4351 (new behavior: elements of array 'array' will be default initialized)
+
+    * The "new behavior" was in Visual Studio 2005, and is no longer considered unexpected.
+
+* Disallowed UClass* properties from having the "Instanced" keyword
+
+    * This never worked, but now it is rejected. If you get this error when upgrading, remove the "instanced" keyword from your UPROPERTY.
+
+* Fixed a crash caused by token stream not properly handling AddReferencedObjects pointers when added just before the end of array tokens.
+
+    * Introduced new token type that marks the end of a pointer data inside of a token stream to correctly handle return count for the array.
+
+* Fixed a crash related to RecastNavMesh when performing hot-reload.
+
+* Fixed a problem with UnrealHeaderTool not being able to parse some code comments.
+
+* Fixed a UnrealHeaderTool crash when generating code using a UObject type which was not parsed.
+
+* Fixed an error message in UWorld::ServerTravel.
+
+* Fixed an object name in a UObject::BeginDestroy assert.
+
+* Fixed crash in FGCReferenceTokenStream::ReplaceOrAddAddReferencedObjectsCall when using TMap properties.
+
+* Fixed crash in UMapProperty serialization.
+
+* Fixed FindObject function returning objects that have not yet been fully streamed in.
+
+* Fixed FString operator commutativity.
+
+    * Fixed some FString comments.
+
+* Fixed FString::FindLastCharByPredicate when the string is empty.
+
+* Fixed hot reload crashes that were occuring when user added C++ enums to the code.
+
+* Fixed hot-reload crash that was occuring during obtaining class default object's duplicate.
+
+* Fixed navigation-related crash after performing hot reload.
+
+* Fixed operator<< return values in FObjectAndNameAsStringProxyArchive and FNameAsStringProxyArchive.
+
+* Fixed pointer alignment issues in binned allocator.
+
+* Fixed possible UELinkerFixup linker errors after a module gets updated.
+
+* Fixed rare case where UnrealHeaderTool assumed existance of a parent class in the wrong module during compilation.
+
+* Fixed some logic which was causing static code analysis warnings.
+
+* Fixed some path manipulation in UnrealBuildTool which could prevent some C++ projects being built in some locales.
+
+* Fixed some TSet compile errors when the ElementType differs from the KeyType.
+
+* Fixed stale linker pointers after streaming a package in.
+
+* Fixed stats not collecting all data on engine startup.
+
+* Fixed TBitArray ranged-for iteration after a call to Reset.
+
+* Fixed TBitArray::GetAllocatedSize and TBitArray::CountBytes.
+
+* Fixed the 1000 log filename limit in UnrealAutomationTool.
+
+* Fixed the FSphere::TransformBy(FTransform) function which wasn't multiplying W by the scale.
+
+* Fixed the functionality to track allocations that didn't happen within the engine's wrappers.
+
+* Fixed UEnum vtable linkage.
+
+* Fixed UnrealBuildTool failing to harvest environment variables when running from some paths containing non-ASCII characters.
+
+* Fixed UnrealBuildTool not rebuilding its makefile when hot-reloading modules.
+
+* Fixed UnrealBuildTool not updating the makefile when a new UHT type is added to a header which didn't have any previously or when the last one is removed.
+
+* Fixed UnrealHeaderTool crash when using a multicast dynamic delegate type in a Blueprint function declaration.
+
+* Fixed UnrealHeaderTool failing to compile delegates which use an enum from earlier in the same header.
+
+* Fixed UnrealHeaderTool failing to skip DECLARE_MULTICAST_DELEGATE declarations at the top of headers.
+
+* Fixed UnrealHeaderTool not parsing meta data information correctly in some situations.
+
+* Fixed VectorMod() working incorrectly for negative values (now matches std fmodf). Fix VectorSinCos() doing incorrect range reduction. Improved math tests at startup.
+
+    * Math tests: added VectorMod, VectorSinCos.
+
+    * Math tests: added more FQuat <-> FRotator conversions and cleaned them up (make more important tests earlier, because those could cause later ones to fail).
+
+* Fixed FInterpCurve serialization and fixed edge case results from EvalDerivative and EvalSecondDerivative.
+
+* Fixed 'You are attempting to compile on a machine that does not have a supported compiler!' and other build exceptions in UnrealAutomationTool.
+
+* Inlined calls to ConditionalSleep to avoid TFunction + lambda overhead in FMultiReaderSingleWriterGT.
+
+* Made sure UAT builds content-only projects as well to avoid issues with missing/outdated UE4Game binaries.
+
+* Made the message about failure to route PostInitProperties more clear.
+
+* Made UnrealHeaderTool throw an error when it encounters bad UENUM() specifiers.
+
+* More accurate InvSqrt() and InvSqrtEst() on some platforms.
+
+    * Matches results for SSE platforms with other platforms that perform 2 Newton-Raphson iterations on the hardware estimate (iOS), or non-SSE that uses the fully accurate sqrt.
+
+    * Scalar functions now match the vectorized versions (VectorReciprocalSqrtAccurate) so that comparisons within 1e-8f are possible for normalizing normalized values (important for FQuat which uses both scalar and vector).
+
+    * Allows normalize vector (1,0,0) to return the same vector rather than (0.9999998xx, 0, 0).
+
+    * Fixed InvSqrtEst() on some platforms actually using higher precision (slower) FGenericPlatformMath InvSqrt().
+
+* More accurate VectorReciprocalSqrtAccurate on SSE platforms.
+
+    * Matches results for SSE platforms with other platforms that perform 2 Newton-Raphson iterations on the hardware estimate (Xbox, iOS), or non-SIMD that uses the fully accurate sqrt.
+
+    * Allows NormalizeRotation (0,0,0,1) to return the same vector rather than (0,0,0,0.99998...).
+
+* Optimized FPaths::GetCleanFilename, GetPath and GetBaseFilename.
+
+* Prevented UnrealHeaderTool from parsing array parameters in UFUNCTIONs.
+
+* Remove unused UObject::CreateArchetype and unimplemented UObject::UpdateArchetype
+
+* Removed "VectorBitwiseAND" defined only for UnrealMathSSE. There is already "VectorBitwiseAnd" defined for all platforms.
+
+* Removed deprecated delegate and timer functions.
+
+* Removed USE_MATRIX_ROTATOR from FQuat. It's not maintained and makes the code less readable.
+
+* Stats from other threads like task graph or pool thread should now be synced with the game thread.
+
+* UnrealHeaderTool will now be able to parse TSubclassOf<USomeClass>& parameters without reporting errors.
+
+## Editor and Tools
+
+* New: You can now help the blueprint utility functions create the correct blueprint for you. Previously, all classes created a standard blueprint. Now, the compiler module is asked Engine/Content/Types/Textures/SupportAndSettings/index.html if there is any registered IBlueprintCompiler with a special blueprint class for a particular UObject class. This fixes the problem of the new asset buttons creating UBlueprints for UUserWidget subclasses, instead of UWidgetBlueprints.
+
+* AutomationTool now correctly handles command line arguments that end with slashes.
+
+* IntroTutorials and Blutilities modules are now loaded when running commandlets, enabling asset inspection, resaving, etc.
+
+* SPropertyEditorEditInline now respects the "noclear" UProperty flag.
+
+* UnrealVS improvements:
+
+    * v1.36 Fixed Set Startup Project bug in which the startup project could be forgotten when set via the UnrealVS toolbar.
+
+    * v1.38 Improved handling of command line arguments stops bad uproject paths in user created projects.
+
+    * v1.39 Moved UnrealVS development to VS2015. This and subsequent versions will be available for VS2012/13/15.
+
+* Moved KismetNodeInfoContext.h to be publically accessible outside of the GraphEditor module.
+
+## Content Browser
+
+* Changing the Content Browser path via the breadcrumb now notifies you of the change.
+
+* Fixed SPathView::OnPathSelected being called with an empty string after the editor loads.
+
+## Landscape
+
+* Moved implementation details of the Landscape Editor into the Private folder.
+
+## Gameplay Framework
+
+* New: Added GameplayStatics::DeprojectScreenToWorld and ProjectWorldToScreen functions.
+
+    * Added simpler version of FSceneView::DeprojectScreenToWorld that takes a single view projection matrix.
+
+    * Added FSceneView::ProjectWorldToScreen.
+
+    * Added PlayerController::GetLocalPlayer().
+
+* New: AActor::bPendingKillPending renamed to bActorIsBeingDestroyed and made private.
+
+* New: Actors spawned from a non-CDO template will now have be named similarly to the template's name.
+
+* New: Add FQuat::GetNormalized() that returns a normalized copy.
+
+* New: Add FTransform methods CopyTranslation(), CopyRotation(), CopyScale3D().
+
+    * Avoids intermediate copies and unaligned loads/stores between VectorRegister and FVector.
+
+    * SceneComponent now uses these in CalcNewComponentToWorld().
+
+* New: Added FBodyInstance::OverlapTestForBody().
+
+    * Avoids heap allocation when only method was one that took a TArray.
+
+    * Templatized OverlapTestForBodies() to allow TArray with different allocator types (namely TInlineAllocator for OverlapTestForBody()).
+
+* New: Added Draw Shadow option to HUD DebugText allowing setting shadow per debug line.
+
+* New: Added ensures that transforms are valid in various places in SceneComponent operations.
+
+* New: Added ETeleportType flag to MovementComponent::MoveUpdatedComponent, SafeMoveUpdatedComponent to match SceneComponent::MoveComponent.
+
+    * Note: currently always treated as not teleporting when used within FScopedMovementUpdate.
+
+* New: Added FMath::GetMapRangeValueUnclamped and renamed GetMappedRangeValue to GetMappedRangeValueClamped.
+
+* New: Added new FActorRange, TActorRange and FSelectedActorRange types to support ranged-for iteration of those concepts.
+
+* New: Added parallel queueing of ticks.
+
+* New: Added stat counter for CharacterMovementComponent::FindFloor().
+
+* New: Added stats for PlayerController in "stat Game" and "stat PlayerController".
+
+* New: Allow TriggerSphere/Capsule/Box classes to be subclassed outside Engine module.
+
+* New: Clean up Character movement stats. Add new stat for AdjustFloorHeight.
+
+* New: Default implementation of USoundNode::CreateStartingConnectors will now make at least GetMinChildNodes() connections.
+
+* New: Exported FOverlapInfo with ENGINE_API.
+
+* New: Exposed the ability to set the tick group for actors and components in Blueprints (using either SetTickGroup or by editing the value for the default object).
+
+* New: Performance: 30% Faster FQuat::Normalize() and FQuat::GetNormalized() when vector intrinsics are enabled.
+
+* New: Performance: Added FTransform methods which on vectorized platforms can avoid the cost of converting to FVector first before the operations.
+
+    * Added static methods AreRotationsEqual, AreTranslationsEqual, AreScale3DsEqual, AddTranslations, SubtractTranslations.
+
+    * Added RotationEquals(FTransform), TranslationEquals(FTransform), Scale3DEquals(FTransform).
+
+* New: Performance: Added vectorized FRotator::Equals() and IsNearlyZero() implementations.
+
+* New: Performance: CharacterMovement nav mesh projection uses cheaper line-plane-intersection, and avoids it altogether when at the destination.
+
+* New: Performance: Implemented faster Quaternion Vector rotation.
+
+    * Scalar FQuat::RotateVector() takes ~18% less time.
+
+    * VectorQuaternionRotateVector() takes ~25% less time.
+
+    * Vectorized FTransform * FTransform takes ~33% less time.
+
+* New: Performance: Optimized FQuat comparison when vector intrinsics are available.
+
+* New: Performance: Remove unnecessary encroach check every time a simulated proxy position or rotation replication is received.
+
+    * Network character proxies only need to check for encroachment when they teleport with zero velocity. Normal movement handles encroachment already.
+
+* New: Performance: Rotation conversions in SceneComponent are cached to avoid Quat<->Rotator conversion when possible.
+
+    * Added FRotationConversionCache which caches Quat<->Rotator conversions and returns the cached conversion if the source rotation is the same.
+
+    * SceneComponent uses this cache for internal uses of RelativeRotation (FRotator) and calls to GetComponentRotation() (returning a FRotator from ComponentToWorld's quaternion).
+
+        * As a result, GetActorRotation() and GetComponentLocation() only perform a conversion if the rotation has changed since previous calls.
+
+* New: Performance: Rotational movement can now avoid UpdateOverlaps when the rotation is symmetric even when currently overlapping components in most cases.
+
+* New: Refactored vectorized Quaternion Vector rotation to a single method (VectorQuaternionRotateVector), rather than two calls to VectorQuaternionMultiply2().
+
+    * This allows other rotation methods to be implemented more easily.
+
+    * Also added VectorQuaternionInverseRotateVector().
+
+* New: Simplify FindGeomOpposingNormal() return value (and remove out param).
+
+* New: Spawning an actor will now respect any default transform on a native root component.
+
+* New: TInlineComponentArray can now also be initialized with a const AActor*.
+
+* New: UGameInstance::InitializeStandalone can now be called when running in editor code.
+
+* New: USoundWaveStreaming has been renamed USoundWaveProcedural to better describe its purpose.
+
+* CameraComponent::SetCameraMesh will now update the proxy mesh if it exists.
+
+* Actor iterators will no longer ever return Actor's that have become marked pending kill while the iteration has been underway.
+
+* Actor::GetActorEnableCollision() is now const.
+
+* Add separate stat for SceneComponent::MoveComponent. It was not tracked as part of "MoveComponent Time" previously.
+
+* Added virtual destructor for FDamageEvent (PR#1141, thanks bozaro!)
+
+* Check that Actor or Component collision is enabled when adding to ignored components list in FCollisionQueryParams.
+
+* Create UWorld::DefaultPhysicsVolume on demand.
+
+    * Fixes crash when dealing with physics volumes when world is not initialized (commandlets etc).
+
+* Fix blocking hit notifications being triggered in invalid cases (pending kill or null actors/components). Fix SceneComponent trying to move when pending kill (PrimitiveComponent already aborts in this case).
+
+* Fix blocking hits dispatched from FScopedMovementUpdate() being dispatched when collision response has changed (from overlaps or other blocking hits).
+
+* Fix character mesh relative translation being applied incorrectly in UCharacterMovementComponent::SmoothClientPosition().
+
+* Fixed APlayerController::ReceivedSpectatorClass() to take a class of the correct type.
+
+* Fixed Character::bSimGravityDisabled not being restored to false for simulated pawns under some conditions.
+
+* Fixed crash trying to use SSSwapControllers debug command with 0 players.
+
+    * Fixed GameInstance execs not firing with 0 players, which prevented re-adding a player.
+
+* Fixed crash when Character capsule is marked pending kill and set null by MovementComponent::TickComponent.
+
+* Fixed crash when UpdatedComponent for Projectile or Rotating movement are marked pending kill and set to null in the next tick.
+
+* Fixed scalar versions of FTransform::TransformFVector4() and TransformFVector4NoScale() setting W to 1 or 2 instead of 0 or 1.
+
+* Fixed SlideAlongSurface pushing out away from the original direction of movement when hitting multiple surfaces.
+
+    * This fixes jitter when characters push against corners of multiple surfaces at once.
+
+* Fixed stat tracker in "stat game" for CharacterMovement not tracking the total cost. Add stats for Character network movement and smoothing.
+
+* FNotifyBufferFinishedHooks can now be used by modules other than the Engine module.
+
+* Fully exposed ADecalActor so that you can derive from it.
+
+* Game/URL Options parsing nodes are now exposed to blueprints.
+
+* Made UCharacterMovementComponent::GroundMovementMode private, and provide Get/Set methods.
+
+    * Enforces that GroundMovementMode is only ever Walking or NavWalking.
+
+    * Set method will switch movement modes if currently a different ground walking mode.
+
+    * Cleaned up constants in Pack/UnpackNetworkMovementMode() and made sure we only have a valid ground mode (Walking or NavWalking).
+
+* Optimized various uses of Pop to prevent incremental reallocation for short-lived or frequently used TArray instances.
+
+* Performance: Avoid recomputing transforms and Quat<->Rotator conversions in FAnimNode_RotateRootBone::Evaluate() when neither Pitch nor Yaw are set.
+
+* Performance: Faster FQuat::operator!=.
+
+* Performance: Faster Quat->Rotator conversion.
+
+    * We only need to normalize 1 axis and only rarely (when pitch is -90 or 90). ASin and ATan2 already produce normalized results (as in [-pi, pi]).
+
+* Performance: Inlined some frequently used Actor, FBodyInstance, and FCollisionResponse functions.
+
+* Performance: Moved static initialization of some text out of functions and into global scope. Avoids branches and extra instructions in these functions.
+
+* Performance: Optimized and simplified CharacterMovementComponent functions: IsFalling(), IsMovingOnGround(), IsSwimming(), IsFlying().
+
+* Performance: Optimized CharacterMovement::PhysicsRotation() to reduce FRotator comparisons.
+
+* Performance: Optimized usage of MasterPoseComponent weak pointer in SkinnedMeshComponent.
+
+* Performance: Optimized weak pointer usage in UPrimitiveComponent::DispatchBlockingHit().
+
+* Performance: Optimized weak pointer usage in USkeletalMeshComponent::CalcBounds().
+
+* Performance: Remove unnecessary FRotator normalization before comparisons in CharacterMovementComponent::PhysicsRotation().
+
+* Performance: Removed static initialization in VectorNormalizeRotator, and moved constants to GlobalVectorConstants.
+
+* Performance: UMovementComponent::SlideAlongSurface internally uses FQuat instead of FRotator.
+
+* Remove GetSafeNormalPrecise() and GetClampedToMaxSizePrecise() from CharacterMovementComponent source.
+
+    * The stock FVector methods are now sufficiently precise after the InvSqrt() changes.
+
+* Remove stale (null, pending kill) overlaps from OverlappingComponents array when an invalid overlap is found in UpdateOverlaps().
+
+* Removed unused InactivePlayerArray variable from GameState.
+
+* Simulated CharacterMovement now uses SetPostLandedPhysics() when landing, rather than hardcoded walking mode.
+
+    * Fixed SetPostLandedPhysics() not respecting default land movement mode that is not a walking mode.
+
+* Sound Node pin and title names are now properly localizable.
+
+* TActorIterator now properly filters out Actors in hidden levels the same as FActorIterator.
+
+* UCharacterMovement::StepUp() no longer uses SafeMoveUpdatedComponent for its movement.
+
+    * This could cause movements to reverse direction or violate step-up height rules.
+
+* Various minor fixes and improvements to CharacterMovementComponent.
+
+    * Added "virtual" modifier to some existing functions.
+
+    * Improved "const" correctness at several places.
+
+    * Improved standing downward force code and moved it to the new overridable function ApplyDownwardForce().
+
+    * Moved calculation of root motion velocity code to the new overridable function CalcRootMotionVelocity().
+
+    * Minor improvements in functions GetLedgeMove() / MoveAlongFloor() / PhysWalking() / StepUp() / ApplyRepulsionForce()
+
+* When clearing a PlayerController's spectator pawn, its regular pawn is used for attachment and auto-managed camera target, if regular pawn has already been set up.
+
+* Zeroed out FRepMovement::AngularVelocity when it's not used. More efficient for replication if it previously had non-zero values (switched from simulated to not simulated).
+
+## Localization
+
+* New: Created a new localization module for all the editor's general needs with regards to localization configuration scripts and localization pipeline configuration.
+
+* Corrected behavior of FText::FindText, which would not properly ensure that the text found actually has a display string that is associated with the source string (if provided).
+
+* Writing manifest to file now uses full/absolute/rooted path to avoid any ambiguity when debugging.
+
+## Networking
+
+* New: Added a for loop to call PreReplication on all replicated components when Actor::PreReplication is called
+
+* New: Added connection state tracking to online beacons so their state can be inspected during use. APartyBeacons also return more useful information about reservation removals.
+
+* New: Added GetMaxAvailableTeamSize function to beacon. It shows the max number of new players that could be accommodated on the existing teams.
+
+* New: Added UStruct to FVariantData converters - very basic, converts UPROPERTY to key-value pairs.
+
+* New: All beacon derived classes have access to new GetNetDriverName function rather than direct access to the name.
+
+* New: Moved majority of ConsoleCommand out of APlayerController and into UPlayer to allow for better sharing (beacons can now use ConsoleCommand).
+
+* New: Party beacons can now handle updates to their existing reservations. Clients can send reservation updates and get acknowledgement or rejection.
+
+* New: Party beacons have new rejection responses to help disambiguate failure cases (Bad Session, Invalid Reservation).
+
+* New: Party reservation beacons can now handle updates to the party leader when the current leader leaves the game. Clients will report the newly negotiated leader to the server.
+
+* New: Party reservation beacons now send updates to all connected clients when a new client joins the beacon. It will send the current remaining reservations at that time.
+
+* New: The code in UWorld::Listen will no longer be compiled into client builds so that it is impossible to execute.
+
+* New: Added PreReplication of components when PreReplication of actors is called.
+
+* Fixed an issue where actors spawned on beacon net drivers weren't properly destroyed because DestroyActor assumed they were on the game net driver.
+
+* Made APartyBeacon duplicate reservation checks first so that other size checks don't fail before it has realized the reservation was a duplicate.
+
+* OnlineBeacon derived objects will dispose of their net drivers correctly on EndPlay, fixing cases where a netdriver wasn't destroyed beforehand and causing garbage collection issues with the world.
+
+* Party beacons start in the NonePending request type state now and will return there after the beacon is finished communicating.
+
+* PreReplication is called for all replicated subcomponents on an Actor.
+
+* Prevent possible DestroyBeacon crashes when there is no world (possibly during GC in server travel for instance).
+
+* The ShowInviteUI API now takes a session name.
+
+* AOnlineBeaconClient::ClientOnConnected is no longer virtual, use OnConnected instead.
+
+* Removed NetConstants.h and added MTU changes to CoreNet.h.
+
+## Online
+
+* New: FOnlineExternalUI::ShowProfileUI has now been implemented for Steam. This displays a specific player's profile in the steam overlay.
+
+* New: FUniqueNetId now derives from TSharedFromThis so that AsShared can be called on const references that are passed around from the OSS. All references have been marked as const as well.
+
+* Added GetInstanceName to IOnlineSubsystem interface. It returns the name of the OSS interface. Games only use "default", but PIE has the name of the instance for referencing the world context.
+
+* Added GetTypedHash function to FUniqueNetIdRepl so that it may be used in TMaps.
+
+* Fixed a possible string buffer overflow issue in Steam matchmaking code when processing advertised keys.
+
+* Party beacons no longer attempt to send reservation updates when there are no clients connected.
+
+## Other
+
+* New: Users can now provide functionality to web browser widgets to handle creation of popup windows.
+
+* New: Users can now provide functionality to web browser widgets to handle the closing of popup windows created by the browser.
+
+* New: The multicast (broadcast) analytics module is now a plugin to make it easier to use. Anyone previously including that module as part of their build will now need to include it via the Plugins UI.
+
+* Made Base.ini no longer be an empty file to avoid some warnings.
+
+* Optimized away Slate vertex buffer stalls on platforms that support parallel rendering.
+
+* Minor improvements to the ParticleSystemAudit commandlet. You can specify an output folder via "-AuditOutputFolder=<SomeFolder>" and only more important data is now shown.
+
+* Cooked builds no longer load editor-only FindInBlueprint data, resulting in less memory allocated.
+
+* Clear colors for hardware clears are now bound to rendertargets at creation time.
+
+## Physics
+
+* New: Add Distance to FHitResult info returned from traces & sweeps. (Contributed by [REDACTED_PERSON], pull request 1054.)
+
+* New: Changed cloth-enable cvar to be checked on any thread instead of just the game thread.
+
+* New: Perf: Optimize usage of weak pointers in collision conversions. Remove copy of FHitResult when filling in results. Avoid expensive MTD code for overlaps because this is only for blocking hits.
+
+* Fixed wrong variable in ensure in ConvertOverlappedShapeToImpactHit().
+
+* Performance: Avoid Quat->Rotator->Quat conversion for every physics component that moves during simulation.
+
+* Performance: Inlined FBodyInstance::IsInstanceSimulatingPhysics() and the accessor for bSimulatePhysics.
+
+* Performance: inlined some common PhysX<->Unreal type conversions.
+
+* Removed deep copy of hit results in UWorld::ComponentSweepMulti().
+
+    * This also ensures that sweeps with multiple shapes actually return a correctly sorted list of hits, since all hits are now sorted together rather than appended in individually sorted chunks.
+
+    * Added note that GeomSweepMulti_PhysX does not clear the incoming OutHits array, only appends to it. Note that GeomSweepMulti_PhysX is now marked DEPRECATED_FORGAME.
+
+* Remove VERIFY_GEOMSWEEPSINGLE and VERIFY_GEOMSWEEPMULTI from PhysXCollision.
+
+* Reorder calls to SetHitResultFromShapeAndFaceIndex in CollisionConversions to better match memory layout of FHitResult.
+
+## Platforms
+
+* New: Added process enumeration abstraction FProcEnumerator in FPlatformProcess. Now it's much easier to obtain system-wide processes information when you need it!
+
+* New: Enabled syntax highlighting for MetalRHI module in Xcode project.
+
+* Fixed %*s formatting on platforms that don't use system vswprintf.
+
+* Fixed IPlatformFile::OpenWrite so that with bAppend set to true behavior is the same across platforms.
+
+    * Specifically, we make sure the file handle returned is set to the end of the file.
+
+* Fixed redundant initialization of target platforms.
+
+### Mac
+
+* Added few warning ignore pragmas on Mac to make the engine compile with new Clang in Xcode 7.
+
+* Fixed linking with CoreAudio frameworks when building with OS X 10.11 SDK.
+
+* Fixed paths to Mac targets in Xcode schemes.
+
+### Playstation 4
+
+* New: Minor fixes to PS4 parallel translate code including removal of "late resolve".
+
+* Add code so the bound shader state cache works with parallel translate on PS4.
+
+* Fixed webapi deprecation warnings.
+
+### Xbox One
+
+* Fixed warnings being displayed for missing ICU PDB files when building Xbox One.
+
+## Programming
+
+* New: We are now using the 64 bit Windows linker, even on 32 bit targets. This linker is faster and can handle larger link jobs without running out of memory.
+
+* New: Added GameplayTagContainer debug visualizer.
+
+* New: Added support for distributed compilation of Win64 builds using SN-DBS.
+
+    * SN-DBSv2.4.0.61 and SN-DBS-VC-Extension-2013v1.0.7.0 are required.
+
+* New: Compiling the engine using Clang on Windows (Beta)
+
+    * It's now possible to compile UE4's C++ code using the LLVM/Clang compiler on Windows
+
+    * This feature is still very experimental, and will need additional work to be generally useful
+
+    * If you're interested in trying it out or helping us to improve Clang support for UE4, please check out the forum thread below
+
+    * https://forums.unrealengine.com/showthread.php?54770-Using-Clang-to-compile-on-Windows
+
+* New: Fast XML parser for UE4
+
+    * This is a simple XML parser that can load XML files very quickly.
+
+        * The new FFastXml class has a single static method ParseXmlFile() for loading XML files. The parser is designed to load files efficiently, but may not support all XML features or be resilient to malformed data.
+
+        * Unlike the FXmlFile class, FFastXml does not generate an XML node tree. Instead, you must supply a callback class using IFastXmlCallback and respond to elements and attributes as they are parsed. Call the ParseXmlFile() function and pass in either the full path to the XML file to parse (XmlFilePath parameter), or load the file yourself and pass the contents using the XmlFileContents parameter.
+
+        * One of either XmlFilePath or XmlFileContents must be valid for ParseXmlFile() to work.
+
+    * To process the XML data:
+
+        * Create your own class that implements the IFastXmlCallback interface. This class will be responsible for processing the XML elements as they are loaded by FFastXml::ParseXmlFile().
+
+        * You'll receive a ProcessElement() call for every XML element that is encountered, along with a corresponding ProcessClose() when that element's scope has ended. ProcessAttribute() will be called for any attributes found within the scope of the current element.
+
+    * Remember to add a module dependency on "XmlParser" in order to be able to call FFastXml::ParserXmlFile().
+
+    * This XML parser is based on John W. Ratcliff's "FastXml".
+
+* New: It is now possible to use `<select>` menus in html content rendered by the Web Browser widget
+
+* New: The Web Browser widget now supports changing the mouse cursor based on css properties.
+
+* New: Tweaked the policy for choosing how many parallel compiles to do. If the number of physical cores is greater than 4 and not equal to the number of logical cores, then use the average of the two core counts for the number of parallel compiles.
+
+* New: Unreal Build Tool support for Visual Studio 2015 (Beta)
+
+    * VS 2015 support is disabled by default in this release (not all code and libraries have been ported.)
+
+    * Use the new "-2015" option when generating project files to enable VS 2015 support
+
+    * Windows SDK 8.1 is used by default. To use Windows SDK 10 instead, enable WindowsPlatform.bUseWindowsSDK10.
+
+    * Universal Windows Platform apps are not yet supported.
+
+    * Various loose ends still remain (search for "@todo UWP" in Unreal Build Tool code).
+
+* New: UObject instances can now be passed to SWebBrowser and accessed from browser javascript.
+
+* Fixed very slow compile times in engine C++ code with Slate widgets.
+
+* Fixed a few compile errors when compiling with Clang on Windows (r231657).
+
+* Fixed slow compile times of SpeedTree importer code (caused by optimizer stumbling on an inner loop).
+
+* Fixed various cases of C++ variable shadowing in Windows-specific code.
+
+* Improved Intellisense for Launcher code projects.
+
+* Previously when using the WebBrowser module on OS X, one had to create a custom subclass of NSApplication. This is no longer required.
+
+* Some unused variables cleaned up in the AutomationTool.
+
+* Unreal Build Tool: Fixed incorrect VS tools versions set in VC++ project files.
+
+## Rendering
+
+* New: Added miscellaneous parallel rendering tweaks and options. Balance deferred contexts after the fact, cvar can control occlusion latency, and a framework for Get Dynamic Mesh Elements tasks. These are all disabled by default.
+
+* New: Added Parallel Frustum Culling.
+
+* New: Added r.CompositionGraphOrder for post processing, changed default order to be node creation order to give more control, which will be useful for AsyncCompute.
+
+* New: Added several performance tweaks. Compact cloth buffers. Option to defer uniform buffer creation until it is visible. Option to defer skeletal components until visible and handle them on the RHI thread. Added the ability for tasks to be considered high priority. Added high priority ticks.
+
+* New: Clear colors for hardware clears are now bound to render targets at creation time. Optimization for upcoming graphics APIs. (DX12, Vulkan).
+
+* New: Extended Visualize Texture/Vis to also have Red, Green and Blue channel modes (by Lionhead Studios)
+
+* New: Optimized Init Views in Deferred Shading Scene Renderer by making front-to-back sorting of base pass draw lists asynchronous. (Contributed by [REDACTED_PERSON], pull request 820.)
+
+* New: Turned on parallel render thread updates to improve game thread performance. This is controlled by the r.AllowAsyncRenderThreadUpdates cvar.
+
+* New: Tweaked and added control for controlling the load balancing on parallel command lists. Relevant cvars: r.RHICmdMergeSmallDeferredContexts, r.RHICmdBalanceTranslatesAfterTasks, r.RHICmdMinCmdlistForParallelTranslate, r.RHICmdMinCmdlistSizeForParallelTranslate.
+
+* New: Users that build from source will now get warned if they did not build ShaderCompilerWorker.
+
+* Added code to flush rendering thread on movie finish wait to avoid several buffering problems.
+
+* Added code to flush the RHI thread and RHI resources in the movie loop.
+
+* Alternate fix for crash when generating thumbnails for added feature packs.
+
+* Fixed a crash related to diagnostic code for TScopedPointer<FStaticMeshRenderData>.
+
+* Fixed a rare bug with translucent lighting. FViewInfo snapshots no longer use the scene allocator.
+
+* Fixed a shader compiler bug when using a local variable in the shaders ending in a word like a uniform buffer (View, Primitive, etc).
+
+* Fixed assert on motion blur when paused.
+
+* Fixed async render thread updates and particle threading.
+
+* Fixed bad recycling on RHICommandLists, was an exit crash and wasn't recycling.
+
+* Fixed crash on zero scaling for HISMC.
+
+* Fixed old bug where bloom and Gaussian DOF size doesn't scale linearly, but costing the linear performance increase.
+
+* Fixed parallel sorting of drawlists.
+
+* Fixed race condition on feature packs and global shaders.
+
+* Fixed shadow culling on foliage.
+
+* fixed threading issue with movie playback and no rhi thread.
+
+* Fixed threading with the atmospheric fog component.
+
+* fixed UE-16338 Crash when setting console command for shadow max resolution to 1
+
+* Merged miscellaneous parallel particle simulation fixes.
+
+* Only emit glsl precision on platforms that need it.
+
+* Pull Request 794: Fixed issue with depth-stencil mips clearing incorrectly.
+
+* Updated information about the expected buffer lifetime of BUF_Dynamic and BUF_Volatile.
+
+### Lighting
+
+* Fixed sRGB issue in Set Light Color for light components. (Contributed by [REDACTED_PERSON], pull request 1161.)
+
+* Improvement to Has Non Zero Data in FQuantizedLightmapData. (Contributed by [REDACTED_PERSON], pull request 790.)
+
+### Postprocessing
+
+* New: Changed all direct access to PassOutputs to GetInput() to allow internal restructuring.
+
+## UI
+
+### Slate
+
+* New: Added a way for UI actions to be repeated if the chord used to perform them is being held down.
+
+* New: Added a way to query Slate whether it is currently sleeping by calling FSlateApplication::Get().IsSlateAsleep().
+
+* New: Added a way to specify the header padding for an SExpandableArea.
+
+* New: Converted FUIAction to use ECheckBoxState when reporting its check state.
+
+    * There are now overloads to FUIAction that take an FGetActionCheckState delegate which is able to return any of the entries from ECheckBoxState.
+
+    * FIsActionChecked still exists for backwards compatibility, and converts the returned bool into a ECheckBoxState via a passthrough function.
+
+* New: Slate menu refactor means all menus can be used in windowed and full screen games.
+
+    * Previous menu API returned menus as windows, but menus in full screen applications must not be windows.
+
+    * Menus now all respect the QueryPopupMethod of the widget that summons them.
+
+    * New API creates menus and returns them as IMenus.
+
+* New: UpdatableTexture now has a method for updating the texture based on raw texture data from 3rd party APIs instead of having to copy it to a TArray first.
+
+* New: Viewport implementations now have a way to specify that the backing texture should not be scaled if it doesn't match the viewport size by overriding ISlateViewport::AllowScaling().
+
+* Fixed a bug in the standalone renderer on OS X that could cause crashes by accessing the rendering context from multiple threads.
+
+* Fixed STableRow::SetContent() from removing the extra UI added by STableRow::ConstructChildren().
+
+* Fixed path comparisons when checking a brush against a file.
+
+### UMG
+
+* New: Adding the concept of shared layers to the Game Layer Manager. Rather than a widget added directly for a player, the layer manager allows you to add a layer by name at a specific ZOrder.
+
+    * The new IGameLayer interface is what you implement and allows itself to be turned into a Widget. You can later request a layer by name, cast it to your needed type and then interact with whatever functions you exposed on your layer class.
+
+* Moving the SConstraintCanvas out of the UMG Module and into Slate proper. This will allow usage in more core engine modules and not just in UMG.
+
+* UE-16247 - We now correctly call SynchronizeProperties on all core widgets. Additionally there's now an ensure to prevent future mistakes due to forgetting to call Super::SynchronizeProperties().
+
+# Upgrade Notes
+
+## C++ API Changes
+
+* SetRenderTargetsAndClear no longer takes clear values per target. The value bound at creation will be used.
+
+* FUniqueNetId now derives from TSharedFromThis so that AsShared can be called on const references that are passed around from the OSS. All references have been marked as const as well.
+
+    * For TSharedFromThis to be valid, the memory must come from the heap rather than the stack. It is possible that there are instances of FUniqueNetId that have been missed. An assert will trigger immediately on first use of AsShared in these cases. Please report any issues you find. No existing code will break, only new uses of AsShared.
+
+* Removed package type parameter from LoadPackageAsync as it was not used for anything. Deprecated the overload with that parameter.
+
+* UEnums will now store values as uint8 instead of int8 (many enums defined values greater than 128).
+
+    * Some of the API functions on UEnum class have changed to accept / return uint8 instead of int8.
+
+* FCompactPose improvements have introduced API changes to core animation code
+
+    * Removed Functions:
+
+        * FA2CSPose::SafeSetCSBoneTransforms()
+
+        * FA2CSPose::LocalBlendCSBoneTransforms()
+
+    * Member Type Changes:
+
+        * FBoneTransform::BoneIndex
+
+* Removed deep copy of hit results in UWorld::ComponentSweepMulti().
+
+    * This also ensures that sweeps with multiple shapes actually return a correctly sorted list of hits, since all hits are now sorted together rather than appended in individually sorted chunks.
+
+    * Added note that GeomSweepMulti_PhysX does not clear the incoming OutHits array, only appends to it.
+
+* GeomSweepMulti_PhysX is now marked DEPRECATED_FORGAME. It was never intended to be a public function that games should use.
+
+* Performance: Converted separate VectorMultiply() and VectorAdd() to VectorMultiplyAdd() in TransformFVector and variants.
+
+* Deprecated UProjectileMovementComponent::GetEffectiveGravityZ() in favor of just overriding GetGravityZ().
+
+* Made UCharacterMovementComponent::GroundMovementMode private, and provided Get/Set methods.
+
+    * Enforces that GroundMovementMode is only ever Walking or NavWalking.
+
+    * Set method will switch movement modes if currently a different ground walking mode.
+
+    * Cleaned up constants in Pack/UnpackNetworkMovementMode() and made sure we only have a valid ground mode (Walking or NavWalking).
+
+## Other Upgrade Notes
+
+## AI
+
+* If your game-specific AI controller class already derives from IGenericTeamAgentInterface, you probably should remove it from your controller's parent classes
+
+## Animation
+
+* Games with duplicate, stale ActiveClassRedirect entries in their DefaultEngine.ini may have broken animation graphs until the duplicate stale redirectors are removed.
+
+* Started moving animation nodes out of Engine into a new AnimGraphRuntime module. Modules containing custom animation nodes will need to have a dependency on the new "AnimGraphRuntime" module added to their Build.cs file.
+
+## Core
+
+* Added message log error for source name collisions in redirector ini sections.
+
+    * Note: If you get this error on start up, please check your [ProjectDir]\Config\DefaultEngine.ini file for a large block of ActiveClassRedirects, K2FieldRedirects, etc... and remove any that are not game specific (for most games this is all of them).
+
+* Any code that uses UEnum::GetNameByIndex, GetValueByIndex, GetNameByValue, GetValueByName and GetMaxEnumValue needs to be modified to handle the new param/return types.
+
+* The "ensure" family has been changed so that the default behavior is to log/break only on the first failure, replacing the “ensureOnce” family. The old functionality is available under the “ensureAlways” family.
+
+    * ensure() should be replaced with ensureAlways().
+
+    * ensureMsgf() should be replaced with ensureAlwaysMsgf().
+
+    * ensureOnce() should be replaced with ensure().
+
+    * ensureOnceMsgf should be replaced with ensureMsgf().
+
+* Modules must be listed in either DynamicallyLoadedModuleNames or Public/PrivateDependencyModuleNames, not both. If in doubt, remove the duplicate from DynamicallyLoadedModuleNames.
+
+## Editor and Tools
+
+* If the ensure in FMainFrameModule::StartupModule() is hit in your game, you need to guard code that is editor-only via either the Editor module type, or checking IsRunningGame().
+
+## Localization
+
+* Make sure to set your game's native culture setting in the Region & Language page of the Editor Preferences, so that assets load up that culture's localization for display.
+
+## Paper2D
+
+* Fixed inconsistency between vertex colors on Paper2D assets and static meshes.
+
+    * Blueprints or C++ code that programmatically modifies the color of a Paper2D component will need to be updated, otherwise the color being applied may be darker than expected.
+
+## Physics
+
+* Note that constraint projection being turned on by default may slightly change the behavior of existing ragdolls. In general we are seeing much better stability as this is coupled with projection fixes inside PhysX.
+
+* The increase in default max angular velocity from 400 degrees per second to 3600 degrees per second can have a slight impact on existing simulated objects that were rotating at the previous max speed. You can use the physics project settings to change it back to the old behavior.
+
+## Platforms
+
+## Playstation 4
+
+* Your project must upgrade to SDK 2.5 to compile successfully
+
+* Project should see a small GPU performance benefit. Up to ~5% over the old compiler.
+
+* Projects must upgrade to SDK 2.5 and obtain 120hz Morpheus kits.
+
+**Mac**
+
+* Using the WebBrowser module no longer requires creating a subclass of NSApplication.
+
+## All Mobile
+
+* If you use either the iOS or Google Play online subsystems, the system login UI will no longer appear automatically during engine startup.
+
+    * You must now explicitly invoke the system login UI by calling either "IOnlineExternalUI::ShowLoginUI" from C++ or "Show External Login UI" from Blueprints.
+
+    * The "Connect to Service" Blueprint node has been deprecated. The "Show External Login UI" node should be used instead.
+
+* The "Connect to Service" Blueprint node has been deprecated. The "Show External Login UI" node should be used instead.
+
+* You must now explicitly invoke the system login UI by calling either "IOnlineExternalUI::ShowLoginUI" from C++ or "Show External Login UI" from Blueprints.
+
+## Programming
+
+* The previously required boilerplate code can now be removed, but should still work.
+
+## Rendering
+
+* Added new cvar r.Filter.SizeScale allowing to scale sample count to tweak performance/quality
+
+* Custom calls to SetRenderTargetsAndClear will need to be changed and custom rendertargets modified to bind the desired clear value at creation time.
+
+* Removed RHIMethods.h and converted the RHI interface to standard C++.
+
+## Postprocessing
+
+* You might need to change some code to use GetInput() instead of accessing the member directly.
+
+## UMG
+
+* Assets saved by editor versions prior to this change will copy the values of the deprecated properties to the style properties that replace them when they are loaded.
+
