4.9 release notes fixes

diff --git a/Engine/Documentation/Source/Support/Builds/ReleaseNotes/2015/4_9/ReleaseNotes_4_9.INT.udn b/Engine/Documentation/Source/Support/Builds/ReleaseNotes/2015/4_9/ReleaseNotes_4_9.INT.udn
--- a/Engine/Documentation/Source/Support/Builds/ReleaseNotes/2015/4_9/ReleaseNotes_4_9.INT.udn
+++ b/Engine/Documentation/Source/Support/Builds/ReleaseNotes/2015/4_9/ReleaseNotes_4_9.INT.udn
@@ -7,15 +7,15 @@ This release brings hundreds of updates for Unreal Engine 4, including 129 impro
 
 Andrew Zhilin (zoon), Artem V. Navrotskiy (bozaro), Artyom Sovetnikov, Ben Rog-Wilhelm(zorbathut), Ben Wiklund (bwiklund), Black Phoenix (PheonixBlack), Błażej Szczygieł (zaps166), Brad Grantham (bradgrantham), Brent Scriver (FineRedMist), Cengiz Terzibas (yaakuro), chipgw, Christian Radich (yoyohobo665), Christopher P. Yarger (cpyarger), Clay Chai (chaiyuntian), Cliff Jolly (ExpiredPopsicle), Dave Newson (dave-newson), Derek van Vliet (derekvanvliet), Dorgon Chang (dorgonman), ewirch, Felix Laurie von Massenbach (erbridge), Gabriel Hare (GabrielHare), gatools, Hakki Ozturk (ozturkhakki), HueyPark, JaredTherriault, Jason Spangler (Stormwind99), Javier Osset (Xaklse), Jeff Rous (JeffRous), JohnAlcatraz, Kitatus Studios (KitatusStudios), Konstantin Nosov (gildor2), korypostma, Lee Berger (MrCrowbar), Maarten Scholl (maartenscholl), Marat Radchenko (slonopotamus), marynate, Matthias Huerbe (MatzeOGH), Maxim (maxpestun), Michael Allar (Allar), Michael3DX, MiniTurtle, Moritz Wundke (moritz-wundke), Nako Sung (nakosung), Nari Demura (demuyan), Nastenko Michael (deM-on), Nathan Stocks (CleanCut), Niels Huylebroeck (red15), Pablo Zurita (pzurita), Patrick Flanagan (valtrain), Pedja Ljubomirovic (3dluvr), Philipp Smorygo (fsmorygo), Pierdek, Piotr Bąk (Pierdek), Quadtree, Rajko Stojadinovik (rajkosto), Rama (EverNewJoy), Rem (rveilleux), rlefebvre, Robert Khalikov (nbjk667), Russ Treadwell (trdwll), Salamanderrake, Sebastian Witkowski (witkowski-seb), [REDACTED_PERSON] (Srombauts), Simon Taylor (simontaylor81), Stephen Whittle (stephenwhittle), szyszq, TK-Master, Tobias Mollstam (mollstam), Tomasz Sterna (smokku), user37337, Victor Xie (FTPiano), Vladimir Ivanov (ArCorvus), Will Stahl (merlin91), yamashi, Zhi Kang Shao (zkshao)
 
-# Major Features 
+## Major Features 
 
-## Enhanced Support for Mobile Devices
+### Enhanced Support for Mobile Devices
 
 We've been working to greatly improve support for mobile platforms in UE4, and this release contains many exciting rendering features to help you create leading-edge mobile titles.  These include efficient dynamic shadows, movable lights and decals.  Expect to see continued improvements to mobile and HTML5 rendering quality in the future.  In this release, IOS In-App Purchases has been improved and we've added Remote Push Notifications and support for CloudKit.  For Android, we've added OpenGL ES 3.1, Immersive Mode support, asynchronous audio decompression, and better orientation features.  And we now include NVIDIA AndroidWorks to make it easy to get up and running on your device.  UE4 now includes a Flurry analytics plugin for iOS out of the box, also!  Many of the new features for mobile devices are also available in your HTML5 games, too.  
 
 ![image alt text](image_0.png)
 
-## New: Dynamic Character Shadows for Mobile
+### New: Dynamic Character Shadows for Mobile
 
 Dynamic modulated character shadows from directional lights are now supported on mobile devices!
 
@@ -25,7 +25,7 @@ Dynamic modulated character shadows from directional lights are now supported on
 
 * To enable this feature, turn on the "Cast Modulated Shadows" option on a directional light!
 
-## New: Dynamic Point Lights for Mobile
+### New: Dynamic Point Lights for Mobile
 
 Dynamic point lights are now supported on mobile devices and HTML5!
 
@@ -37,7 +37,7 @@ Dynamic point lights are now supported on mobile devices and HTML5!
 
 * Note:  Dynamic shadows cast from point lights are not yet supported
 
-## New: Decals on Mobile
+### New: Decals on Mobile
 
 Decals are now supported by the mobile renderer!
 
@@ -47,9 +47,9 @@ Decals are now supported by the mobile renderer!
 
 * Note:  Only unlit decals are supported on mobile and web, currently.
 
-## New: Major VR Updates
+### New: Major VR Updates
 
-### Steam VR (HTC Vive)
+#### Steam VR (HTC Vive)
 
 Unreal Engine's SteamVR plugin includes many major fixes to make it easier to use, and more performant.
 
@@ -65,13 +65,13 @@ SteamVR improvements include:
 
 * Various improvements and bug fixes
 
-### Gear VR Updated to Mobile SDK 0.6.0
+#### Gear VR Updated to Mobile SDK 0.6.0
 
 This revision fixes several major Gear VR rendering bugs and issues present in the 0.5.0 release, as well as some hardware compatibility issues. 
 
 ![image alt text](image_5.png)
 
-## New: VR Motion Controller Support
+### New: VR Motion Controller Support
 
 Motion Controllers are now supported through a common abstraction layer in UE4!  
 
@@ -79,7 +79,7 @@ Motion Controllers are now supported through a common abstraction layer in UE4!
 
 This means that motion controller devices will now all be available through a common interface.  The HTC Vive controller is currently supported, with more (such as the PlayStation Move) to come soon.  Input button presses and triggers can now be routed to a single player from multiple controllers, using the Motion Controller key abstraction.  Additionally, you can add motion tracking to your project by simply adding a Motion Controller component to your character, which automatically updates anything attached to it to follow the position of your controllers!  We recommend that all plugins for motion controllers move over to the new system, so that they can be part of the device-agnostic interface.
 
-## New: Experimental DirectX 12 Support
+### New: Experimental DirectX 12 Support
 
 DirectX 12 is now supported as an experimental feature! If you are using Windows 10, try it out by running the engine with "-DX12" on the command line.
 
@@ -87,7 +87,7 @@ DirectX 12 is now supported as an experimental feature! If you are using Windows
 
 Microsoft's engineers added support for DirectX 12 to UE4 and we have worked with them to integrate their changes into 4.9. The feature is still new and is considered experimental. DirectX 12 offers a much lower-level rendering API that is more efficient and allows for rendering commands to be submitted in parallel across many threads, a feature inspired by console rendering APIs. Going forward, we'll continue to improve support for DirectX 12 and look for ways to leverage the new API in upcoming versions of the engine.
 
-## New: Full Scene Particle Collision with Mesh Distance Fields
+### New: Full Scene Particle Collision with Mesh Distance Fields
 
 Distance field GPU particle collision allows particles to efficiently collide with the entire scene.
 
@@ -103,7 +103,7 @@ It can also be used for various effects like and [snow that accumulates on surfa
 
 * The Collision module on GPU particle emitters has a new property called 'Collision Mode' where you can specify 'Distance Field'
 
-![image alt text](image_11.png)
+	![image alt text](image_11.png)
 
 * In order for a particle to collide we need to have a colliding plane, and this can be extracted very efficiently from the distance field.  A single sample of the distance field gives distance to the nearest colliding surface.  Computing the gradient of the distance field provides the direction to that surface, and the surface's normal.  Putting these together gives the collision plane.
 
@@ -113,11 +113,11 @@ It can also be used for various effects like and [snow that accumulates on surfa
 
 * Particle collision with distance fields is very efficient - it costs about the same as particle collision with scene depth.
 
-## New: Hierarchical LOD Fast Preview and Clustering
+### New: Hierarchical LOD Fast Preview and Clustering
 
 Hierarchical Level-of-Detail (HLOD) is a system introduced in 4.8 which allows many objects in your levels to be collapsed down to only a few objects when they are small on screen.  This helps you achieve much higher quality levels when viewing objects up close, and faster overall performance for your level.  We've made improvements to HLOD in 4.9 and will continue to in future releases.
 
-### Fast HLOD Preview
+#### Fast HLOD Preview
 
 In order to allow for rapid iteration on Hierarchical LOD builds, you can now use the Preview HLOD option. This will create the HLOD Actors with only a visual representation of their bounds and objects without merging and creating the HLODs meshes themselves. Their bounds are rendered according to the Draw Distance setting inside the LOD System settings for the specific Hierarchical LOD Level. 
 
@@ -125,13 +125,13 @@ In order to allow for rapid iteration on Hierarchical LOD builds, you can now us
 
 The image above shows both the first and second Hierarchical LOD level clusters popping into view, this gives a clear indication how the level is clustered without having to go through the time-consuming process of merging and creating the Hierarchical LOD meshes. Be aware that this is an initial implementation of the preview system and may tend to change and improve over time.
 
-### Custom LOD Clustering (HLOD Volumes)
+#### Custom LOD Clustering (HLOD Volumes)
 
 In order to allow for more granular control over the clustering process for the Hierarchical LODs, you can now place Hierarchical LOD Volumes which will define a cluster for the Actors it encapsulates. Use this to manually define which areas and objects within a level should be clustered together.
 
 ![image alt text](image_13.png)
 
-## New: "Arch Vis" Character Controls
+### New: "Arch Vis" Character Controls
 
 The new "ArchVisCharacter” plugin adds a character class that provides out-of-the-box controls appropriate for architectural visualization applications. It is intended to feel pleasant and natural to control at realistic world scales, and to enable smoother videos and live demos.
 
@@ -139,7 +139,7 @@ The new "ArchVisCharacter” plugin adds a character class that provides out-of-
 
 You can hook up ArchVisCharacter directly as your Pawn Class, or you can create custom blueprints from it to tweak the movement settings as you see fit.
 
-## New: Widget Depth Ordering
+### New: Widget Depth Ordering
 
 Widget Components in screen space now automatically sort based on distance from the viewer, ensuring the closest screen space widget component is always on top of the other components.
 
@@ -147,7 +147,7 @@ Widget Components in screen space now automatically sort based on distance from
 
 We now sort all widget components against one another based on distance to the camera. Previously it was based a manual Z Order, or whoever was added to the viewport first.
 
-## New: Area Shadows (for stationary lights)
+### New: Area Shadows (for stationary lights)
 
 You're no longer limited to sharp shadows with uniform penumbra sizes.  The Lightmass ray-tracer now supports area shadows are now fully support for Stationary lights!
 
@@ -163,7 +163,7 @@ Performance details:
 
 * There's also a new optimization for areas of maps which are only affected by a single stationary light channel (eg one Stationary Directional light).  Their shadow maps will be packed into a single channel texture instead of a 4 channel texture, for 1/4th the memory cost.
 
-## New: Ambient Occlusion Material Mask
+### New: Ambient Occlusion Material Mask
 
 The new Ambient Occlusion Material Mask feature lets you access Lightmass calculated AO in your material, which can be useful for procedural texturing, for example to add in aging effects and dirt in areas where it would accumulate.
 
@@ -185,7 +185,7 @@ Performance details:
 
 * AOMaterialMask texture memory can be inspected with the ‘ListTextures' console command.
 
-## New: Mesh Distance Field Materials
+### New: Mesh Distance Field Materials
 
 This feature allows your materials to cheaply determine the distance to the nearst solid from any point in space.  There are two new nodes to access the global distance field properties within the material editor. 
 
@@ -197,7 +197,7 @@ UE4 user Roel Bartstra has already found some interesting uses for this - one to
 
 As with all distance field features, this requires 'Generate Mesh Distance Fields' to be enabled in the Rendering Project settings.  Distance field material access is a Shader Model 5 feature.  Use a FeatureLevelSwitch material node to provide fallbacks for Shader Model 4 hardware.
 
-## New: Improved Distance Field Ambient Occlusion
+### New: Improved Distance Field Ambient Occlusion
 
 Sky occlusion gives nice soft shadows around objects lit by the sky, especially important in overcast lighting scenarios, a use case traditionally hard to render at high quality without precomputed lighting. Distance Field Ambient Occlusion solves this for fully dynamic games where both the time of day and environment can change at any time. We've put a lot of effort into improving the quality and performance of the technology for 4.9 and believe that it is now a shippable feature on medium spec PC and PlayStation 4 level hardware with a total GPU cost of under 4ms for a typical game scene!
 
@@ -231,7 +231,7 @@ This is a visualization of ray tracing the global distance field vs object dista
 
 ![image alt text](image_25.png)
 
-## New: Content Browser Advanced Search
+### New: Content Browser Advanced Search
 
 The Content Browser now supports advanced search syntax.  Unreal Engine 4 is used on many large projects, including some with teams comprising hundreds of developers, and content libraries with hundreds of thousands of assets!  These new search features improve the workflow for finding your content on projects of any size.
 
@@ -247,25 +247,25 @@ This augments the existing simple search syntax by providing a way to match key-
 
 * Collection (alias: Tag) - Tests against the names of any collections that contain the asset.
 
-## New: Collection Improvements
+### New: Collection Improvements
 
 This release includes many changes to Collections that allow you to help you manage assets in your game such as collection nesting, smart collections and tagging functionality.
 
-### Nested Collections
+#### Nested Collections
 
 * Your collections of assets can now be organized into a hierarchy.
 
-![image alt text](image_27.png)
+	![image alt text](image_27.png)
 
 * Child collections can be created by choosing the "New…" option from a given collections context menu. Alternatively, collections can be parented via drag-and-drop.
 
-![image alt text](image_28.png)
+	![image alt text](image_28.png)
 
-### Dynamic Collections
+#### Dynamic Collections
 
 * Dynamic collections leverage the power of the Advanced Content Browser Search Syntax to allow you to create and share Content Browser filters.
 
-![image alt text](image_29.png)
+	![image alt text](image_29.png)
 
 * Anything you can do with the Content Browser text search can be saved as a dynamic collection, including referencing other collections (both static and dynamic).
 
@@ -273,21 +273,21 @@ This release includes many changes to Collections that allow you to help you man
 
 	![image alt text](image_30.png)
 
-### Quick Asset Management
+#### Quick Asset Management
 
 * You can now quickly add an asset to many collections at once using the new Quick Asset Management check boxes.
 
-![image alt text](image_31.png)
+	![image alt text](image_31.png)
 
 * These are available within the collections view itself, or alternatively, via the context menu of any selected asset(s).
 
-![image alt text](image_32.png)
+	![image alt text](image_32.png)
 
-### Improved Robustness
+#### Improved Robustness
 
 * Collections will now automatically follow any redirectors that they contain, which avoids the issue of an asset seeming to vanish from a collection when it was renamed or moved. Additionally, collections are now considered when fixing up (and ultimately removing) redirectors via the Content Browser.
 
-### Improved Status Reporting
+#### Improved Status Reporting
 
 * Each collection now has a little status indicator at the right hand side of its row item. This changes colour depending on the current state of the collection, with each colour meaning:
 
@@ -301,13 +301,13 @@ This release includes many changes to Collections that allow you to help you man
 
     * Grey - The collection is empty and up-to-date.
 
-### Improved Asset Tooltips
+#### Improved Asset Tooltips
 
 * The static collections for a given object are now shown as part of its tooltip in the Content Browser.
 
-![image alt text](image_33.png)
+	![image alt text](image_33.png)
 
-## New: Plugin Creation Wizard
+### New: Plugin Creation Wizard
 
 Its now easy to create new C++ plugins for Unreal Engine 4!  Check out or new plugin wizard.
 
@@ -319,7 +319,7 @@ Its now easy to create new C++ plugins for Unreal Engine 4!  Check out or new pl
 
 * This will create all of the initial files you need for your plugin to compile and be loaded into the editor
 
-## New: Enhanced Curve Editor
+### New: Enhanced Curve Editor
 
 We've made many improvements to the Curve Editor to make it easier to manipulate curve keys and tangents.
 
@@ -347,13 +347,13 @@ We've made many improvements to the Curve Editor to make it easier to manipulate
 
     * Changed the text of the menu item to make it clear what's going to happen
 
-## New: Multiple Return Nodes in Blueprints
+### New: Multiple Return Nodes in Blueprints
 
 A function can now have multiple return nodes. As soon as execution reaches a return the function's execution is terminated. This can be used to return early from a Blueprint function, or to branch and return differing values.
 
 ![image alt text](image_36.png)
 
-## New: Construct Custom Objects in Blueprints
+### New: Construct Custom Objects in Blueprints
 
 In Blueprints, you can now spawn baser object types (not just actors/components). The Construct Object From Class node takes a class and creates a new object of that type, similar to Spawn Actor From Class (except for types that are not an actor).
 
@@ -363,7 +363,7 @@ In Blueprints, you can now spawn baser object types (not just actors/components)
 
 * For actor classes, you'll still use the Spawn Actor From Class node.  And for widgets, you'll use the Create Widget node.  In a future release we may try to combine these different features.
 
-## New: Blueprint Class Defaults
+### New: Blueprint Class Defaults
 
 In Blueprints, it is now possible to access a class's defaults through the new "Get Class Defaults" node. This can be helpful when working with data-only Blueprints, giving you access to those Blueprints' values without having to create an instance.
 
@@ -371,7 +371,7 @@ In Blueprints, it is now possible to access a class's defaults through the new "
 
 * To keep the output pins from growing too unwieldy, use the node's details to show/hide specific defaults
 
-## New: Level Blueprint Communication
+### New: Level Blueprint Communication
 
 Level Blueprints can now be communicated with through interfaces. They can adopt and implement interface functions just as Class Blueprints can. After an interface has been added to a Level Blueprint, calls can be made through interface message nodes (using a streaming level as the target).
 
@@ -381,25 +381,25 @@ Level Blueprints can now be communicated with through interfaces. They can adopt
 
     * Since you have to use "Get Streaming Level" to target a specific level, this will only work for sub-levels
 
-## New: Optimized Math Expression
+### New: Optimized Math Expression
 
 We've optimized the Math Expression node so that it is more performant than regular nodes. It is now roughly twice as fast to use a Math Expression node, when compared to a series of operation/function nodes representing the same expression.
 
 ![image alt text](image_40.png)
 
-## New: Blueprint Asset IDs
+### New: Blueprint Asset IDs
 
 We've made it possible for Blueprints to load assets on-demand instead of right away.  There are two new data types exposed to Blueprints: "Asset ID" and “Class Asset ID”. An Asset ID represents a loaded or unloaded asset, and can be resolved into an object using the “Resolve Asset” node. If the asset was not loaded, then the “Resolve Asset” node returns an invalid object. You can load the asset using the “Load Asset” node. Similarly, a “Class Asset ID” will resolve into a class, and can be loaded with the “Load Asset Class” node. 
 
 ![image alt text](image_41.png)
 
-## New: Montage Element Timing Interface
+### New: Montage Element Timing Interface
 
 ​We've added a new panel to the Montage editor to help you control the order that Montage events will fire when your animations execute at runtime.
 
 ![image alt text](image_42.png)
 
-## New: Non-Linear Animation Blending
+### New: Non-Linear Animation Blending
 
 We now support many different blending functions between animation poses and transitions!
 
@@ -417,7 +417,7 @@ This will allow you to specify just about any blend you want to perform. The len
 
 ![image alt text](image_46.gif)(convert:false)
 
-## New: Bone-driven Animation Controllers
+### New: Bone-driven Animation Controllers
 
 This allows a "Driver" bone to dynamically affect the motion of one or more “Driven” bones.  This is great for characters that have accessories attached!  You can avoid geometry intersecting during an animation, even when a lot of blending is used.
 
@@ -431,21 +431,21 @@ Here are the new settings for bone-driven animation controllers:
 
 ![image alt text](image_48.png)
 
-## New: Animation Transition Rules
+### New: Animation Transition Rules
 
 This release contains improvements to animation blending trees, particularly to how transition rules work.
 
-### Better handling of animation Asset Overrides
+#### Better handling of animation Asset Overrides
 
 Nodes like Time Remaining in transition rules now correctly respond to changes in animation length due to overriden animations in Child Animation Blueprints.  This means that animations referenced by transition rules no longer have to be the same length in every child Animation Blueprint.
 
-### Referencing the most relevant animation player
+#### Referencing the most relevant animation player
 
 In order to make state machines more maintainable, you can now use a new class of getter that always picks the highest weighted animation from the source state, rather than targeting a specific animation that might change in the future.  The following nodes are now available:
 
 ![image alt text](image_49.png)
 
-### Custom Blend Graph Improvements
+#### Custom Blend Graph Improvements
 
 Additional information is now exposed in Custom Blend graphs to allow them to be used for more than just a one-shot custom transition animation.  The following nodes that provide information about the corresponding transition node and source/destination states are now available in custom blend graphs:
 
@@ -453,32 +453,33 @@ Additional information is now exposed in Custom Blend graphs to allow them to be
 
 You can also use getter methods that reference the asset players, just like in a transition rule.
 
-## New: Animation Curve Evaluation Change
+### New: Animation Curve Evaluation Change
 
 Animation Curve Evaluation used to happen in Update phase of animation, but now it is moved to Evaluate phase. This comes with following benefits. 
 
-1. The curve's weight is properly is evaluated with blended weights. 
+* The curve's weight is properly is evaluated with blended weights. 
 
-    1. Cache Nodes
+    * Cache Nodes
 
-    2. Additive Nodes: Additive nodes will apply the delta of additive curve data from base to the current pose as animation bone transform does
+    * Additive Nodes: Additive nodes will apply the delta of additive curve data from base to the current pose as animation bone transform does
 
-    3. Layered Node: This was a bit more tricky to make it work properly as you don't know which curve is influencing which part of joint. Right now it does provide option of how to blend curves. 
+    * Layered Node: This was a bit more tricky to make it work properly as you don't know which curve is influencing which part of joint. Right now it does provide option of how to blend curves. 
 
-        1. Max Weight : It choose the max weight of curves from blended poses
+        * Max Weight : It choose the max weight of curves from blended poses
 
-        2. Normalize by Weight : It sums up all weights of blended poses for all curves and normalize to 1
+        * Normalize by Weight : It sums up all weights of blended poses for all curves and normalize to 1
 
-        3. Blend by Weight: It just sums up (weights * curve value)
-![image alt text](image_51.png)
+        * Blend by Weight: It just sums up (weights * curve value)
+	
+		![image alt text](image_51.png)
 
-2. Curve Evaluation has moved to multi-threading if you're using parallel evaluation.
+* Curve Evaluation has moved to multi-threading if you're using parallel evaluation.
 
-3. Eventually this will allow us to create a node that can drive curve from bone transform and vice versa
+* Eventually this will allow us to create a node that can drive curve from bone transform and vice versa
 
-4. However, if you need a curve that hasn't been evaluated (if you're ticking but not evaluating), you won't get the up-to-date data. Curve data is treated as bone transform now as opposed to treated as notifies before. 
+* However, if you need a curve that hasn't been evaluated (if you're ticking but not evaluating), you won't get the up-to-date data. Curve data is treated as bone transform now as opposed to treated as notifies before. 
 
-## New: Animation Asset Metadata Support
+### New: Animation Asset Metadata Support
 
 We now support Metadata that can be added to animation assets.  The metadata is a Blueprintable class derived from the Anim Meta Data class.  This allows you to add custom metadata to animation assets.  Anim Sequence, Anim Montage, Anim Composite, and Blendspace are all supported. 
 
@@ -488,7 +489,7 @@ You can query those data from an animation asset in C++ by calling the GetMetaDa
 
 ​
 
-## New: Sound Quality Levels
+### New: Sound Quality Levels
 
 A new feature has been added that gives sound designers the ability to manage audio memory usage and (to some extent) performance for lower end machines/devices.
 
@@ -504,7 +505,7 @@ When playing a standalone game, the Quality Level is specified via the Game User
 
 For Play in Editor you can specify which Quality Level to use via the Level Editor Play Settings. Changing this value will work while the Play session is active and new sounds will use the current Quality Level.
 
-## New: Custom Audio Attenuation Curves
+### New: Custom Audio Attenuation Curves
 
 You can now define your own custom curve for attenuating sounds rather than simply being limited to the built in algorithms.
 
@@ -512,7 +513,7 @@ In the attenuation settings set the Distance Algorithm to "Custom" and you will
 
 ![image alt text](image_55.png)
 
-## New: Actor Tick Intervals
+### New: Actor Tick Intervals
 
 You can now specify an interval at which a tick function will be executed allowing you to reduce the overhead of per-frame tick execution without the overhead or complexity of timers.
 
@@ -532,7 +533,7 @@ A few implementation details to be aware of.
 
 * An Actor's Custom Time Dilation does not apply to the Tick Interval.  The frequency with which the tick is dispatched remains constant regardless of whether the individual Actor has its own Time Dilation, that Time Dilation is only applied to the DeltaTime passed in to the tick function.
 
-## New: Actor Encroachment Detection
+### New: Actor Encroachment Detection
 
 Encroachment detection is also more robust and works in more situations.  We've added new customization options for encroachment, too.
 
@@ -540,13 +541,13 @@ Encroachment detection is also more robust and works in more situations.  We've
 
 The "Spawn Even If Colliding" option ("bNoCollisionFail" in C++) when spawning an actor has extended to support four methods of handling collision at the spawn point:
 
-1. Spawn it there anyway.
+* Spawn it there anyway.
 
-2. Spawn at suitable location nearby if possible, else just spawn it there anyway.
+* Spawn at suitable location nearby if possible, else just spawn it there anyway.
 
-3. Spawn only if a suitable location can be found nearby, else don't spawn.
+* Spawn only if a suitable location can be found nearby, else don't spawn.
 
-4. Just don't spawn it.
+* Just don't spawn it.
 
 Existing blueprints will update automatically to choose the correct corresponding choice. Some code has been deprecated and will need to be manually updated.
 
@@ -554,7 +555,7 @@ Actor classes have a new blueprint-exposed member called Spawn Collision Handlin
 
 For actors with a Movement Component, it is assumed the updated component is the primary collision shape and will be the only component tested for encroachment. For other actors, all components with appropriate collision settings will now be checked. Be aware this may result in some situations where actors fail to spawn after upgrading -- double check your collision settings, handling method, and spawn transform if this occurs.
 
-## New: Post Process Blending
+### New: Post Process Blending
 
 Post process "Blendables" are used to transition post process effects between different states. We've made improvements to this system.
 
@@ -570,7 +571,7 @@ Finally, here's an example of how you'd use Blueprints to edit the blending weig
 
 ![image alt text](image_60.png)
 
-## New: Runtime Asset Cache
+### New: Runtime Asset Cache
 
 Some games need to support storing local copies of assets that may not have shipped with the game. For example if you have a lot of procedurally created assets, or assets that come from a server (such as advertisement graphics.) The new Runtime Asset Cache is a generic asset cache system that can be used to persistently store data generated during runtime in configurable asset buckets.
 
@@ -588,7 +589,7 @@ The Runtime Asset Cache has the following characteristics:
 
 * When cache grows larger than specified size, oldest entries are removed from it.
 
-## New: Volume Decals (Experimental only!)
+### New: Volume Decals (Experimental only!)
 
 3D function field rendering is an alternative to polygons for describing geometry. It is better-suited to solving some hard rendering problems such as shadows, Boolean operations, smoothing, ambient occlusion, ray tracing, and deformation. 2D and 3D distance functions are already in use in various places in the engine, like fonts, lightmaps, and distance field shadows/ambient occlusion. The new Volume Decals feature enables a distance function to be rendered as opaque objects into the GBuffer, using all of the information normally available to deferred passes, including ambient occlusion, screen-space reflections, lighting, and depth of field. The bowl in the following image is a Volume Decal.
 
@@ -614,25 +615,25 @@ This is how the distance function for the sphere is defined:
 
 The feature is unfinished. Shadow casting, normal map support, and indirect lighting are missing, and performance and quality optimizations have not been made. In addition, the feature will always be costly on the GPU (distance function needs to be evaluated maybe hundreds of times for each pixel). We will polish it further by adding material functions, shadow casting, documentation, etc., but it might never be fast enough for real-time use.
 
-## New: UE4 Documentation Tags, Version and Skill Level
+### New: UE4 Documentation Tags, Version and Skill Level
 
 * Documentation pages now have skill level, engine version, and category tags which are displayed at the top and bottom of pages that enable you to find other pages with matching tags.
 
-![image alt text](image_65.png)
+	![image alt text](image_65.png)
 
 * Added ability to filter the sitemap by skill level and/or engine version. The sitemap is now an expandable tree view. When filtered, the tree is expanded to show matching items while non-matching items appear grayed out
 
-![image alt text](image_66.png)
+	![image alt text](image_66.png)
 
 * Skill Level, Engine Version, and Tags on pages also now link to the sitemap and automatically filter it.  Clicking a tag opens the sitemap filtered to only show pages with that tag.
 
-![image alt text](image_67.png)
+	![image alt text](image_67.png)
 
 * Links to tutorials (how-to's, quick starts, etc.) display the skill level and engine version of that tutorial:
 
-![image alt text](image_68.png)
+	![image alt text](image_68.png)
 
-## New and Updated UE4 Documentation
+### New and Updated UE4 Documentation
 
 * [Gear VR Global Menu](https://docs.unrealengine.com/latest/INT/Platforms/GearVR/GlobalMenu/index.html)
 
@@ -666,17 +667,17 @@ The feature is unfinished. Shadow casting, normal map support, and indirect ligh
 
     * [Multicast](https://docs.unrealengine.com/latest/INT/Gameplay/Analytics/Multicast/index.html)
 
-## New: Asynchronous Real-time Audio Decompression
+### New: Asynchronous Real-time Audio Decompression
 
 Platforms that support real-time ogg-vorbis based audio decompression (PC, Xbox One, Mac, Android) now perform the decompression asynchronously when possible to avoid blocking operations on the main game thread. 
 
-## New: Shared Resources for Feature Packs and Templates
+### New: Shared Resources for Feature Packs and Templates
 
 We now support shared resources in templates and feature packs. This is achieved by allowing both templates and feature packs to specify other feature packs to include. This means that many assets that were common to all templates now only exist once. It also means that we no longer need to duplicate assets that are shared between the Blueprint and C++ versions of templates.
 
 Multiple target platforms are supported, so different assets can exist in a template or feature pack for mobile and desktop targets. That's useful when you want lower detail models on one platform but not another, for example.
 
-#3New: Improved HTML5 (Easy Setup, Amazon S3 Support, Networking)
+### New: Improved HTML5 (Easy Setup, Amazon S3 Support, Networking)
 
 You no longer need to install any third party SDKs to start development with HTML5 and package your game for a web browser! To package, simply select HTML5 in the File -> Package menu. 
 
@@ -690,11 +691,11 @@ Other improvements include:
 
 * We now support uploading to Amazon S3 during packaging. Once the credentials for AWS are filled in, the build will be uploaded to S3 and the game can be served to browser directly from S3. 
 
-![image alt text](image_69.png)
+	![image alt text](image_69.png)
 
-# Release Notes
+## Release Notes
 
-## AI
+### AI
 
 * New: A dedicated AI asset category has been created and is available in Content Browser's context menu.![image alt text](image_70.png)
 
@@ -732,7 +733,7 @@ Other improvements include:
 
 * Blueprint compilation and loading are much more efficient.
 
-### Behavior Tree
+#### Behavior Tree
 
 * New: Added dynamic subtree injection support for behavior trees.
 
@@ -764,7 +765,7 @@ Other improvements include:
 
 * Replaced Has Reached Goal decorator with Is At Location.
 
-### Debugging Tools
+#### Debugging Tools
 
 * New: Added "-LogNavOctree" command line parameter, to log Nav Octree debug geometry to visual logger, on stop logging.
 
@@ -786,7 +787,7 @@ Other improvements include:
 
 * Tweaked description width for EQS debug display from 200 to 312 to allow for slightly longer descriptions without overlapping the next column. Item Description Width, Item Score Width, and Test Score Width can now be tweaked in code (they are now variables rather than in-line values).
 
-### Navigation
+#### Navigation
 
 * New: Added a dedicated blueprint-bindable event to AI Perception Component that's being triggered for every newly sensed stimulus.
 
@@ -838,7 +839,7 @@ Other improvements include:
 
 * Updated deprecation messages for Navigation System's "get random point" functions family to be more precise.
 
-## Animation
+### Animation
 
 * New: "Observe Bone" nodes for debugging animations.
 
@@ -846,11 +847,11 @@ Other improvements include:
 
 * The observe bone nodes are meant for debugging.  They don't cost much but are not currently compiled out of a shipping build, so you should probably remove them when you are done working with the graph.
 
-![image alt text](image_71.png)
+	![image alt text](image_71.png)
 
 * New: Added turntable feature for Persona.
 
-![image alt text](image_72.gif)(convert:false)
+	![image alt text](image_72.gif)(convert:false)
 
 * New: Added a "Remove this LOD" button to each LOD (after 0) in Persona's mesh mode.
 
@@ -860,9 +861,9 @@ Other improvements include:
 
     * Time from Sequence End to trigger blend out.
 
-    * <0 means using Blend Out Time, so Blend Out finishes as Montage ends.
+    * `<0` means using Blend Out Time, so Blend Out finishes as Montage ends.
 
-    * >=0 means using "Sequence End - Blend Out Trigger Time" to trigger blend out.
+    * `>=0` means using "Sequence End - Blend Out Trigger Time" to trigger blend out.
 
     * Also added as an optional parameter to Play Slot Animation As Dynamic Montage.
 
@@ -986,7 +987,7 @@ Other improvements include:
 
 * While piloting an object, 'move to' object button functionality works properly.
 
-## Audio
+### Audio
 
 * New: Added ability to set the volume headroom per platform with a new Engine.ini field.
 
@@ -1010,7 +1011,7 @@ Other improvements include:
 
 * Fixed a Core Audio API crash for Mac when the number of active sounds exceeded the number of available audio units.
 
-## Automation
+### Automation
 
 * New: Basic lighting tests have been separated from the 'General Editor Test'.
 
@@ -1028,7 +1029,7 @@ Other improvements include:
 
 * Fixed UAT base path being invalid if the UE4 workspace is mounted directly onto a drive letter (e.g. C:\Engine, D:\Engine, etc.).
 
-## Blueprints
+### Blueprints
 
 * New: Added Blueprint-exposed GetBoundingBox function to StaticMesh.
 
@@ -1366,7 +1367,7 @@ Other improvements include:
 
 * When signature of a function is updated, its instances in the Event Graph are properly updated.
 
-## Core
+### Core
 
 * New: A plugin with the same name but of different versions can be placed in multiple game projects. This allows studios with multiple projects to utilize the same plugin at different revisions as long as the plugin is part of the game project.
 
@@ -1548,7 +1549,7 @@ Other improvements include:
 
 * When queueing new async packages, async loading code will also look for existing packages in loaded packages to process to prevent infinite streaming times.
 
-## Editor and Tools
+### Editor and Tools
 
 * New: Added "Zip Up Project" button under File -> Package Project to grab necessary project files and package into a zip. For easily sending the project to other people excluding unnecessary files.
 
@@ -1556,7 +1557,7 @@ Other improvements include:
 
     * The vector and rotator widgets now support automatically rendering in a more compact mode as the details panel shrinks, making it possible to edit properties over a wider range of widths.
 
-![image alt text](image_73.gif)(convert:false)
+	![image alt text](image_73.gif)(convert:false)
 
 * Rotator UI Consistency:
 
@@ -1564,11 +1565,11 @@ Other improvements include:
 
     * In a Blueprint:
 
-![image alt text](image_74.png)
+		![image alt text](image_74.png)
 
     * In the level editor:
 
-![image alt text](image_75.png)
+		![image alt text](image_75.png)
 
     * All existing content is automatically updated and will work as expected, but you may notice wires crossing that didn't cross before since the order has changed.
 
@@ -1614,7 +1615,7 @@ Other improvements include:
 
     * This is so that it is easier to import animation for rigid bodies.
 
-* New:"'Light Types” show menu:
+* New: "Light Types” show menu:
 
     * "Lighting Components" -> “Ambient Occlusion” affects both DFAO and SSAO now.
 
@@ -1856,7 +1857,7 @@ Other improvements include:
 
 * Updated Material Preview Mesh properties to include Adjacency Info.
 
-## Content Browser
+### Content Browser
 
 * New: Improved the behavior of Ctrl+B in the Content Browser and Persona.
 
@@ -1886,7 +1887,7 @@ Other improvements include:
 
 * Local collections are now preserved when creating a copy of your project.
 
-## Landscape
+### Landscape
 
 * New: Added a "Mesh Vertical Offset" to landscape spline control points:
 
@@ -1944,7 +1945,7 @@ Other improvements include:
 
 * Undo now works correctly with "Clear Region Selection".
 
-## Material Editor
+### Material Editor
 
 * New: Added Map Check warning after loading a level in the editor that uses a material that needs one of its Usage Flags set. Unless resaved, such materials will render as the default material in cooked builds.
 
@@ -1964,7 +1965,7 @@ Other improvements include:
 
 * TextureSample material nodes now fetch the correct tooltips.
 
-## Matinee
+### Matinee
 
 * Fix bug that could cause sounds to start not quite at the beginning during playback.
 
@@ -1978,7 +1979,7 @@ Other improvements include:
 
 * Switching editor modes will now warn before closing Matinee mode.
 
-## Cascade
+### Cascade
 
 * All particle modules now reset to their distribution properties to their correct defaults instead of "None".
 
@@ -1986,7 +1987,7 @@ Other improvements include:
 
 * Mesh particles now properly populate dynamic parameter names in Cascade.
 
-## Persona
+### Persona
 
 * New: Added key handling to Persona segment panel for montages so delete key now works correctly.
 
@@ -2014,17 +2015,17 @@ Other improvements include:
 
 * Fixed UI refresh being triggered for any property edit event fired from any UObject. Now filtered to refresh only when notify data is changed.
 
-## PhAT
+### PhAT
 
 * Fixed issue where changing the mesh rendering mode in PhAT would not update in the viewport until the second click.
 
 * Fixed PhAT Editor crashing due to being able to select nonexistent bones.
 
-## World Outliner
+### World Outliner
 
 * New: Changed the level name tooltip in the Levels Browser to display the full package path name.
 
-## Static Mesh Editor
+### Static Mesh Editor
 
 * New: Added the ability to push the bounds for StaticMesh objects manually.
 
@@ -2036,7 +2037,7 @@ Other improvements include:
 
     * The flag will be auto-set for old meshes preserving existing behavior.
 
-## Source Control
+### Source Control
 
 * New: Added make file writable button for advanced project setting pages, which allows files to be edited when non-writable and not connected to source control.
 
@@ -2050,7 +2051,7 @@ Other improvements include:
 
     * This addresses issues where users have repositories and checkouts that were made with different versions of SVN to the client binary we ship. We now attempt to detect a system-level svn binary, and only use the shipped binary as a fallback.
 
-## World Browser
+### World Browser
 
 * New: LOD sub-level generation has more options for material simplification (texture size, constant material parameters).
 
@@ -2062,7 +2063,7 @@ Other improvements include:
 
 * World composition now incrementally updates sub-level tile images, preventing large editor hitches when a large number of tiles need update at the same time.
 
-## Foliage
+### Foliage
 
 * New: Added a bulk foliage instance removal Blueprint function to HierarchicalInstancedStaticMeshComponent
 
@@ -2094,7 +2095,7 @@ Other improvements include:
 
 * We no longer try to create physics data for InstancedStaticMeshComponent (eg Foliage) instances with a zero scale, which removes some PhysX warnings and improves performance.
 
-## Cooker
+### Cooker
 
 * New: General cooking performance optimizations.
 
@@ -2120,7 +2121,7 @@ Other improvements include:
 
 * Fixed various issues that were only showing up during the process of packaging a project from the editor.
 
-## Gameplay Framework
+### Gameplay Framework
 
 * New: Actor transforms when spawning.
 
@@ -2336,19 +2337,19 @@ Other improvements include:
 
 * You can now set the GameMode HUD class to be none.
 
-# Learning Resources
+## Learning Resources
 
-## Documentation
+### Documentation
 
 * Improved docs for Add Yaw Input and Add Controller Yaw Input etc., to point to the scaling values that are applied.
 
 * Improved docs on Movement Component and Floating Pawn Movement to indicate that collision is wholly dependent on the Updated Component (normally the root component).
 
-## Sample Content
+### Sample Content
 
 * Fixed UI icons not being present in sample content.
 
-## Localization
+### Localization
 
 * New: Added game native culture setting to control which culture's localization data the editor loads for game assets. Game assets now display the specified culture's localization data when in the editor.
 
@@ -2382,11 +2383,11 @@ Other improvements include:
 
 * Updated logic in commandlet for gathering text from assets which was falsely reporting that nothing was gathered because no files were loaded. Now properly reports whether there were any packages found to be processed, rather than if any need to be loaded and processed in a deferred manner.
 
-## Marketplace Content
+### Marketplace Content
 
 * Infiltrator Demo - Various maps have been cleaned up to remove NULL referenced scene assets.
 
-## Networking
+### Networking
 
 * New: A warning will now be printed to the output log if an actor tries to send a non-multicast remote function but has no owning connection.
 
@@ -2496,7 +2497,7 @@ Other improvements include:
 
 * When scrubbing replays, camera will continue to follow the original view target.
 
-## Online
+### Online
 
 * New in BuildPatchServices:
 
@@ -2530,7 +2531,7 @@ Other improvements include:
 
 * Fixed a rare issue where attempting to save out Steam achievements multiple times in the same frame could prevent further Steam tasks from executing.
 
-## Other
+### Other
 
 * New: Added a generic motion controller key abstraction
 
@@ -2594,7 +2595,7 @@ Other improvements include:
 
 * Added a dialog box to alert users when there is a problem loading a plugin descriptor file.
 
-## Paper2D
+### Paper2D
 
 * New: Exposed a way to set collision thickness / offset for tile maps and tile layers, and added a RebuildCollision method for use after SetTile calls.
 
@@ -2616,7 +2617,7 @@ Other improvements include:
 
     * All sprite colors are now treated as being defined in linear space, and existing assets will be converted on load.
 
-## Physics
+### Physics
 
 * New: Added GetPhysicsLinearVelocityAtPoint, SetCenterOfMass, SetMassOverrideInKg, GetClosestPointOnCollision to PrimitiveComponent.
 
@@ -2688,7 +2689,7 @@ Other improvements include:
 
 * Changing PhysicalMaterial override on an instance now correctly updates mass properties.
 
-## Platforms
+### Platforms
 
 * New: A prominent message is show when a user attempts to package a game that does not have binaries built.
 
@@ -2720,11 +2721,11 @@ Other improvements include:
 
 * Virtual joysticks now deal with screen resolution or orientation changes properly (repositioning and sizing the controls).
 
-## Android
+### Android
 
 * New: NVIDIA AndroidWorks installer.
 
-![image alt text](image_76.png)
+	![image alt text](image_76.png)
 
     * TADP replaced with NVIDIA AndroidWorks for Windows and Mac OSX.
 
@@ -2804,7 +2805,7 @@ Other improvements include:
 
 * Cleaned up some tooltips.
 
-## iOS
+### iOS
 
 * New: Added initial support for the Cloud Kit framework. Cloud-based databases can be accessed via the Online User Cloud Interface API and the Online Shared Cloud Interface API.
 
@@ -2842,7 +2843,7 @@ Other improvements include:
 
 * The splash screen will now properly be removed if a message box is displayed.
 
-## Linux
+### Linux
 
 * New: Added CodeLite project generator and source accessor. (Contributed by [REDACTED_PERSON], pull requests 1378 and 1379.)
 
@@ -2904,7 +2905,7 @@ Other improvements include:
 
 * UBT will properly recognize when UHT exits abnormally on Unix platforms (Linux, Mac) and report it as a crash.
 
-## Mac
+### Mac
 
 * New: Added support for Xcode 7 (for iOS 9 and MacOS 10.11).
 
@@ -2956,7 +2957,7 @@ Other improvements include:
 
 * Made the Xcode source accessor more robust on OS X by ensuring that the Xcode project required exists and allowing the user to switch to a different application while waiting for Xcode.
 
-## Playstation 4
+### Playstation 4
 
 * New: Added RE_Z feature to shader definitions.
 
@@ -3050,7 +3051,7 @@ Other improvements include:
 
 * PS4 now always attempts to pan for surround sound.
 
-## VR
+### VR
 
 * New: MSAA is now enabled by default when running on GearVR on Galaxy Note 4 and Galaxy S6 devices. You can disable this by setting gearvr.EnableMSAA=0 in your ConsoleVariables.ini file.
 
@@ -3058,7 +3059,7 @@ Other improvements include:
 
 * Fixed stereo rendering not working when Playing in VR while Blueprint editor was open.
 
-## HTML5
+### HTML5
 
 * New: HTML5 networking plug-in is now enabled by default.
 
@@ -3086,11 +3087,11 @@ Other improvements include:
 
 * We now properly copy the web server as part of the packaging process.
 
-## Windows
+### Windows
 
 * Changing OS Version for minimum supported OS on windows, no longer requires an editor restart.
 
-## All Mobile
+### All Mobile
 
 * New: Changed how non-power-of-two textures work on platforms that don't support them. Instead of silently failing, it will replace contents with an "error texture" so you can see on device what to do, instead of seeing nothing.
 
@@ -3110,7 +3111,7 @@ Other improvements include:
 
 * Remove redundant checks for Pixel-Buffer-Object support from the OpenGL RHIUpdateTexture2D & RHIUpdateTexture3D functions to allow them to be called on mobile. (Contributed by [REDACTED_PERSON], pull request 1086.)
 
-## Rendering
+### Rendering
 
 * New: Added a material function that can convert RGB colors to HSV.
 
@@ -3218,7 +3219,7 @@ Other improvements include:
 
 * Reverting diffuse model back to Lambert for simplicity with capturing materials.
 
-## FX
+### FX
 
 * New: Added DefaultValue property to Dynamic Parameter material node. Minor workflow improvement for VFX artists.
 
@@ -3234,7 +3235,7 @@ Other improvements include:
 
 * Removed "Particle" options from Beam Source and Target options as these features are not implemented in UE4.
 
-## Lighting
+### Lighting
 
 * New: Added bUseVertexNormalForHemisphereGather Lightmass primitive setting for foliage static lighting, in cases where artists want foliage lighting to match the underlying terrain.
 
@@ -3268,7 +3269,7 @@ Other improvements include:
 
 * Updated grouped shadow caster bounds calculation to ignore non-shadow casting primitives.
 
-## Materials
+### Materials
 
 * New: A few material functions have been added for generating and combining Shapes. These are mostly useful for the "Volumetric Decal" material type, but they may be useful in other situations as well.
 
@@ -3334,7 +3335,7 @@ Other improvements include:
 
 * Modified the camera offset material function node to reference vertex world position instead of actor position to remove z-depth scaling.
 
-## Postprocessing
+### Postprocessing
 
 * New: Improved Blendables, exposed weight on each blendable entry, better blueprint access.
 
@@ -3350,7 +3351,7 @@ Other improvements include:
 
 * Improved Temporal AA quality using bicubic filtering (slightly sharper).
 
-## Optimizations
+### Optimizations
 
 * New: Don't use IndirectLightingCache shaders when there are no precomputed lighting volumes.
 
@@ -3368,7 +3369,7 @@ Other improvements include:
 
 * Updated distortion pass to use a stencil mask which improves performance with lower screen coverage.
 
-## Mobile Rendering
+### Mobile Rendering
 
 * New: Enabled limited post process support for devices that do not support full HDR rendering, meaning those which lack float16 render target support (e.g. Galaxy S6).
 
@@ -3386,13 +3387,13 @@ Other improvements include:
 
 * Fixed stencil buffer rendering when using the Metal RHI.
 
-## UI
+### UI
 
 * New: In game objects Mouse Enter and Leave Events no longer fire if UI (Slate/UMG) is over it.
 
 * New: Materials designed for UI now have a new material domain that can be set in the material editor. This replaces the "Used with UI" checkbox on materials. All existing materials using this checkbox will be automatically updated.
 
-## Slate
+### Slate
 
 * New: Editable text widgets in Slate now allow the context menu to be override or disabled.
 
@@ -3456,7 +3457,7 @@ Other improvements include:
 
 * When we update the RHI data, we now also update what the slate render proxy thinks is the actual size of the texture, which might change due to a reimport.
 
-## UMG
+### UMG
 
 * New: Add features to allow for better in-timeline editing of animations.
 
@@ -3600,9 +3601,9 @@ Other improvements include:
 
     * Implementing controller support for the SScrollBox; it has custom logic for handling navigation inside of it scrolling the next element into view for the focuser and moving focus to the next logical child.
 
-# Programming Release Notes
+## Programming Release Notes
 
-## AI
+### AI
 
 * New: Added InvalidOrientation constant to AISystem.
 
@@ -3616,7 +3617,7 @@ Other improvements include:
 
 * Replaced EQS skip item mechanic with forcing item status.
 
-### Behavior Tree
+#### Behavior Tree
 
 * Added support for behavior tree's action task losing its PawnAction due to some external events.
 
@@ -3624,7 +3625,7 @@ Other improvements include:
 
 * Unified behavior tree's StopTree & StopLogic functions.
 
-### Navigation
+#### Navigation
 
 * New: Added pruning for crowd agent's corridor when close enough to next intermediate point (corner vertex). Depends on angle between first 2 segments of path, should reduce twitching when large group is moving through narrow corridor.
 
@@ -3640,7 +3641,7 @@ Other improvements include:
 
 * Navigation system will store weak object pointers to all custom navlinks to ensure their validity.
 
-## Animation
+### Animation
 
 * New: Added DLL export to some skeletal control classes and methods, making them extensible in another module.
 
@@ -3654,7 +3655,7 @@ Other improvements include:
 
 * Added ConvertMeshRotationPoseToLocalSpace.
 
-## Blueprints
+### Blueprints
 
 * New: BlueprintCallable functions can now specify the metadata HideSelfPin to indicate that the node should not show the Target pin, nor should it be connectable, thus filtering it from the context menu when dragging off a pin and ensuring it can only be used in the self context.
 
@@ -3672,7 +3673,7 @@ Other improvements include:
 
 * Fixed runtime blueprint use of atmospheric component.
 
-## Core
+### Core
 
 * New: Add ForEachObjectWithOuter and ForEachObjectOfClass utility functions that will call a TFunctionRef for each object instead of using GetObjectsWithOuter/OfClass which requires copying each object in to a (potentially very large) array.
 
@@ -3730,7 +3731,7 @@ Other improvements include:
 
 * New: Added new stats command which adds a custom marker to the stats stream.
 
-    * Can be used by typing stat namedmarker #markername# in console.
+    * Can be used by typing stat namedmarker #markername## in console.
 
 * New: Added static_asserts to TSharedPtr<T> to ensure that it is not used in the opposite thread-safety mode of T's TSharedFromThis.
 
@@ -3980,7 +3981,7 @@ Other improvements include:
 
 * UnrealHeaderTool will now be able to parse TSubclassOf<USomeClass>& parameters without reporting errors.
 
-## Editor and Tools
+### Editor and Tools
 
 * New: You can now help the blueprint utility functions create the correct blueprint for you. Previously, all classes created a standard blueprint. Now, the compiler module is asked Engine/Content/Types/Textures/SupportAndSettings/index.html if there is any registered IBlueprintCompiler with a special blueprint class for a particular UObject class. This fixes the problem of the new asset buttons creating UBlueprints for UUserWidget subclasses, instead of UWidgetBlueprints.
 
@@ -4000,17 +4001,17 @@ Other improvements include:
 
 * Moved KismetNodeInfoContext.h to be publically accessible outside of the GraphEditor module.
 
-## Content Browser
+### Content Browser
 
 * Changing the Content Browser path via the breadcrumb now notifies you of the change.
 
 * Fixed SPathView::OnPathSelected being called with an empty string after the editor loads.
 
-## Landscape
+### Landscape
 
 * Moved implementation details of the Landscape Editor into the Private folder.
 
-## Gameplay Framework
+### Gameplay Framework
 
 * New: Added GameplayStatics::DeprojectScreenToWorld and ProjectWorldToScreen functions.
 
@@ -4240,7 +4241,7 @@ Other improvements include:
 
 * Zeroed out FRepMovement::AngularVelocity when it's not used. More efficient for replication if it previously had non-zero values (switched from simulated to not simulated).
 
-## Localization
+### Localization
 
 * New: Created a new localization module for all the editor's general needs with regards to localization configuration scripts and localization pipeline configuration.
 
@@ -4248,7 +4249,7 @@ Other improvements include:
 
 * Writing manifest to file now uses full/absolute/rooted path to avoid any ambiguity when debugging.
 
-## Networking
+### Networking
 
 * New: Added a for loop to call PreReplication on all replicated components when Actor::PreReplication is called
 
@@ -4292,7 +4293,7 @@ Other improvements include:
 
 * Removed NetConstants.h and added MTU changes to CoreNet.h.
 
-## Online
+### Online
 
 * New: FOnlineExternalUI::ShowProfileUI has now been implemented for Steam. This displays a specific player's profile in the steam overlay.
 
@@ -4306,7 +4307,7 @@ Other improvements include:
 
 * Party beacons no longer attempt to send reservation updates when there are no clients connected.
 
-## Other
+### Other
 
 * New: Users can now provide functionality to web browser widgets to handle creation of popup windows.
 
@@ -4324,7 +4325,7 @@ Other improvements include:
 
 * Clear colors for hardware clears are now bound to rendertargets at creation time.
 
-## Physics
+### Physics
 
 * New: Add Distance to FHitResult info returned from traces & sweeps. (Contributed by [REDACTED_PERSON], pull request 1054.)
 
@@ -4350,7 +4351,7 @@ Other improvements include:
 
 * Reorder calls to SetHitResultFromShapeAndFaceIndex in CollisionConversions to better match memory layout of FHitResult.
 
-## Platforms
+### Platforms
 
 * New: Added process enumeration abstraction FProcEnumerator in FPlatformProcess. Now it's much easier to obtain system-wide processes information when you need it!
 
@@ -4364,7 +4365,7 @@ Other improvements include:
 
 * Fixed redundant initialization of target platforms.
 
-### Mac
+#### Mac
 
 * Added few warning ignore pragmas on Mac to make the engine compile with new Clang in Xcode 7.
 
@@ -4372,7 +4373,7 @@ Other improvements include:
 
 * Fixed paths to Mac targets in Xcode schemes.
 
-### Playstation 4
+#### Playstation 4
 
 * New: Minor fixes to PS4 parallel translate code including removal of "late resolve".
 
@@ -4380,11 +4381,11 @@ Other improvements include:
 
 * Fixed webapi deprecation warnings.
 
-### Xbox One
+#### Xbox One
 
 * Fixed warnings being displayed for missing ICU PDB files when building Xbox One.
 
-## Programming
+### Programming
 
 * New: We are now using the 64 bit Windows linker, even on 32 bit targets. This linker is faster and can handle larger link jobs without running out of memory.
 
@@ -4460,7 +4461,7 @@ Other improvements include:
 
 * Unreal Build Tool: Fixed incorrect VS tools versions set in VC++ project files.
 
-## Rendering
+### Rendering
 
 * New: Added miscellaneous parallel rendering tweaks and options. Balance deferred contexts after the fact, cvar can control occlusion latency, and a framework for Get Dynamic Mesh Elements tasks. These are all disabled by default.
 
@@ -4524,19 +4525,19 @@ Other improvements include:
 
 * Updated information about the expected buffer lifetime of BUF_Dynamic and BUF_Volatile.
 
-### Lighting
+#### Lighting
 
 * Fixed sRGB issue in Set Light Color for light components. (Contributed by [REDACTED_PERSON], pull request 1161.)
 
 * Improvement to Has Non Zero Data in FQuantizedLightmapData. (Contributed by [REDACTED_PERSON], pull request 790.)
 
-### Postprocessing
+#### Postprocessing
 
 * New: Changed all direct access to PassOutputs to GetInput() to allow internal restructuring.
 
-## UI
+### UI
 
-### Slate
+#### Slate
 
 * New: Added a way for UI actions to be repeated if the chord used to perform them is being held down.
 
@@ -4568,7 +4569,7 @@ Other improvements include:
 
 * Fixed path comparisons when checking a brush against a file.
 
-### UMG
+#### UMG
 
 * New: Adding the concept of shared layers to the Game Layer Manager. Rather than a widget added directly for a player, the layer manager allows you to add a layer by name at a specific ZOrder.
 
@@ -4578,9 +4579,9 @@ Other improvements include:
 
 * UE-16247 - We now correctly call SynchronizeProperties on all core widgets. Additionally there's now an ensure to prevent future mistakes due to forgetting to call Super::SynchronizeProperties().
 
-# Upgrade Notes
+## Upgrade Notes
 
-## C++ API Changes
+### C++ API Changes
 
 * SetRenderTargetsAndClear no longer takes clear values per target. The value bound at creation will be used.
 
@@ -4626,19 +4627,19 @@ Other improvements include:
 
     * Cleaned up constants in Pack/UnpackNetworkMovementMode() and made sure we only have a valid ground mode (Walking or NavWalking).
 
-## Other Upgrade Notes
+### Other Upgrade Notes
 
-## AI
+#### AI
 
 * If your game-specific AI controller class already derives from IGenericTeamAgentInterface, you probably should remove it from your controller's parent classes
 
-## Animation
+#### Animation
 
 * Games with duplicate, stale ActiveClassRedirect entries in their DefaultEngine.ini may have broken animation graphs until the duplicate stale redirectors are removed.
 
 * Started moving animation nodes out of Engine into a new AnimGraphRuntime module. Modules containing custom animation nodes will need to have a dependency on the new "AnimGraphRuntime" module added to their Build.cs file.
 
-## Core
+#### Core
 
 * Added message log error for source name collisions in redirector ini sections.
 
@@ -4658,29 +4659,29 @@ Other improvements include:
 
 * Modules must be listed in either DynamicallyLoadedModuleNames or Public/PrivateDependencyModuleNames, not both. If in doubt, remove the duplicate from DynamicallyLoadedModuleNames.
 
-## Editor and Tools
+#### Editor and Tools
 
 * If the ensure in FMainFrameModule::StartupModule() is hit in your game, you need to guard code that is editor-only via either the Editor module type, or checking IsRunningGame().
 
-## Localization
+#### Localization
 
 * Make sure to set your game's native culture setting in the Region & Language page of the Editor Preferences, so that assets load up that culture's localization for display.
 
-## Paper2D
+#### Paper2D
 
 * Fixed inconsistency between vertex colors on Paper2D assets and static meshes.
 
     * Blueprints or C++ code that programmatically modifies the color of a Paper2D component will need to be updated, otherwise the color being applied may be darker than expected.
 
-## Physics
+#### Physics
 
 * Note that constraint projection being turned on by default may slightly change the behavior of existing ragdolls. In general we are seeing much better stability as this is coupled with projection fixes inside PhysX.
 
 * The increase in default max angular velocity from 400 degrees per second to 3600 degrees per second can have a slight impact on existing simulated objects that were rotating at the previous max speed. You can use the physics project settings to change it back to the old behavior.
 
-## Platforms
+#### Platforms
 
-## Playstation 4
+##### Playstation 4
 
 * Your project must upgrade to SDK 2.5 to compile successfully
 
@@ -4688,11 +4689,11 @@ Other improvements include:
 
 * Projects must upgrade to SDK 2.5 and obtain 120hz Morpheus kits.
 
-**Mac**
+##### Mac
 
 * Using the WebBrowser module no longer requires creating a subclass of NSApplication.
 
-## All Mobile
+##### All Mobile
 
 * If you use either the iOS or Google Play online subsystems, the system login UI will no longer appear automatically during engine startup.
 
@@ -4704,11 +4705,11 @@ Other improvements include:
 
 * You must now explicitly invoke the system login UI by calling either "IOnlineExternalUI::ShowLoginUI" from C++ or "Show External Login UI" from Blueprints.
 
-## Programming
+#### Programming
 
 * The previously required boilerplate code can now be removed, but should still work.
 
-## Rendering
+#### Rendering
 
 * Added new cvar r.Filter.SizeScale allowing to scale sample count to tweak performance/quality
 
@@ -4716,11 +4717,11 @@ Other improvements include:
 
 * Removed RHIMethods.h and converted the RHI interface to standard C++.
 
-## Postprocessing
+#### Postprocessing
 
 * You might need to change some code to use GetInput() instead of accessing the member directly.
 
-## UMG
+#### UMG
 
 * Assets saved by editor versions prior to this change will copy the values of the deprecated properties to the style properties that replace them when they are loaded.
 
