#UE4doc #docs
Section 4 preliminary revision pass.

diff --git a/Engine/Documentation/Source/Programming/Tutorials/FirstPersonShooter/4/2/FirstPersonShooter_4_2.INT.udn b/Engine/Documentation/Source/Programming/Tutorials/FirstPersonShooter/4/2/FirstPersonShooter_4_2.INT.udn
--- a/Engine/Documentation/Source/Programming/Tutorials/FirstPersonShooter/4/2/FirstPersonShooter_4_2.INT.udn
+++ b/Engine/Documentation/Source/Programming/Tutorials/FirstPersonShooter/4/2/FirstPersonShooter_4_2.INT.udn
@@ -27,7 +27,7 @@ During this step, you'll edit the animation's Event Graph to ensure that your an
 
 ## Updating your State Variables
 
-1.	Open the event graph by double-clicking on **EventGraph** in the **My Blueprint** tab.
+1.	Open the event graph by **double-clicking** on **EventGraph** in the **My Blueprint** tab.
 
 	![](OpenEventGraph.png)
 
