#UE4doc #docs
Section 3 preliminary review pass.

diff --git a/Engine/Documentation/Source/Programming/Tutorials/FirstPersonShooter/3/1/FirstPersonShooter_3_1.INT.udn b/Engine/Documentation/Source/Programming/Tutorials/FirstPersonShooter/3/1/FirstPersonShooter_3_1.INT.udn
--- a/Engine/Documentation/Source/Programming/Tutorials/FirstPersonShooter/3/1/FirstPersonShooter_3_1.INT.udn
+++ b/Engine/Documentation/Source/Programming/Tutorials/FirstPersonShooter/3/1/FirstPersonShooter_3_1.INT.udn
@@ -22,19 +22,19 @@ Version: 4.9
 
 %Steps%
 
-Now that you've set-up your character, it's time to implement a simple projectile weapon so that when you fire, a simple grenade-like projectile will shoot from the center of the screen and fly until it collides with the world. During this step, you're going to add input and create a new code class for your projectile.
+Now that you've set-up your character, it's time to implement a projectile weapon so that when you fire, a simple grenade-like projectile will shoot from the center of the screen and fly until it collides with the world. During this step, you're going to add input and create a new code class for your projectile.
 
 ## Adding Fire Action Mapping
 
-1.	In the **Edit** menu, click on **Project Settings**.
+1.	In the **Edit** menu, **click** on **Project Settings**.
 
-1.	Under the **Engine** heading on the left side of the **Project Settings** tab, click on **Input**.
+1.	Under the **Engine** heading on the left side of the **Project Settings** tab, **click** on **Input**.
 
-1.	Under **Bindings**, click on the plus sign next to **Action Mappings**.
+1.	Under **Bindings**, **click** on the plus sign next to **Action Mappings**.
 
-1.	Click on the arrow to the left of **Action Mappings**.
+1.	**Click** on the arrow to the left of **Action Mappings**.
 
-1.	Type "Fire" into the text field that appears, then click on the arrow to the left of the text box to expand the action binding options.
+1.	Type "Fire" into the text field that appears, then **click** on the arrow to the left of the text box to expand the action binding options.
 
 1.	In the dropdown menu, select **Left Mouse Button** from the **Mouse** dropdown list.
 
@@ -48,7 +48,7 @@ Now that you've set-up your character, it's time to implement a simple projectil
 
 1.	In the File menu, select **New C++ Class...** to choose your new parent class.
 
-1.	The **Choose Parent Class** menu will open. Scroll down, select **Actor** as the parent class and click **Next**.
+1.	The **Choose Parent Class** menu will open. Scroll down, select **Actor** as the parent class and **click Next**.
 
 	![](AddProjectileClass.png)
 
@@ -727,7 +727,7 @@ It's now time to compile and check your newly implemented projectile code.
 
 1.	Locate **FPSProject** in the **Solution Explorer**.
 
-1.	Right-click on **FPSProject** and select **Build** to compile your project.
+1.	**Right-click** on **FPSProject** and select **Build** to compile your project.
 
 	![](BuildFPSProject.png)
 
diff --git a/Engine/Documentation/Source/Programming/Tutorials/FirstPersonShooter/3/2/FirstPersonShooter_3_2.INT.udn b/Engine/Documentation/Source/Programming/Tutorials/FirstPersonShooter/3/2/FirstPersonShooter_3_2.INT.udn
--- a/Engine/Documentation/Source/Programming/Tutorials/FirstPersonShooter/3/2/FirstPersonShooter_3_2.INT.udn
+++ b/Engine/Documentation/Source/Programming/Tutorials/FirstPersonShooter/3/2/FirstPersonShooter_3_2.INT.udn
@@ -226,7 +226,7 @@ During this step, you're going to implement the `OnFire` function so that your c
 
 1.	Locate **FPSProject** in the **Solution Explorer**.
 
-1.	Right-click on **FPSProject** and select **Build** to compile your project.
+1.	**Right-click** on **FPSProject** and select **Build** to compile your project.
 
 	![](BuildFPSProject.png)
 
@@ -327,6 +327,6 @@ Before continuing, download and extract our sample mesh from the following link:
 
 	![](FiringProjectiles.png)
 
-1.	Click the **Stop** button in the Level Editor to exit Play in Editor (PIE) mode.
+1.	**Click** the **Stop** button in the Level Editor to exit Play in Editor (PIE) mode.
 
 %Steps%
\ No newline at end of file
diff --git a/Engine/Documentation/Source/Programming/Tutorials/FirstPersonShooter/3/3/FirstPersonShooter_3_3.INT.udn b/Engine/Documentation/Source/Programming/Tutorials/FirstPersonShooter/3/3/FirstPersonShooter_3_3.INT.udn
--- a/Engine/Documentation/Source/Programming/Tutorials/FirstPersonShooter/3/3/FirstPersonShooter_3_3.INT.udn
+++ b/Engine/Documentation/Source/Programming/Tutorials/FirstPersonShooter/3/3/FirstPersonShooter_3_3.INT.udn
@@ -137,7 +137,7 @@ UE4 comes with several useful collision chanels; however, the engine also provid
 
 1.	Locate **FPSProject** in the **Solution Explorer**.
 
-1.	Right-click on **FPSProject** and select **Build** to compile your project.
+1.	**Right-click** on **FPSProject** and select **Build** to compile your project.
 
 	![](BuildFPSProject.png)
 
diff --git a/Engine/Documentation/Source/Programming/Tutorials/FirstPersonShooter/3/4/FirstPersonShooter_3_4.INT.udn b/Engine/Documentation/Source/Programming/Tutorials/FirstPersonShooter/3/4/FirstPersonShooter_3_4.INT.udn
--- a/Engine/Documentation/Source/Programming/Tutorials/FirstPersonShooter/3/4/FirstPersonShooter_3_4.INT.udn
+++ b/Engine/Documentation/Source/Programming/Tutorials/FirstPersonShooter/3/4/FirstPersonShooter_3_4.INT.udn
@@ -175,7 +175,7 @@ Now that you can detect your projectile's collision interactions, you can determ
 
 1.	Locate **FPSProject** in the **Solution Explorer**.
 
-1.	Right-click on **FPSProject** and select **Build** to compile your project.
+1.	**Right-click** on **FPSProject** and select **Build** to compile your project.
 
 	![](BuildFPSProject.png)
 
@@ -223,6 +223,6 @@ Now that you can detect your projectile's collision interactions, you can determ
 
 	![](ProjectileComplete.png)
 
-1.	Click the **Stop** button in the Level Editor to exit Play in Editor (PIE) mode.
+1.	**Click** the **Stop** button in the Level Editor to exit Play in Editor (PIE) mode.
 
 %Steps%
\ No newline at end of file
diff --git a/Engine/Documentation/Source/Programming/Tutorials/FirstPersonShooter/3/5/FirstPersonShooter_3_5.INT.udn b/Engine/Documentation/Source/Programming/Tutorials/FirstPersonShooter/3/5/FirstPersonShooter_3_5.INT.udn
--- a/Engine/Documentation/Source/Programming/Tutorials/FirstPersonShooter/3/5/FirstPersonShooter_3_5.INT.udn
+++ b/Engine/Documentation/Source/Programming/Tutorials/FirstPersonShooter/3/5/FirstPersonShooter_3_5.INT.udn
@@ -179,7 +179,7 @@ During this step, you'll add a crosshair HUD element to your game so that you ca
 
 1.	Locate **FPSProject** in the **Solution Explorer**.
 
-1.	Right-click on **FPSProject** and select **Build** to compile your project.
+1.	**Right-click** on **FPSProject** and select **Build** to compile your project.
 
 	![](BuildFPSProject.png)
 
@@ -228,17 +228,17 @@ During this step, you'll add a crosshair HUD element to your game so that you ca
 
 1.	Locate **FPSProject** in the **Solution Explorer**.
 
-1.	Right-click on **FPSProject** and select **Build** to compile your project.
+1.	**Right-click** on **FPSProject** and select **Build** to compile your project.
 
 	![](BuildFPSProject.png)
 
 1.	After the build finishes, open your **FPSProject** in Unreal Editor.
 
-1.	Click the **Play** button in the Level Editor Toolbar. You should be able to aim the projectiles with your newly created crosshairs.
+1.	**Click** the **Play** button in the Level Editor Toolbar. You should be able to aim the projectiles with your newly created crosshairs.
 
 	![](CrosshairsInGame.png)
 
-1.	Click the **Stop** button in the Level Editor to exit Play in Editor (PIE) mode.
+1.	**Click** the **Stop** button in the Level Editor to exit Play in Editor (PIE) mode.
 
 Congratulations! You've learned how to:
 
diff --git a/Engine/Documentation/Source/Programming/Tutorials/FirstPersonShooter/3/FirstPersonShooter_3.INT.udn b/Engine/Documentation/Source/Programming/Tutorials/FirstPersonShooter/3/FirstPersonShooter_3.INT.udn
--- a/Engine/Documentation/Source/Programming/Tutorials/FirstPersonShooter/3/FirstPersonShooter_3.INT.udn
+++ b/Engine/Documentation/Source/Programming/Tutorials/FirstPersonShooter/3/FirstPersonShooter_3.INT.udn
@@ -16,7 +16,7 @@ The purpose of this section is to show you how to implement projectiles for your
 
 ## Objectives
 
-By the end of this section of the tutorial you'll be able to:
+By the end of this section of the tutorial, you'll be able to:
 
 * Add Projectiles to your Game
 * Implement Shooting
