#UE4doc #docs
Section 2 preliminary revision pass.

diff --git a/Engine/Documentation/Source/Programming/Tutorials/FirstPersonShooter/2/1/FirstPersonShooter_2_1.INT.udn b/Engine/Documentation/Source/Programming/Tutorials/FirstPersonShooter/2/1/FirstPersonShooter_2_1.INT.udn
--- a/Engine/Documentation/Source/Programming/Tutorials/FirstPersonShooter/2/1/FirstPersonShooter_2_1.INT.udn
+++ b/Engine/Documentation/Source/Programming/Tutorials/FirstPersonShooter/2/1/FirstPersonShooter_2_1.INT.udn
@@ -52,7 +52,7 @@ First, you're going to edit your **GameMode** so that FPSCharacter is the defaul
 
 	![](ExpandedSolutionExplorer.png)
 
-1.	Double-click `FPSGameMode.cpp` to open the implementation file for your **FPSGameMode** class.
+1.	**Double-click** `FPSGameMode.cpp` to open the implementation file for your **FPSGameMode** class.
 
 1.	First, include the `FPSCharacter` header at the top of `FPSGameMode.cpp`:
 
@@ -95,7 +95,7 @@ First, you're going to edit your **GameMode** so that FPSCharacter is the defaul
 
 	![](ExpandedSolutionExplorer.png)
 
-1.	Double-click `FPSCharacter.cpp` to open the implementation file for your **FPSCharacter** class.
+1.	**Double-click** `FPSCharacter.cpp` to open the implementation file for your **FPSCharacter** class.
 
 1.	Add the following lines of code in the `BeginPlay()` function to verify that `FPSCharacter` class is being used:
 
@@ -148,7 +148,7 @@ First, you're going to edit your **GameMode** so that FPSCharacter is the defaul
 
 1.	Locate **FPSProject** in the **Solution Explorer**.
 
-1.	Right-click on **FPSProject** and select **Build** to compile your project.
+1.	**Right-click** on **FPSProject** and select **Build** to compile your project.
 
 	![](BuildProject.png)
 
@@ -156,7 +156,7 @@ First, you're going to edit your **GameMode** so that FPSCharacter is the defaul
 
 	![](FPSCharacterContentBrowser.png)
 
-1.	Click the **Play** button in the Level Editor Toolbar. "We are using FPSCharacter." should now be displayed below "We are using FPSGameMode." in the upper left corner of the viewport.
+1.	**Click** the **Play** button in the Level Editor Toolbar. "We are using FPSCharacter." should now be displayed below "We are using FPSGameMode." in the upper left corner of the viewport.
 
 	![](VerifyingFPSCharacterResult.png)
 
@@ -164,6 +164,6 @@ First, you're going to edit your **GameMode** so that FPSCharacter is the defaul
 	Your new Character doesn't have any movement controls yet, so you won't be able to move around in the level. If you're unable to move, you are using FPSCharacter as your Pawn correctly!
 	[/REGION]
 
-1.	Click the **Stop** button in the Level Editor to exit Play in Editor (PIE) mode.
+1.	**Click** the **Stop** button in the Level Editor to exit Play in Editor (PIE) mode.
 
 %Steps%
\ No newline at end of file
diff --git a/Engine/Documentation/Source/Programming/Tutorials/FirstPersonShooter/2/2/FirstPersonShooter_2_2.INT.udn b/Engine/Documentation/Source/Programming/Tutorials/FirstPersonShooter/2/2/FirstPersonShooter_2_2.INT.udn
--- a/Engine/Documentation/Source/Programming/Tutorials/FirstPersonShooter/2/2/FirstPersonShooter_2_2.INT.udn
+++ b/Engine/Documentation/Source/Programming/Tutorials/FirstPersonShooter/2/2/FirstPersonShooter_2_2.INT.udn
@@ -39,9 +39,9 @@ You will now set-up axis mappings for the W, A, S, and D keys.
 
 1.	Under **Bindings**, click on the plus sign next to **Axis Mappings**.
 
-1.	Click on the arrow to the left of **Axis Mappings**.
+1.	**Click** on the arrow to the left of **Axis Mappings**.
 
-1.	Type "MoveForward" into the text field that appears, then click on the arrow to the left of the text box to expand the axis binding options.
+1.	Type "MoveForward" into the text field that appears, then **click** on the arrow to the left of the text box to expand the axis binding options.
 
 1.	In the dropdown menu, select **W** from the **Keyboard** dropdown list.
 
@@ -49,7 +49,7 @@ You will now set-up axis mappings for the W, A, S, and D keys.
 
 	![](MoveForwardAxisMap_W.png)
 
-1.	Click on the plus sign next to **MoveForward**.
+1.	**Click** on the plus sign next to **MoveForward**.
 
 1.	In the second dropdown menu, select **S** from the **Keyboard** dropdown list.
 
@@ -71,7 +71,7 @@ You will now set-up axis mappings for the W, A, S, and D keys.
 
 	![](MoveRightAxisMap_D.png)
 
-1.	Click on the plus sign next to **MoveRight**.
+1.	**Click** on the plus sign next to **MoveRight**.
 
 1.	In the second dropdown menu, select **A** from the **Keyboard** dropdown list.
 
diff --git a/Engine/Documentation/Source/Programming/Tutorials/FirstPersonShooter/2/3/FirstPersonShooter_2_3.INT.udn b/Engine/Documentation/Source/Programming/Tutorials/FirstPersonShooter/2/3/FirstPersonShooter_2_3.INT.udn
--- a/Engine/Documentation/Source/Programming/Tutorials/FirstPersonShooter/2/3/FirstPersonShooter_2_3.INT.udn
+++ b/Engine/Documentation/Source/Programming/Tutorials/FirstPersonShooter/2/3/FirstPersonShooter_2_3.INT.udn
@@ -219,14 +219,14 @@ It's now time to compile and test our newly implemented character movement funct
 
 1.	Locate **FPSProject** in the **Solution Explorer**.
 
-1.	Right-click on **FPSProject** and select **Build** to compile your project.
+1.	**Right-click** on **FPSProject** and select **Build** to compile your project.
 
 	![](BuildProject.png)
 
 1.	After the build finishes, open your **FPSProject** in Unreal Editor.
 
-1.	Click the **Play** button in the Level Editor Toolbar. You should be able to move and strafe throughout the map (although your camera will be fixed in place).
+1.	**Click** the **Play** button in the Level Editor Toolbar. You should be able to move and strafe throughout the map (although your camera will be fixed in place).
 
-1.	Click the **Stop** button in the Level Editor to exit Play in Editor (PIE) mode.
+1.	**Click** the **Stop** button in the Level Editor to exit Play in Editor (PIE) mode.
 
 %Steps%
\ No newline at end of file
diff --git a/Engine/Documentation/Source/Programming/Tutorials/FirstPersonShooter/2/4/FirstPersonShooter_2_4.INT.udn b/Engine/Documentation/Source/Programming/Tutorials/FirstPersonShooter/2/4/FirstPersonShooter_2_4.INT.udn
--- a/Engine/Documentation/Source/Programming/Tutorials/FirstPersonShooter/2/4/FirstPersonShooter_2_4.INT.udn
+++ b/Engine/Documentation/Source/Programming/Tutorials/FirstPersonShooter/2/4/FirstPersonShooter_2_4.INT.udn
@@ -27,15 +27,15 @@ During this step, you're going to add the ability for your character to look aro
 
 ## Turn Axis Mapping
 
-1.	In the **Edit** menu, click on **Project Settings**.
+1.	In the **Edit** menu, **click** on **Project Settings**.
 
-1.	Under the **Engine** heading on the left side of the **Project Settings** tab, click on **Input**.
+1.	Under the **Engine** heading on the left side of the **Project Settings** tab, **click** on **Input**.
 
-1.	Under **Bindings**, click on the plus sign next to **Axis Mappings**.
+1.	Under **Bindings**, **click** on the plus sign next to **Axis Mappings**.
 
-1.	Click on the arrow to the left of **Axis Mappings**.
+1.	**Click** on the arrow to the left of **Axis Mappings**.
 
-1.	Type "Turn" into the text field that appears, then click on the arrow to the left of the text box to expand the axis binding options.
+1.	Type "Turn" into the text field that appears, then **click** on the arrow to the left of the text box to expand the axis binding options.
 
 1.	In the dropdown menu, select **Mouse X** from the **Mouse** dropdown list.
 
@@ -45,9 +45,9 @@ During this step, you're going to add the ability for your character to look aro
 
 ## Look-up Axis Mapping
 
-1.	Under **Bindings**, click on the plus sign next to **Axis Mappings**.
+1.	Under **Bindings**, **click** on the plus sign next to **Axis Mappings**.
 
-1.	Type **LookUp** into the text field that appears, then click on the arrow to the left of the text box to expand the axis binding options.
+1.	Type **LookUp** into the text field that appears, then **click** on the arrow to the left of the text box to expand the axis binding options.
 
 1.	In the dropdown menu, select **Mouse Y** from the **Mouse** dropdown list.
 
@@ -90,6 +90,8 @@ If you want to perform additional processing, such as adding support for sensiti
 			InputComponent->BindAxis("LookUp", this, &AFPSCharacter::AddControllerPitchInput);
 		}
 
+## Testing Mouse Camera Control
+
 1.	Save the `FPSCharacter` implementation file in Visual Studio.
 
 1.	Locate **FPSProject** in the **Solution Explorer**.
@@ -100,12 +102,8 @@ If you want to perform additional processing, such as adding support for sensiti
 
 1.	After the build finishes, open your **FPSProject** in Unreal Editor.
 
-1.	Click the **Play** button in the Level Editor Toolbar. You can now control the camera with your mouse.
-
-1.	Click the **Stop** button in the Level Editor to exit Play in Editor (PIE) mode.
-
-## Testing Mouse Camera Control
-
+1.	**Click** the **Play** button in the Level Editor Toolbar. You can now control the camera with your mouse.
 
+1.	**Click** the **Stop** button in the Level Editor to exit Play in Editor (PIE) mode.
 
 %Steps%
\ No newline at end of file
diff --git a/Engine/Documentation/Source/Programming/Tutorials/FirstPersonShooter/2/5/FirstPersonShooter_2_5.INT.udn b/Engine/Documentation/Source/Programming/Tutorials/FirstPersonShooter/2/5/FirstPersonShooter_2_5.INT.udn
--- a/Engine/Documentation/Source/Programming/Tutorials/FirstPersonShooter/2/5/FirstPersonShooter_2_5.INT.udn
+++ b/Engine/Documentation/Source/Programming/Tutorials/FirstPersonShooter/2/5/FirstPersonShooter_2_5.INT.udn
@@ -31,15 +31,15 @@ In general, Action Mappings deal with inputs for discrete events, they allow you
 
 ## Jump Action Mapping
 
-1.	In the **Edit** menu, click on **Project Settings**.
+1.	In the **Edit** menu, **click** on **Project Settings**.
 
-1.	Under the **Engine** heading on the left side of the **Project Settings** tab, click on **Input**.
+1.	Under the **Engine** heading on the left side of the **Project Settings** tab, **click** on **Input**.
 
-1.	Under **Bindings**, click on the plus sign next to **Action Mappings**.
+1.	Under **Bindings**, **click** on the plus sign next to **Action Mappings**.
 
-1.	Click on the arrow to the left of **Action Mappings**.
+1.	**Click** on the arrow to the left of **Action Mappings**.
 
-1.	Type "Jump" into the text field that appears, then click on the arrow to the left of the text box to expand the action binding options.
+1.	Type "Jump" into the text field that appears, then **click** on the arrow to the left of the text box to expand the action binding options.
 
 1.	In the dropdown menu, select **Space Bar** from the **Keyboard** dropdown list.
 
@@ -233,14 +233,14 @@ It's now time to compile and test your newly implemented character movement func
 
 1.	Locate **FPSProject** in the **Solution Explorer**.
 
-1.	Right-click on **FPSProject** and select **Build** to compile your project.
+1.	**Right-click** on **FPSProject** and select **Build** to compile your project.
 
 	![](BuildProject.png)
 
 1.	After the build finishes, open your **FPSProject** in Unreal Editor.
 
-1.	Click the **Play** button in the Level Editor Toolbar. You should be able to jump throughout the map.
+1.	**Click** the **Play** button in the Level Editor Toolbar. You should be able to jump throughout the map.
 
-1.	Click the **Stop** button in the Level Editor to exit Play in Editor (PIE) mode.
+1.	**Click** the **Stop** button in the Level Editor to exit Play in Editor (PIE) mode.
 
 %Steps%
\ No newline at end of file
diff --git a/Engine/Documentation/Source/Programming/Tutorials/FirstPersonShooter/2/6/FirstPersonShooter_2_6.INT.udn b/Engine/Documentation/Source/Programming/Tutorials/FirstPersonShooter/2/6/FirstPersonShooter_2_6.INT.udn
--- a/Engine/Documentation/Source/Programming/Tutorials/FirstPersonShooter/2/6/FirstPersonShooter_2_6.INT.udn
+++ b/Engine/Documentation/Source/Programming/Tutorials/FirstPersonShooter/2/6/FirstPersonShooter_2_6.INT.udn
@@ -159,19 +159,19 @@ It's now time to compile and test your newly added mesh in-editor.
 
 1.	Locate **FPSProject** in the **Solution Explorer**.
 
-1.	Right-click on **FPSProject** and select **Build** to compile your project.
+1.	**Right-click** on **FPSProject** and select **Build** to compile your project.
 
 	![](BuildProject.png)
 
 1.	After the build finishes, open your **FPSProject** in Unreal Editor.
 
-1.	Click the **Play** button in the Level Editor Toolbar. You should be able to see your character's shadow as you move around.
+1.	**Click** the **Play** button in the Level Editor Toolbar. You should be able to see your character's shadow as you move around.
 
 	[REGION:note]
-	To freely move the camera to see your character's mesh, press F1 to regain control of your mouse cursor, then click the **Eject** button in the toolbar.
+	To freely move the camera to see your character's mesh, press F1 to regain control of your mouse cursor, then **click** the **Eject** button in the toolbar.
 	![](FreelyMovingCamera.png)
 	[/REGION]
 
-1.	Click the **Stop** button in the Level Editor to exit Play in Editor (PIE) mode.
+1.	**Click** the **Stop** button in the Level Editor to exit Play in Editor (PIE) mode.
 
 %Steps%
\ No newline at end of file
diff --git a/Engine/Documentation/Source/Programming/Tutorials/FirstPersonShooter/2/7/FirstPersonShooter_2_7.INT.udn b/Engine/Documentation/Source/Programming/Tutorials/FirstPersonShooter/2/7/FirstPersonShooter_2_7.INT.udn
--- a/Engine/Documentation/Source/Programming/Tutorials/FirstPersonShooter/2/7/FirstPersonShooter_2_7.INT.udn
+++ b/Engine/Documentation/Source/Programming/Tutorials/FirstPersonShooter/2/7/FirstPersonShooter_2_7.INT.udn
@@ -235,16 +235,16 @@ It's now time to compile and test your newly implemented camera code.
 
 1.	Locate **FPSProject** in the **Solution Explorer**.
 
-1.	Right-click on **FPSProject** and select **Build** to compile your project.
+1.	**Right-click** on **FPSProject** and select **Build** to compile your project.
 
 	![](BuildProject.png)
 
 1.	After the build finishes, open your **FPSProject** in Unreal Editor.
 
-1.	Click the **Play** button in the Level Editor Toolbar. Your camera should be above the character's head and if you look down, you should be able to see the top of your character's head.
+1.	**Click** the **Play** button in the Level Editor Toolbar. Your camera should be above the character's head and if you look down, you should be able to see the top of your character's head.
 
 	![](NewCameraComponent.png)
 
-1.	Click the **Stop** button in the Level Editor to exit Play in Editor (PIE) mode.
+1.	**Click** the **Stop** button in the Level Editor to exit Play in Editor (PIE) mode.
 
 %Steps%
\ No newline at end of file
diff --git a/Engine/Documentation/Source/Programming/Tutorials/FirstPersonShooter/2/8/FirstPersonShooter_2_8.INT.udn b/Engine/Documentation/Source/Programming/Tutorials/FirstPersonShooter/2/8/FirstPersonShooter_2_8.INT.udn
--- a/Engine/Documentation/Source/Programming/Tutorials/FirstPersonShooter/2/8/FirstPersonShooter_2_8.INT.udn
+++ b/Engine/Documentation/Source/Programming/Tutorials/FirstPersonShooter/2/8/FirstPersonShooter_2_8.INT.udn
@@ -312,7 +312,7 @@ Before continuing, download and extract our sample mesh from the following link:
 	![](FreeMovingCamera.png)
 	[/REGION]
 
-1.	Click the **Stop** button in the Level Editor to exit Play in Editor (PIE) mode.
+1.	**Click** the **Stop** button in the Level Editor to exit Play in Editor (PIE) mode.
 
 Congratulations! You've learned how to:
 
diff --git a/Engine/Documentation/Source/Programming/Tutorials/FirstPersonShooter/2/FirstPersonShooter_2.INT.udn b/Engine/Documentation/Source/Programming/Tutorials/FirstPersonShooter/2/FirstPersonShooter_2.INT.udn
--- a/Engine/Documentation/Source/Programming/Tutorials/FirstPersonShooter/2/FirstPersonShooter_2.INT.udn
+++ b/Engine/Documentation/Source/Programming/Tutorials/FirstPersonShooter/2/FirstPersonShooter_2.INT.udn
@@ -16,7 +16,7 @@ The purpose of this section is to show you how to implement your First Person Sh
 
 ## Objectives
 
-By the end of this section of the tutorial you'll be able to:
+By the end of this section of the tutorial, you'll be able to:
 
 * Make a New Character
 * Set-up Axis Mapping
