#UE4doc #doc
Section 1 preliminary revision pass.

diff --git a/Engine/Documentation/Source/Programming/Tutorials/FirstPersonShooter/1/3/FirstPersonShooter_1_3.INT.udn b/Engine/Documentation/Source/Programming/Tutorials/FirstPersonShooter/1/3/FirstPersonShooter_1_3.INT.udn
--- a/Engine/Documentation/Source/Programming/Tutorials/FirstPersonShooter/1/3/FirstPersonShooter_1_3.INT.udn
+++ b/Engine/Documentation/Source/Programming/Tutorials/FirstPersonShooter/1/3/FirstPersonShooter_1_3.INT.udn
@@ -31,7 +31,7 @@ You'll add a log message to your newly created FPSGameMode in this step. Log mes
 
 	![](InterfaceFile.png)
 
-1.	Double-click `FPSGameMode.h` to open the interface file for your newly created **FPSGameMode** class.
+1.	**Double-click** `FPSGameMode.h` to open the interface file for your newly created **FPSGameMode** class.
 
 1.	Your class declaration should like the following:
 
@@ -83,7 +83,7 @@ You'll add a log message to your newly created FPSGameMode in this step. Log mes
 
 	![](ImplementationFile.png)
 
-1.	Double-click `FPSGameMode.cpp` to open the implementation file for your **FPSGameMode** class.
+1.	**Double-click** `FPSGameMode.cpp` to open the implementation file for your **FPSGameMode** class.
 
 1.	Add the following constructor definition to `FPSGameMode.cpp`:
 
diff --git a/Engine/Documentation/Source/Programming/Tutorials/FirstPersonShooter/1/4/FirstPersonShooter_1_4.INT.udn b/Engine/Documentation/Source/Programming/Tutorials/FirstPersonShooter/1/4/FirstPersonShooter_1_4.INT.udn
--- a/Engine/Documentation/Source/Programming/Tutorials/FirstPersonShooter/1/4/FirstPersonShooter_1_4.INT.udn
+++ b/Engine/Documentation/Source/Programming/Tutorials/FirstPersonShooter/1/4/FirstPersonShooter_1_4.INT.udn
@@ -27,7 +27,7 @@ You will now compile your project in this step so that you can see your code cha
 
 1.	Locate **FPSProject** in the **Solution Explorer**.
 
-1.	Right-click on **FPSProject** and select **Build** to compile your project.
+1.	**Right-click** on **FPSProject** and select **Build** to compile your project.
 
 	![](CompileProject.png)
 
diff --git a/Engine/Documentation/Source/Programming/Tutorials/FirstPersonShooter/1/5/FirstPersonShooter_1_5.INT.udn b/Engine/Documentation/Source/Programming/Tutorials/FirstPersonShooter/1/5/FirstPersonShooter_1_5.INT.udn
--- a/Engine/Documentation/Source/Programming/Tutorials/FirstPersonShooter/1/5/FirstPersonShooter_1_5.INT.udn
+++ b/Engine/Documentation/Source/Programming/Tutorials/FirstPersonShooter/1/5/FirstPersonShooter_1_5.INT.udn
@@ -23,7 +23,7 @@ Version: 4.9
 
 %Steps%
 
-Now that we've successfully compiled our newly added Game Mode, we'll need to set **FPSProject** to use **FPSGameMode** as the default Game Mode in this step.
+Now that you've successfully compiled your newly added Game Mode, you'll need to set **FPSProject** to use **FPSGameMode** as the default Game Mode in this step.
 
 1.	In the **Edit** menu, click on **Project Settings**.
 
diff --git a/Engine/Documentation/Source/Programming/Tutorials/FirstPersonShooter/1/FirstPersonShooter_1.INT.udn b/Engine/Documentation/Source/Programming/Tutorials/FirstPersonShooter/1/FirstPersonShooter_1.INT.udn
--- a/Engine/Documentation/Source/Programming/Tutorials/FirstPersonShooter/1/FirstPersonShooter_1.INT.udn
+++ b/Engine/Documentation/Source/Programming/Tutorials/FirstPersonShooter/1/FirstPersonShooter_1.INT.udn
@@ -15,7 +15,7 @@ The purpose of this section is to show you how to set up your First Person Shoot
 
 ## Objectives
 
-By the end of this section of the tutorial you'll be able to:
+By the end of this section of the tutorial, you'll be able to:
 
 * Set up a New Project
 * Use the C++ Class Wizard
