Re-Adding a link to the 13th step image as it was removed for some reason.

diff --git a/Engine/Documentation/Source/Engine/Physics/Destructibles/HowTo/DestructibleMesh/DestructibleMesh.INT.udn b/Engine/Documentation/Source/Engine/Physics/Destructibles/HowTo/DestructibleMesh/DestructibleMesh.INT.udn
--- a/Engine/Documentation/Source/Engine/Physics/Destructibles/HowTo/DestructibleMesh/DestructibleMesh.INT.udn
+++ b/Engine/Documentation/Source/Engine/Physics/Destructibles/HowTo/DestructibleMesh/DestructibleMesh.INT.udn
@@ -77,6 +77,8 @@ For this tutorial we are using the **Blueprint First Person** template with **St
 
 1. Select the Destructible Mesh that was just placed in the level and then in the **Details** panel under the **Physics** section turn on the **Simulate Physics** option.
 
+	![](Enable_Physics.png)
+
 1. Now press the **Play** button and when the level starts you should see the Destructible Mesh that was added fall down to the ground. 
 	Once the Destructible Mesh makes contact with the ground it will fracture into many smaller pieces like in the video below.
 
