#UEDocs #UEDoc: Updated node names in tips and tricks

diff --git a/Engine/Documentation/Source/Engine/Rendering/Materials/HowTo/DetailTexturing/DetailTexturing.INT.udn b/Engine/Documentation/Source/Engine/Rendering/Materials/HowTo/DetailTexturing/DetailTexturing.INT.udn
--- a/Engine/Documentation/Source/Engine/Rendering/Materials/HowTo/DetailTexturing/DetailTexturing.INT.udn
+++ b/Engine/Documentation/Source/Engine/Rendering/Materials/HowTo/DetailTexturing/DetailTexturing.INT.udn
@@ -175,8 +175,8 @@ This is often referred to as Distance Based Texture Blending and can be set up f
 
 1. First make a copy of the Detail_Material_2 and name it Distance_Fade. Once named, open up the Material and search in the **Palette** for the following Material Expression nodes and add them to the Material Graph.
 
-	* Absolute World Position x 1
-	* Camera Position x 1
+	* World_Position x 1
+	* Camera_Position_WS x 1
 	* Distance x 1
 	* Divide  x 1
 	* Power x 1
