#UEDocs #UEDoc: Multliple edits submitted by "master_rigel"

diff --git a/Engine/Documentation/Source/Engine/Rendering/Materials/HowTo/ColoredTransluscentShadows/ColoredTransluscentShadows.INT.udn b/Engine/Documentation/Source/Engine/Rendering/Materials/HowTo/ColoredTransluscentShadows/ColoredTransluscentShadows.INT.udn
--- a/Engine/Documentation/Source/Engine/Rendering/Materials/HowTo/ColoredTransluscentShadows/ColoredTransluscentShadows.INT.udn
+++ b/Engine/Documentation/Source/Engine/Rendering/Materials/HowTo/ColoredTransluscentShadows/ColoredTransluscentShadows.INT.udn
@@ -48,7 +48,7 @@ Notes
 1. Select the **Material Attributes** node and change the following in its **Details** panel:
 * Blend Mode: Translucent
 * (Optional) Two Sided: Checked 
-* Translucency Lighting Model: TLM_Surface
+* Translucency Lighting Model: Surface Translucency Volume
 1. Plug the **Alpha** channel of the Texture Sample into **Opacity** slot to mask parts of material.
 
 #### Custom Opacity
diff --git a/Engine/Documentation/Source/Engine/Rendering/Materials/HowTo/Human_Skin/HumanSkin.INT.udn b/Engine/Documentation/Source/Engine/Rendering/Materials/HowTo/Human_Skin/HumanSkin.INT.udn
--- a/Engine/Documentation/Source/Engine/Rendering/Materials/HowTo/Human_Skin/HumanSkin.INT.udn
+++ b/Engine/Documentation/Source/Engine/Rendering/Materials/HowTo/Human_Skin/HumanSkin.INT.udn
@@ -26,7 +26,7 @@ Getting it right is challenging but once the tools are mastered, they can be app
 Before we get to making the Material, first we should gather some reference. 
 This is because it is usually best to start with multiple photo references that have multiple lighting situations so you can verify how the material looks under different lighting conditions. 
 If at all possible, having scanned and calibrated source data would be ideal as this will help to make sure that you are getting the best possible results. 
-When viewing your content inside the editor, make sure to disable the [Editor Eye Adaptation](Engine/Rendering/PostProcessEffects/AutomaticExposure) for a more controlled environment.
+When viewing your content inside the editor, make sure to disable the **Editor Eye Adaptation** for a more controlled environment: _Edit Menu > Project Settings > Rendering > Default Post Processing Settings > Auto Exposure_
 
 [region:tip]
  A calibrated environment can be helpful in making sure that you get the colors correct but it is not required.
@@ -67,7 +67,7 @@ While this Material does look complex, we will be breaking down what each sectio
 
 ![](HS_Base_Color.png)
 
-One of the first things that need to be understood is that from a very far distance, the Subsurface Scattering effect is harder to perceive. 
+One of the first things that needs to be understood is that from a very far distance, the Subsurface Scattering effect is harder to perceive. 
 This means that it is best to have the Base Color of your texture looking right when viewed in the distance. 
 You could also use a grayscale texture to darken the base color and the specular (caused by occlusion) to give an overall better look to the skin. 
 You can also add some Vector and Scalar controls to adjust the color and power of the skin to help you further tweak to get the look you are after.
diff --git a/Engine/Documentation/Source/Engine/Rendering/Materials/HowTo/Subsurface_Profile/Subsurface_Profile.INT.udn b/Engine/Documentation/Source/Engine/Rendering/Materials/HowTo/Subsurface_Profile/Subsurface_Profile.INT.udn
--- a/Engine/Documentation/Source/Engine/Rendering/Materials/HowTo/Subsurface_Profile/Subsurface_Profile.INT.udn
+++ b/Engine/Documentation/Source/Engine/Rendering/Materials/HowTo/Subsurface_Profile/Subsurface_Profile.INT.udn
@@ -10,7 +10,7 @@ Related: Resources/ContentExamples/MaterialNodes/1_7
 [TOC(start:2)]
 
 
-The ability to rendering realistic looking human skin is a must for any modern Video Game engine. To fulfill this need, Unreal Engine 4 (UE4) now offers a shading method specifically for skin or wax like surfaces called **Subsurface Profile**. 
+The ability to render realistic looking human skin is a must for any modern Video Game engine. To fulfill this need, Unreal Engine 4 (UE4) now offers a shading method specifically for skin or wax like surfaces called **Subsurface Profile**. 
 While the Subsurface Profile shading model has similar properties to the **Subsurface** shading model, its key difference is in how it renders. 
 SubsurfaceProfile uses a rendering approach that is based in Screen Space because this helps to better display the subtle subsurface effects seen in human skin where backscattering is a secondary effect only seen in few cases like ears for example.
 The following document will cover what a Subsurface Profile is and how you can use them in your work. 
@@ -23,7 +23,7 @@ Once created, make sure that you name the Material. For this example, the Materi
 When completed, your **Content Browser** should look like this. 
 	![](SSP_Create_New_Material.png)
 
-1. Next, open up the Material by **Double-Clicking** on it in the **Content Browser** with the **Right Mouse Button**.
+1. Next, open up the Material by **Double-Clicking** on it in the **Content Browser** with the **Left Mouse Button**.
 	[REGION:lightbox]
 	![](SSP_Open_New_Material.png)
 	[/REGION]
