#DocFix #Forums

Updating information on a doc page requested from a user.

diff --git a/Engine/Documentation/Source/Engine/UI/SourceControl/SourceControl.INT.udn b/Engine/Documentation/Source/Engine/UI/SourceControl/SourceControl.INT.udn
--- a/Engine/Documentation/Source/Engine/UI/SourceControl/SourceControl.INT.udn
+++ b/Engine/Documentation/Source/Engine/UI/SourceControl/SourceControl.INT.udn
@@ -68,7 +68,7 @@ The **Content Browser** will display special icons on the **upper-right corner**
 [REGION:imagetable]
 | ![CheckedOutByYou.png](CheckedOutByYou.png)(w:178) | ![CheckedOut.png](CheckedOut.png)(w:178) | ![MarkedForAdd.png](MarkedForAdd.png)(w:178) | ![notInDepot.png](notInDepot.png)(w:178) | ![NotHadRevision.png](NotHadRevision.png)(w:178) 
 | --- | --- | --- | --- | --- |
-| Checked out by you | Checked out by another user | Marked for add | Not in depot | Not had revision
+| Checked out by you | Checked out by another user | Marked for add | Not in depot | Newer version of file exists in source control
 
 [/REGION]
 
