#UE4doc #docs:
Added engine version field
Replaced many images due to changes in UI layout/names/functionality
Rewrote many steps as the nature of the Mixamo Animation Pack has changed from being a project to being an asset pack.
Removed over half of the page's content due to redundancy.
Changed font settings for multiple terms.

diff --git a/Engine/Documentation/Source/Engine/Animation/UsingMixamoContent/UsingMixamoContent.INT.udn b/Engine/Documentation/Source/Engine/Animation/UsingMixamoContent/UsingMixamoContent.INT.udn
--- a/Engine/Documentation/Source/Engine/Animation/UsingMixamoContent/UsingMixamoContent.INT.udn
+++ b/Engine/Documentation/Source/Engine/Animation/UsingMixamoContent/UsingMixamoContent.INT.udn
@@ -2,6 +2,7 @@ Availability:Public
 Title: Using Mixamo Content
 Crumbs: %ROOT%, Engine, Engine/Animation
 Description:How to get started using the animations and skeletal meshes in the Mixamo project.
+version: 4.9
 
 The 15 characters provided in the Mixamo Animation Pack can be migrated to your project, either as a Character Blueprint or as a combination of a Skeletal Mesh and Animation Blueprint.
 The first method is a quick and easy way to drop a third-person character into your project, complete with movement controls and attacking. The second method allows you to change the 
@@ -17,113 +18,44 @@ provided Blueprint, or swapping the Skeletal Mesh and Animation Blueprint in oth
 1. Create a new  project based on the Blueprint Third Person template.
 
 	![](MixamoExampleProject.png)(w:700)
-1. Open the **MixamoAnimPack** sample project.
-	
-	![](MixamoAnimPack.png)(w:700)
-1. In the **Content Browser**, open the **MixamoAnimPack** folder. If you know what character you would like to migrate, go ahead and open that folder now. If you would like to see all the possible characters, 
-click on **Filters** in the **Content Browser** and select **Blueprint**.
-	![](filter_BP.png)
-	
-1. Locate the Character Blueprint you would like to use. This example will use the **Mixamo_Maw** character.
-	
-	![](Selecting_MawJLaygo.png)
-1. **Right-click** on the Character Blueprint, and select **Migrate** from the context menu that appears.
-
-	![](Migrate_CharacterBP.png)
-
-1. Click **OK** when the asset report appears.
-
-	![](AssetReport_CharacterBP.png)
-1. Navigate to your new project's Content folder, then click **OK**.
-
-	![](ContentFolder.png)
-1. Switch to your new project.
-1. In the **Content Browser**, navigate to **Game/Blueprints**. 
-1. **Double-click** on the **MyGame** GameMode Blueprint to open it.
-
-	![](MyGameGameMode.png)
-1. In the **Defaults** mode of the Blueprint Editor, expand the **Default Pawn Class** dropdown.
-
-	![](ExpandDefaultPawn.png)
-1. Select the name of the character you migrated over. In this example, we selected the **Mixamo_Maw** character. 
-
-	![](MixamoCharacter.png)
-
-1. Click on the **Compile** button to compile the Blueprint, then save your Blueprint. 
-
-	![](ClickCompile.png)
-1. Your new character is now ready to go. Click the **Play** button in the **Toolbar** to test it out!
-
-	![](MixamoJump.png)
-
-## General Migration 
+1. In the Marketplace tab of the Epic Games Launcher, select the **Characters and Animations** section.
 
-If your project is based on one of the following templates, you can use this more general setup. It will work for both Blueprints and C++ versions of these templates.
+	![](Marketplace_CharactersAndAnimations.png)
 
-* Side Scroller
-* Top Down
-* Third Person
+1. Scroll down to the bottom of the section to find the **Mixamo Animation Pack** and select it.
 
-[REGION:tip]
-These steps assume that your project has movement and input setups compatible with the Mixamo animations. If you have your own project, or have substantially changed a template's setup, you
-may need to add input bindings, or add Blueprints and/or C++ code to trigger the animations.
-[/REGION]
-
-1. Create a new project based on one of the above templates. If you choose a C++ template, compile it in **Development Editor** and then open it in the editor. In this example, we are using the Top Down C++ template.
-	![](MixamoExample2.png)(w:700)
-1. Open the **MixamoAnimPack** sample project.
-	
-	![](MixamoAnimPack.png)(w:700)
-1. In the **Content Browser**, open the **MixamoAnimPack** folder. If you know what character you would like to migrate, go ahead and open that folder now. You can find the Skeletal Mesh in the **Mesh** subfolder.
-
-	If you would like to see all the possible characters, click on **Filters** in the **Content Browser** and select **Skeletal Mesh**.
-	![](SkelMesh_Filter.png)
-1. **Right-click** on the Skeletal Mesh you would like to use, and select **Migrate** from the context menu that appears. In this example, we are going to use Mixamo_Zoe. 
-
-	![](Migrate_SkelMesh.png)
-
-1. Click **OK** when the asset report appears.
-
-	![](AssetReport_SkelMesh.png)
-1. Navigate to your new project's Content folder, then click **OK**.
-
-	![](Example2_Content.png)
-1. Now, in the **Content Browser**, turn off the **Skeletal Mesh** filter, and filter by **Animation Blueprint** instead. If you are already in the folder for your chosen character, you can find the Animation Blueprint in the **Anims** subfolder.
+	![](MixamoAnimPack.png)
 
-	![](AnimBP_Filter.png)
-1. Locate the **Animation Blueprint** that matches the Skeletal Mesh you just migrated. **Right-click** on it and select **Migrate** from the context menu that appears.
-	
-	![](Migrate_AnimBP.png)
-1. Click **OK** when the asset report appears.
+1. Once on the **Mixamo Animation Pack** page, select the version of the engine that you're using and then select **Add to project**.
+	![](AddToProject.png)
+	[REGION:note]
+	If you're using a version of the engine that is newer than the most recent version listed, using the most recent version should be fine in most cases.
+	[/REGION]
+1. Select the project that you wish to add the asset pack to.
 
-	![](AssetReport_AnimBP.png)
-1. Make sure your new project's Content folder is still selected, then click **OK**.
+	![](SelectProjectToAdd.png)
 
-	![](Example2_Content.png)
-1. Click **Yes All** when prompted to overwrite any duplicate assets.
-
-	![](Overwrite.png)
+	[REGION:note]
+	If you are using an anim pack from a different version than your project, you will need to select the "Show All Projects" checkbox and then select correct version before being able to add the asset pack
+	[/REGION]
 1. Switch to your new project.
-1. In the **Content Browser**, navigate to **Game/Blueprints**. 
-1. **Double-click** on the **MyCharacter** Character Blueprint to open it.
+1. In the **Content Browser**, navigate to **Game/ThirdPersonBP/Blueprints**. 
+1. **Double-click** on the **ThirdPersonGameMode** GameMode Blueprint to open it.
 
-	![](MyCharacter_BP.png)
-1. In the **Defaults** mode of the Blueprint Editor, find the **Mesh** section and the **Skeletal Mesh** property. You can also find this in **Components Mode**.
+	![](MyGameGameMode.png)
+1. After selecting the **Class Defaults** button in the Blueprint Editor, expand the **Default Pawn Class** dropdown.
 
-	![](Defaults_Mixamo.png)
-1. Select **Mixamo_Zoe** (or whichever Skeletal Mesh you migrated) as the **Skeletal Mesh**, using the dropdown menu.
+	![](ExpandDefaultPawn.png)
+1. Select the a character from the Mixamo Animation Pack. In this example, we selected the **Mixamo_Maw** character. 
 
-	![](SkelMesh_Dropdown.png)
+	![](MixamoCharacter.png)
 
-1. Now, find the **Animation** section. Expand the dropdown menu next to **Anim Blueprint Generated Class**.
+1. Click on the **Compile** button to compile the Blueprint, then save your Blueprint. 
 
-	![](AnimBP_Dropdown.png)
-1.  Select the specific Animation Blueprint you exported.
+	![](ClickCompile.png)
+1. Your new character is now ready to go. Delete the ThirdPersonCharacter that is already in the scene and then click the **Play** button in the **Toolbar** to test it out.
 
-	![](Selected_Options.png)
-1. Save and compile your Blueprint, then close it.
-1. Your new character is now ready to go. Click the **Play** button in the **Toolbar** to test it out!
+	![](MixamoJump.png)
 
-	![](TopDown_UE4.png)
 
 [/REGION:fullwidth]
