#UE4doc #docs:
Added back the second section as it is still relevant but changed many steps related to it to fit the latest version of the editor.
Added/changed multiple pictures

diff --git a/Engine/Documentation/Source/Engine/Animation/UsingMixamoContent/UsingMixamoContent.INT.udn b/Engine/Documentation/Source/Engine/Animation/UsingMixamoContent/UsingMixamoContent.INT.udn
--- a/Engine/Documentation/Source/Engine/Animation/UsingMixamoContent/UsingMixamoContent.INT.udn
+++ b/Engine/Documentation/Source/Engine/Animation/UsingMixamoContent/UsingMixamoContent.INT.udn
@@ -4,15 +4,10 @@ Crumbs: %ROOT%, Engine, Engine/Animation
 Description:How to get started using the animations and skeletal meshes in the Mixamo project.
 version: 4.9
 
-The 15 characters provided in the Mixamo Animation Pack can be migrated to your project, either as a Character Blueprint or as a combination of a Skeletal Mesh and Animation Blueprint.
+The 15 characters provided in the Mixamo Animation Pack can be added to your project, either as a Character Blueprint or as a combination of a Skeletal Mesh and Animation Blueprint.
 The first method is a quick and easy way to drop a third-person character into your project, complete with movement controls and attacking. The second method allows you to change the 
 look of a character you already have, or in different styles of templates, if the movement and inputs have similar setups to the Mixamo base character. 
 
-[REGION:note]
-If you would like to add all of the provided animations, meshes, and Blueprints to your project, you can [add the asset pack using the Launcher](Engine/Basics/AssetsAndPackages#assetpacks). The steps and notes below about changing the Game Mode to use the 
-provided Blueprint, or swapping the Skeletal Mesh and Animation Blueprint in other types of projects, still apply. You just will not need to open the MixamoAnimPack project or migrate anything.
-[/REGION]
-
 ## Third-Person Blueprint Character
 [REGION:fullwidth]
 1. Create a new  project based on the Blueprint Third Person template.
@@ -57,5 +52,58 @@ provided Blueprint, or swapping the Skeletal Mesh and Animation Blueprint in oth
 
 	![](MixamoJump.png)
 
+## General Usage
+
+If your project is based on one of the following templates, you can use this more general setup. It will work for both Blueprints and C++ versions of these templates.
+
+* Side Scroller
+* Top Down
+* Third Person
+
+[REGION:tip]
+These steps assume that your project has movement and input setups compatible with the Mixamo animations. If you have your own project, or have substantially changed a template's setup, you
+may need to add input bindings, or add Blueprints and/or C++ code to trigger the animations.
+[/REGION]
+
+1. Create a new project based on one of the above templates. If you choose a C++ template, compile it in **Development Editor** and then open it in the editor. In this example, we are using the Top Down C++ template.
+	![](MixamoExample2.png)(w:700)
+1. In the Marketplace tab of the Epic Games Launcher, select the **Characters and Animations** section.
+
+	![](Marketplace_CharactersAndAnimations.png)
+
+1. Scroll down to the bottom of the section to find the **Mixamo Animation Pack** and select it.
+
+	![](MixamoAnimPack.png)
+
+1. Once on the **Mixamo Animation Pack** page, select the version of the engine that you're using and then select **Add to project**.
+	![](AddToProject.png)
+	[REGION:note]
+	If you're using a version of the engine that is newer than the most recent version listed, using the most recent version should be fine in most cases.
+	[/REGION]
+1. Select the project that you wish to add the asset pack to.
+
+	![](SelectProjectToAdd.png)
+
+	[REGION:note]
+	If you are using an anim pack from a different version than your project, you will need to select the "Show All Projects" checkbox and then select correct version before being able to add the asset pack
+	[/REGION]
+1. Switch to your new project.
+1. In the **Content Browser**, naviate to the Game/TopDownCPP/Blueprints folder and **Double-click** the **TopDownCharacter** Blueprint.
+
+	![](TopDownCharacterBP_Select.png)
+1. Select the **Mesh** component from the **Components** panel
+	
+	![](Mesh_Select.png)
+1. In the **Skeletal Mesh** field under **Mesh**, select the Mixamo skeletal mesh of your choice. For this example, we chose Mixamo_Zoe.
+
+	![](SkeletalMesh_Select.png)
+1. In the **Anim Blueprint Generated Class** field under **Animation**, select the Animation Blueprint that matches the skeletal mesh you chose.
+
+	![](AnimBP_Select2.png)
+1. Click **Compile** and close the **Blueprint Editor**
+
+	![](ClickCompile.png)
+1. Your swapped out Skeletal Mesh is ready to go. Click **Play** and test it out.
+	![](ZoeRun.png)
 
 [/REGION:fullwidth]
