#UE4doc #docs:
Added engine version field

diff --git a/Engine/Documentation/Source/Engine/Audio/Audio.INT.udn b/Engine/Documentation/Source/Engine/Audio/Audio.INT.udn
--- a/Engine/Documentation/Source/Engine/Audio/Audio.INT.udn
+++ b/Engine/Documentation/Source/Engine/Audio/Audio.INT.udn
@@ -3,6 +3,7 @@ Title: Audio and Sound
 Crumbs: %ROOT%, Engine
 Description:System for playing in-game sounds, including the use of SoundCue node-based audio assets.
 Navigation:topic
+version: 4.9
 
 [VAR:Topic]
 [OBJECT:Topic]
