#UE4doc #docs:
Added engine version field
Changed font settings for multiple terms
Changed title of pages as 'Distance Model' is not how this is referred to anymore in the editor.

diff --git a/Engine/Documentation/Source/Engine/Audio/DistanceModelAttenuation/DistanceModelAttenuation.INT.udn b/Engine/Documentation/Source/Engine/Audio/DistanceModelAttenuation/DistanceModelAttenuation.INT.udn
--- a/Engine/Documentation/Source/Engine/Audio/DistanceModelAttenuation/DistanceModelAttenuation.INT.udn
+++ b/Engine/Documentation/Source/Engine/Audio/DistanceModelAttenuation/DistanceModelAttenuation.INT.udn
@@ -1,24 +1,24 @@
 Availability: Public
-Title:Distance Model Attenuation
+Title:Sound Attenuation
 Description:Reference for the various distance models available for controlling the attenuation of sounds.
 Crumbs: %ROOT%, Engine, Engine/Audio
-
+version: 4.9
 [TOC (start:2 end:2)]
 
 
 
-## Overview
 
-Attenuation is essentially the ability of a sound to lower in volume as the player moves away from it. It works using two radii: MinRadius and MaxRadius. As you move from the sound's origin through the MinRadius, volume of the sound is at 100%. 
+
+**Sound Attenuation** is essentially the ability of a sound to lower in volume as the player moves away from it. It works using two radii: **MinRadius** and **MaxRadius**. As you move from the sound's origin through the MinRadius, volume of the sound is at 100%. 
 As you pass between the MinRadius and the MaxRadius, the volume linearly fades between 100% and silence. 
 The rate at which this fade occurs is based on the DistanceModel property and the **Distance Algorithm** setting, which provides several types of falloff curves to control the volume in-between the radii. 
 Once you pass outside the MaxRadius, you are outside the limit of the sound and hear only silence.
 
-While the **Distance Algorithm** can be used to define the falloff, **Attenuation Shapes** can be used to specify the shape of the Attenuation Volume itself. **Sphere**, **Capsule**, **Box**, or **Cone** can be used as the **Attenuation Shape** based on your audio needs and how you want to confine the way the audio is presented. 
+While the Distance Algorithm can be used to define the falloff, **Attenuation Shapes** can be used to specify the shape of the Attenuation Volume itself. **Sphere**, **Capsule**, **Box**, or **Cone** can be used as the Attenuation Shape based on your audio needs and how you want to confine the way the audio is presented. 
 
 
 ## Distance Algorithm
-Here is a list and description of the available DistanceModel Attenuation curves.
+Here is a list and description of the available Distance Model Attenuation curves.
 
 In all of the following graphs, The Y axis is volume, the X axis is distance. MinRadius is at the origin of the X axis, and MaxRadius is located at x at the right edge of the X axis.
 
@@ -96,7 +96,7 @@ Min 900/Max 1000:
  
  
 ## Attenuation Shapes
-Examples of the **Attenuation Shapes** available can be found below along with sample use cases. 
+Examples of the Attenuation Shapes available can be found below along with sample use cases. 
 
 **Attenuation Sphere**
 
