#UE4doc #docs:
Added engine version field
Changed font settings for many terms
Replaced a few images due to updated UI

diff --git a/Engine/Documentation/Source/Engine/Audio/Overview/Overview.INT.udn b/Engine/Documentation/Source/Engine/Audio/Overview/Overview.INT.udn
--- a/Engine/Documentation/Source/Engine/Audio/Overview/Overview.INT.udn
+++ b/Engine/Documentation/Source/Engine/Audio/Overview/Overview.INT.udn
@@ -2,6 +2,7 @@ Availability: Public
 Title:Audio System Overview
 Crumbs: %ROOT%, Engine, Engine/Audio
 Description:Overview of the audio system used for playing in-game sounds, including the use of SoundCues node-based audio assets.
+version: 4.9
 
 [VAR:Topic]
 [OBJECT:Topic]
@@ -46,7 +47,7 @@ Description:Overview of the audio system used for playing in-game sounds, includ
 
 
 The audio system in Unreal Engine 4 is made up of several components, each working together to produce the audio experience for players. 
-When you import an audio file into the engine and drop it in a level, you will have several options such as the basic **Volume** or **Pitch** levels to adjust, as well as more fine-tuning settings such as **Attenuation**, which defines how a sound is heard based on your distance from its origin. 
+When you import an audio file into the engine and drop it in a level, you will have several options such as the basic **Volume** or **Pitch** levels to adjust, as well as more fine-tuning settings such as **Sound Attenuation**, which defines how a sound is heard based on your distance from its origin. 
 
 Unreal Engine 4 also allows for building composite sounds in the form of **Sound Cues** and the **Sound Cue Editor** which enable you to combine sounds as well as apply modifiers called **Sound Nodes** to alter the final output. 
 There are additional elements that are used to define how a sound is heard or played which are covered on this page and links to more documentation can be found in their respective sections. 
@@ -62,7 +63,7 @@ Unreal Engine 4 currently supports importing uncompressed little endian 16 bit w
 | Bit Rate | 16 |
 | Speaker Channels | Mono, Stereo, 2.1, 4.1, 5.1 6.1, 7.1 |
 
-Importing a sound file into the editor generates a **Sound Wave** asset that can be dropped directly into a level or can be used to create a **Sound Cue** and edited inside the **Sound Cue Editor**. 
+Importing a sound file into the editor generates a **Sound Wave** asset that can be dropped directly into a level or can be used to create a Sound Cue and edited inside the **Sound Cue Editor**. 
 
 Importing a sound file can be handled in two ways: 
 
@@ -84,41 +85,39 @@ For more information on importing and multichannel usage, refer to the [Audio Fi
 
 ## Sound Asset Types
 
-To add a **Sound Asset**, click the ![](CB_button_NewAsset.png) button from the **Content Browser** and select **Sounds** followed by the asset you wish to add from the menu. 
+To add a Sound asset, click the ![](CB_button_NewAsset.png) button from the **Content Browser** and select **Sounds** followed by the asset you wish to add from the menu. 
 
 ![](selectSoundAsset.png)(h:400)
 
-There are several different types of **Sound Assets** that can be added to your projects which are described below. 
+There are several different types of Sound assets that can be added to your projects which are described below. 
 
 ### Sound Cue
 ![](soundCueEditor.png)(w:640)
 
-**Sound Cues** are composite sounds that allow you to modify the behavior of audio playback, combine audio effects, and apply audio modifiers with **Sound Nodes** to alter the final output. 
+Sound Cues are composite sounds that allow you to modify the behavior of audio playback, combine audio effects, and apply audio modifiers with **Sound Nodes** to alter the final output. 
 
 For more information, refer to the [Sound Cue Editor](Engine/Audio/SoundCues/Editor) page. 
 
 ### Sound Attenuation
-**Sound Attenuation** assets allow the definition of attenuation properties in a reusable manner. Any place you can specify one-time use attenuation properties, you can instead specify the **Sound Attenuation** asset. This allows adjustment to attenuation properties without having to revisit every sound individually.
+Sound Attenuation assets allow the definition of attenuation properties in a reusable manner. Any place you can specify one-time use attenuation properties, you can instead specify the Sound Attenuation asset. This allows adjustment to attenuation properties without having to revisit every sound individually.
 
 ![](soundAttenuation.png)
 
-For more information on **Attenuation** properties, see the **Attenuation Properties** section of the [Ambient Sound Actor User Guide](Engine\Audio\SoundActors) page. 
-
-For more information on **Distance Model Attenuation**, see the [Distance Model Attenuation](Engine/Audio/DistanceModelAttenuation) page.
+For more information on **Attenuation**, see the [Sound Attenuation](Engine/Audio/DistanceModelAttenuation) page.
 
 
 ### Reverb Effects
 
 **Reverb Effects** are a definable asset with several properties that can be easily adjusted and applied to any **Audio Volume** placed in your level. 
-With a **Reverb Effect**, you can adjust settings (pictured below) that will allow you to control elements like the echo density, overall reverb gain, air absorption, and more to help craft the overall feel you may be after.  
+With a Reverb Effect, you can adjust settings (pictured below) that will allow you to control elements like the echo density, overall reverb gain, air absorption, and more to help craft the overall feel you may be after.  
 
 ![](reverbeffect.png)
 
 
 ### Sound Class
-**Sound Classes** are a collection of properties that can be applied to a number of **Sound Assets**.
+**Sound Classes** are a collection of properties that can be applied to a number of Sound assets.
 
-The properties inside a **Sound Class** act as multipliers to the existing values and will be carried out by all **Sound Assets** assigned to the **Sound Class**.  
+The properties inside a Sound Class act as multipliers to the existing values and will be carried out by all Sound assets assigned to the Sound Class.  
 
 ![](soundClassProperties.png)
 
@@ -126,43 +125,43 @@ Hierarchies can be created by adding **Child Classes**, which will allow you to
 
  ![](soundClassEditor.png)(w:640)
 
-You can also add **Passive Sound Mixes** (see the **Sound Mix** section below) to a **Sound Class** which will kick in and activate automatically whenever the **Sound Class** is played (for example, having music automatically lower whenever a dialogue **Sound Class** is played).  
+You can also add **Passive Sound Mixes** (see the **Sound Mix** section below) to a Sound Class which will kick in and activate automatically whenever the Sound Class is played (for example, having music automatically lower whenever a dialogue Sound Class is played).  
 
 
 ### Sound Mix
-**Sound Mixes** allow you to set the **EQ settings** (Equalizer Settings) and modify **Volume** and **Pitch** properties of **Sound Classes**.
+Sound Mixes allow you to set the **EQ Settings** (Equalizer Settings) and modify **Volume** and **Pitch** properties of Sound Classes.
 
 ![](soundMixDetails.png)
 
-Multiple **Sound Mixes** can be active at the same time, all contributing to the overall audio effect. You can **Push** (Activate) or **Pop** (Deactivate) **Sound Mixes** directly inside a **Blueprint** with the **Push Sound Mix Modifier** and **Pop Sound Mix Modifier** nodes or activate them passively whenever a sound with a given **Sound Class** is playing within a specified threshold.
+Multiple Sound Mixes can be active at the same time, all contributing to the overall audio effect. You can **Push** (Activate) or **Pop** (Deactivate) Sound Mixes directly inside a Blueprint with the **Push Sound Mix Modifier** and **Pop Sound Mix Modifier** nodes or activate them passively whenever a sound with a given Sound Class is playing within a specified threshold.
 
-Within the **Sound Mix** asset itself, which can be opened by **Double-Clicking** the asset in the **Content Browser**, several properties exist. 
-You can specify **EQ Settings** for the mix to adjust the high, middle, and low frequencies and gains. As the **EQ Settings** of multiple **Sound Mixes** cannot be combined, the **EQ Priority** allows you to control which active mix's properties are applied at any given time.
+Within the Sound Mix asset itself, which can be opened by **Double-Clicking** the asset in the **Content Browser**, several properties exist. 
+You can specify EQ Settings for the mix to adjust the high, middle, and low frequencies and gains. As the EQ Settings of multiple Sound Mixes cannot be combined, the **EQ Priority** allows you to control which active mix's properties are applied at any given time.
 
-Inside the **Sound Classes** section, you set which **Sound Classes** are to be affected by the mix.  For each **Sound Class** you can set the **Volume** or **Pitch** adjusters, set if the mix settings are to be applied to **Child Classes**, or modify the **VoiceCenterChannelVolume**.
+Inside the Sound Classes section, you set which Sound Classes are to be affected by the mix.  For each Sound Class you can set the **Volume** or **Pitch** adjusters, set if the mix settings are to be applied to Child Classes, or modify the **VoiceCenterChannelVolume**.
 
-The **Sound Mix** section allows you to specify how the **Sound Mix** properties are applied or removed.  **Delay** indicates how long before the mix properties should begin being applied.  **Fade in Time** and **Fade out Time** specify how quickly to transition from no effect to the specified properties.  **Duration** allows a pushed **Sound Mix** to automatically pop itself after the specified duration.  A value of -1 indicates to never automatically pop and passively applied **Sound Mixes** will not automatically pop.
+The Sound Mix section allows you to specify how the Sound Mix properties are applied or removed.  **Delay** indicates how long before the mix properties should begin being applied.  **Fade in Time** and **Fade out Time** specify how quickly to transition from no effect to the specified properties.  **Duration** allows a pushed Sound Mix to automatically pop itself after the specified duration.  A value of -1 indicates to never automatically pop and passively applied Sound Mixes will not automatically pop.
 
 
 ### Dialogue Voice & Dialogue Wave
 The **Dialogue Voice** and **Dialogue Wave** assets are used for generating in-game dialogue events, crafting subtitles, and for supplementing localization efforts. 
 
-When editing a newly created **Dialogue Voice** asset, you can define the **Gender** and **Plurality** of a voice actor and although you do not specify any audio assets inside the **Dialogue Voice**, the information provided here can be referenced inside the **Dialogue Wave**. 
+When editing a newly created Dialogue Voice asset, you can define the **Gender** and **Plurality** of a voice actor and although you do not specify any audio assets inside the Dialogue Voice, the information provided here can be referenced inside the Dialogue Wave. 
 
-The **Dialogue Wave** provides more options and is where the connection between audio and speaker/listener(s) is made. It is also where the correlation between dialogue audio and subtitle text is made. The **Dialogue Wave** represents a single line of dialogue and the core component of the **Dialogue Wave** settings is the **Dialogue Contexts** section. 
+The Dialogue Wave provides more options and is where the connection between audio and speaker/listener(s) is made. It is also where the correlation between dialogue audio and subtitle text is made. The Dialogue Wave represents a single line of dialogue and the core component of the Dialogue Wave settings is the **Dialogue Contexts** section. 
 
 ![](dialogueWaveDetails.png)
 
-Inside the **Dialogue Contexts** section, you can specify the **Speaker** or who the dialogue is being **Directed At** in their respective sections by assigning a **Dialogue Voice**. 
-The actual audio line of dialogue can be added as a **Sound Wave** by expanding the context window and choosing the desired asset from the drop-down menu or by pointing to an asset in the **Content Browser**. 
+Inside the **Dialogue Contexts** section, you can specify the **Speaker** or who the dialogue is being **Directed At** in their respective sections by assigning a Dialogue Voice. 
+The actual audio line of dialogue can be added as a Sound Wave by expanding the context window and choosing the desired asset from the drop-down menu or by pointing to an asset in the **Content Browser**. 
 
 In the event that you have multiple actors who say the same line of dialogue, the **Add Dialogue Context** option will allow you to create a new entry for the dialogue where you can set new **Speaker** and **Directed At** sections.
 
 [REGION:note]
-**Dialogue Wave** assets can also be applied to a **Sound Cue** by using the **Dialogue Player** node inside the **Sound Cue Editor**. Also, you can access a **Dialogue Wave** directly from within a **Blueprint** using the **Play Dialogue at Location** and **Play Dialogue Attached** nodes. 
+Dialogue Wave assets can also be applied to a Sound Cue by using the Dialogue Player node inside the **Sound Cue Editor**. Also, you can access a Dialogue Wave directly from within a Blueprint using the **Play Dialogue at Location** and **Play Dialogue Attached** nodes. 
 [/REGION]
 
-In addition to **Dialogue Contexts**, you can apply a **Mature** filter which flags the dialogue as containing mature/adult content. Under **Script**, you can enter the text that is spoken in the attached audio inside the **Spoken Text** section. 
+In addition to Dialogue Contexts, you can apply a **Mature** filter which flags the dialogue as containing mature/adult content. Under **Script**, you can enter the text that is spoken in the attached audio inside the **Spoken Text** section. 
 You can also enter contextual information for the audio for translation purposes or for notes to a voice actor in the **Voice Actor Direction** section. 
 
 
@@ -173,7 +172,7 @@ You can also enter contextual information for the audio for translation purposes
 Regarding loudest potential volume, there is variation in the overall bandwidth to work within. 
 For example, a stereo file at 1.0 volume will be 2x louder than a mono file, likewise (4) mono files would be 4x louder. However, eventually you will hit the overall threshold and the output will begin to clip and distort.
 
-On any given **Sound Cue**, volume settings up to around 2.0 will increase the perceived loudness of the audio file, and anything above that would not. 
+On any given Sound Cue, volume settings up to around 2.0 will increase the perceived loudness of the audio file, and anything above that would not. 
 A single cue will never distort, but you would not want all of your files at maximum volume because it will likely overload when multiple cues play simultaneously in-game.
 
 You may want to consider coming up with a consistent volume scheme, or at least some general guidelines for your volumes:
@@ -193,11 +192,11 @@ Additionally, you may consider using mono assets almost everywhere to maintain c
 
 The process usually consists of: 
 
-* Identifying a **Sound Class** that should cause ducking (e.g. dialog).
-* When a sound from that **Sound Class** is triggered, other **Sound Classes** begin to decrease in volume (**Sound Mix** : **Fade in Time** = 0.3 seconds).
-* Specify the amount other **Sound Classes** decrease in volume (**Sound Mix** : **Volume Adjuster** = 0.4).
-* When a sound from a ducking **Sound Class** stops, other classes increase in volume back to normal volume (**Sound Mix** : **Fade Out Time** = 0.3 seconds).
-* You may also want **Sound Class** exceptions to the ducking process (e.g. music), or a **Sound Class** called _Exceptions_ that is not affected.
+* Identifying a Sound Class that should cause ducking (e.g. dialog).
+* When a sound from that Sound Class is triggered, other Sound Classes begin to decrease in volume (Sound Mix : Fade in Time = 0.3 seconds).
+* Specify the amount other Sound Classes decrease in volume (Sound Mix : Volume Adjuster = 0.4).
+* When a sound from a ducking Sound Class stops, other classes increase in volume back to normal volume (Sound Mix : Fade Out Time = 0.3 seconds).
+* You may also want Sound Class exceptions to the ducking process (e.g. music), or a Sound Class called Exceptions that is not affected.
 
 Using a similar approach will ensure that emphasis is placed on the important aspects of your audio and that it is not "drowned" out by superficial or un-important audio elements. 
 
