#UE4doc #docs:
Added engine version field
Removed Overview header
Changed font settings for many terms
Replaced a picture due to outdated UI representation

diff --git a/Engine/Documentation/Source/Engine/Audio/SoundActors/SoundActors.INT.udn b/Engine/Documentation/Source/Engine/Audio/SoundActors/SoundActors.INT.udn
--- a/Engine/Documentation/Source/Engine/Audio/SoundActors/SoundActors.INT.udn
+++ b/Engine/Documentation/Source/Engine/Audio/SoundActors/SoundActors.INT.udn
@@ -2,14 +2,13 @@ Availability: Public
 Title:Ambient Sound Actor User Guide
 Crumbs: %ROOT%, Engine, Engine/Audio
 Description:Guide to using the Ambient Sound Actor for creating ambient and localized sounds in-game.
-
+version: 4.9
 [TOC(start:2 end:2)]
 
 
 
-## Overview
 
-Unreal Engine 4 streamlines the process by which you can produce or modify ambient sounds through the use of the **Ambient Sound Actor**. When a **Sound Wave** or **Sound Cue** asset is placed in a level, an **Ambient Sound Actor** with that sound asset is created. The Ambient Sound Actor has several properties which are covered on this page that will allow you to modify how players receive the sound.
+Unreal Engine 4 streamlines the process by which you can produce or modify ambient sounds through the use of the **Ambient Sound Actor**. When a **Sound Wave** or **Sound Cue** asset is placed in a level, an Ambient Sound Actor with that sound asset is created. The Ambient Sound Actor has several properties which are covered on this page that will allow you to modify how players receive the sound.
 
 For more information on importing files as [Sound Waves](Engine/Audio/WAV) or modifying files through [Sound Cues](Engine/Audio/SoundCues/Editor), refer to their respective pages for additional documentation. 
 
@@ -26,42 +25,42 @@ For more information on importing files as [Sound Waves](Engine/Audio/WAV) or mo
 	[/PARAM]
 [/OBJECT]
 
-The **Ambient Sound Actor** (icons pictured at left) can be used for many purposes such as ambient looping sounds as well as non-looping sounds. 
-Generally, the **Ambient Sound Actor** conforms to the real world where the closer you are to a sound, the louder it will appear. By comparison, a sound that is normally loud, may appear soft if it is further away. 
+The Ambient Sound Actor (icons pictured at left) can be used for many purposes such as ambient looping sounds as well as non-looping sounds. 
+Generally, the Ambient Sound Actor conforms to the real world where the closer you are to a sound, the louder it will appear. By comparison, a sound that is normally loud, may appear soft if it is further away. 
 
-If the **Ambient Sound Actor** is set to **Auto Activate**, it will immediately begin to play once it is created (at the beginning of play or on spawn), even if the player is not in a position to hear it.
+If the Ambient Sound Actor is set to **Auto Activate**, it will immediately begin to play once it is created (at the beginning of play or on spawn), even if the player is not in a position to hear it.
 
-The sound asset the **Ambient Sound Actor** points to will only play once per trigger unless specified as **Looping** in the **Sound Wave** or defined as part of a **Sound Cue** asset. 
+The sound asset the Ambient Sound Actor points to will only play once per trigger unless specified as **Looping** in the Sound Wave or defined as part of a Sound Cue asset. 
 
-You can add an **Ambient Sound Actor** to your level in several ways:
+You can add an Ambient Sound Actor to your level in several ways:
 
-The first is by selecting the **Ambient Sound Actor** from the **Modes** panel under **All Classes** and dragging and dropping it into your level as seen below. 
+The first is by selecting the Ambient Sound Actor from the **Modes** panel under **All Classes** and dragging and dropping it into your level as seen below. 
 
 ![](ambient_modesmenu.png)(w:640)
 
-You can also highlight a **Sound Wave** or **Sound Cue** in the **Content Browser** and in your level, **right-click** to bring up the context menu and select **Place Actor**. 
+You can also highlight a Sound Wave or Sound Cue in the **Content Browser** and in your level, **Right-click** to bring up the context menu and select **Place Actor**. 
 
 ![](ambient_contextmenu.png)(h:540)
 
-Lastly, you can simply **left-click** and drag a **Sound Wave** or **Sound Cue** from the **Content Browser** into your level.
+Lastly, you can simply **Left-click** and drag a Sound Wave or Sound Cue from the **Content Browser** into your level.
 
 ### Sound Properties
 
-An **Ambient Sound Actor** has no function without an associated sound asset. 
+An Ambient Sound Actor has no function without an associated sound asset. 
 You can assign a sound asset from the **Details** panel by selecting an asset from the **Sound** setting's drop-down menu or by highlighting a sound asset in the **Content Browser** and clicking the ![](soundactor_cb_arrow.png) button.
 
  ![](assign_soundasset.png)
 
-Once a **Sound Cue** asset has been assigned as the associated sound, the **Edit** option will become available so that you can open the **Sound Cue** asset up for editing inside the **Sound Cue Editor**. 
-Optionally, instead of assigning a sound asset, you can select the **New** option which will create a new **Sound Cue** asset and will open it in the **Sound Cue Editor** for editing. 
+Once a Sound Cue asset has been assigned as the associated sound, the **Edit** option will become available so that you can open the Sound Cue asset up for editing inside the **Sound Cue Editor**. 
+Optionally, instead of assigning a sound asset, you can select the **New** option which will create a new Sound Cue asset and will open it in the **Sound Cue Editor** for editing. 
 
-The **Play** and **Stop** options allow you to preview the assigned sound inside the editor and will **Play** and **Stop** playing the sound respectively. 
+The **Play** and **Stop** options allow you to preview the assigned sound inside the editor and will play and stop playing the sound respectively. 
 
 Additional options inside the **Sound** section of the **Details** panel are defined below. 
 
 |Property| Description |
 | ------- | ------------ |
-| **Sound** | Points to a **Sound Wave** or **Sound Cue** asset. |
+| **Sound** | Points to a Sound Wave or Sound Cue asset. |
 | **Is UI Sound** | Whether or not the sound asset plays when the game is paused. |
 | **Volume Multiplier** | Modifier applied to set the overall Volume of the sound. |
 | **Pitch Multiplier** | Modifier applied to set the overall Pitch of the sound. |
@@ -74,7 +73,7 @@ Attenuation is the ability of a sound to decrease in volume as the player moves
 As you move from the sound's origin through the MinRadius, the volume of the sound is at 100%. As you move within the space between the MinRadius and the MaxRadius, the volume level is adjusted based on the **Distance Algorithm**. 
 Outside of the MaxRadius, the volume of the sound is at 0%.
 
-For additional information on Attenuation, please refer to the [Distance Model Attenuation](Engine/Audio/DistanceModelAttenuation) page. 
+For additional information on Attenuation, please refer to the [](Engine/Audio/DistanceModelAttenuation) page. 
 
 The options under the **Attenuation** section that can be modified inside the **Details** panel are defined below: 
 
@@ -99,7 +98,7 @@ When **Override Attenuation** is checked, the ability to set additional Attenuat
 | **Spatialize** | Enable the source to be positioned in 3D.  |
 | **Distance Algorithm** | The type of volume versus distance algorithm to use for the attenuation model. |
 | **Attenuation Shape** | The shape of the attenuation volume. |
-| **Radius** | The overall size of the volume. <BR> By default, the volume shape is set to **Sphere**. This option will change based on the **Attenuation Shape** selected. |
+| **Radius** | The overall size of the volume. <BR> By default, the volume shape is set to **Sphere**. This option will change based on the Attenuation Shape selected. |
 | **Falloff Distance** | The distance over which falloff occurs. |
 | **Non-Spatialized Radius** | At what distance to start treating the sound source as spatialized. |
 | **Attenuate with LPF** | Enables attenuation via Low Pass Filter. |
@@ -107,8 +106,8 @@ When **Override Attenuation** is checked, the ability to set additional Attenuat
 | **LPFRadius Max** | The range at which to apply the maximum amount of Low Pass Filter. |
 
 [REGION:tip]
-To allow for more control over the Attenuation volume, Unreal Engine 4 allows you to set the **Attenuation Shape** to four different volume types: **Sphere**, **Capsule**, **Box**, or **Cone**. 
-In general, a **Sphere** volume is a reasonable choice, but there may be instances when you want finer control over the way the volume fits inside your levels. For example, filling a square room with a **Box** volume makes more sense. 
+To allow for more control over the Attenuation volume, Unreal Engine 4 allows you to set the Attenuation Shape to four different volume types: Sphere, **Capsule**, **Box**, or **Cone**. 
+In general, a Sphere volume is a reasonable choice, but there may be instances when you want finer control over the way the volume fits inside your levels. For example, filling a square room with a Box volume makes more sense. 
 [/REGION]
 
 ### Modulation Properties
@@ -127,11 +126,11 @@ The **Modulation** settings allow you to control the min/max modulation for both
 
 ## Adding Audio Components To Blueprints
 
-Inside a **Blueprint**, in **Component Mode**, select the **Add Component** button and select the **Audio** component. 
+Inside a **Blueprint**, in the **Component** panel, select the **Add Component** button and select the **Audio** component. 
 
 ![](soundactor_componentstab.png)
 
-The **Audio** component can be referenced through **Blueprints** and many of its settings can be altered by using different types of functions.
+The **Audio** component can be referenced through Blueprints and many of its settings can be altered by using different types of functions.
 
 ![](soundactor_audiocomponent.png)(w:480)
 
@@ -146,10 +145,10 @@ The **Audio** component does have a unique event which can be used. This **OnAud
 
 
 ## Audio Volumes
-For an added level of control, you can include **Audio Volumes** in conjunction with **Ambient Sound Actors**. For more information on how to use an **Audio Volume**, refer to the Audio Volumes section of the [Volumes](Engine/Actors/Volumes) page. 
+For an added level of control, you can include **Audio Volumes** in conjunction with **Ambient Sound Actors**. For more information on how to use an Audio Volume, refer to the Audio Volumes section of the [Volumes](Engine/Actors/Volumes) page. 
 
 [REGION:tip]
-Making changes to an **Audio Volume** will not work in realtime inside the editor. You will need to rebuild the geometry of the level that contains the edited volume in order for your changes to take effect. 
+Making changes to an Audio Volume will not work in realtime inside the editor. You will need to rebuild the geometry of the level that contains the edited volume in order for your changes to take effect. 
 [/REGION]
 
 
