#UE4doc #docs:
Added engine version field
Removed Overview header
Changed font settings for multiple terms

diff --git a/Engine/Documentation/Source/Engine/Audio/SoundCues/Editor/SoundCueEditor.INT.udn b/Engine/Documentation/Source/Engine/Audio/SoundCues/Editor/SoundCueEditor.INT.udn
--- a/Engine/Documentation/Source/Engine/Audio/SoundCues/Editor/SoundCueEditor.INT.udn
+++ b/Engine/Documentation/Source/Engine/Audio/SoundCues/Editor/SoundCueEditor.INT.udn
@@ -2,7 +2,7 @@ Availability: Public
 Title:Sound Cue Editor
 Crumbs: %ROOT%, Engine, Engine/Audio
 Description:Guide to using the tools available for setting up and editing SoundCue node-based audio assets.
-
+version: 4.9
 
 [VAR:Topic]
 [OBJECT:Topic]
@@ -49,7 +49,6 @@ Description:Guide to using the tools available for setting up and editing SoundC
 
 
 
-## Overview
 
 The behavior of audio playback in the Unreal Engine is defined within Sound Cues. The Sound Cue Editor is a node-based editor that is used to work with audio.
 
@@ -66,16 +65,16 @@ The volume and pitch settings are used to manage relative Sound Cue volumes. Thi
 
 ### Sound Wave
 
-A Sound Wave is the representation of an imported audio file in the Sound Cue Editor. You can select a Sound Wave asset in the **Content Browser** and then add it to the Sound Cue you are editing by **right-clicking** anywhere in the Sound Cue Editor, and then choosing the Sound Wave in the **From Selected** category on the context menu.
+A Sound Wave is the representation of an imported audio file in the Sound Cue Editor. You can select a Sound Wave asset in the **Content Browser** and then add it to the Sound Cue you are editing by **Right-clicking** anywhere in the Sound Cue Editor, and then choosing the Sound Wave in the **From Selected** category on the context menu.
 
-Because other Sound Cues may use the same Sound Wave, you cannot modify the volume and pitch values of a Sound Wave in the Sound Cue Editor. However, you can modify a Sound Wave's properties when you import it, or by opening it in the **Content Browser** by **double-clicking** it.
+Because other Sound Cues may use the same Sound Wave, you cannot modify the volume and pitch values of a Sound Wave in the Sound Cue Editor. However, you can modify a Sound Wave's properties when you import it, or by opening it in the **Content Browser** by **Double-clicking** it.
 
 
 ## Importing Sounds
 
 To import sounds in the Unreal Editor, do the following:
-1. In the **Content Browser**, click the ![New Asset button](button_CB_NewAsset.png) button or **right-click** in an empty space in the **Content Browser** to bring up the New Asset menu.
-1. Under **Import Assets**, click **Import to _folder name_**.
+1. In the **Content Browser**, click the ![New Asset button](button_CB_NewAsset.png) button or **Right-click** in an empty space in the **Content Browser** to bring up the New Asset menu.
+1. Under **Import Assets**, click **Import to "_folder name_"**.
 1. In the **Import** dialog box, browse to the .wav file that you want to import, select it, and click **Open**. This will import the .wave file as a Sound Cue asset. 
 
 [REGION:note]
@@ -86,14 +85,14 @@ At this time, you can only import WAV files as sound files.
 ## Creating Sound Cues
 
 To create a Sound Cue, do the following:
-1. In the **Content Browser**, click the ![New Asset button](button_CB_NewAsset.png) button or **right-click** in an empty space in the **Content Browser** to bring up the New Asset menu.
+1. In the **Content Browser**, click the ![New Asset button](button_CB_NewAsset.png) button or **Right-click** in an empty space in the **Content Browser** to bring up the New Asset menu.
 1. Under **Other Assets**, click **Sounds**, and then click **Sound Cue**.
 1. Type a name for your new Sound Cue.
 
 
 ## Opening the Sound Cue Editor
 
-You can open the Sound Cue Editor by **right-clicking** a Sound Cue asset in the **Content Browser** and then clicking **Edit** on the context menu that appears, or by **double-clicking** a Sound Cue asset.
+You can open the Sound Cue Editor by **Right-clicking** a Sound Cue asset in the **Content Browser** and then clicking **Edit** on the context menu that appears, or by **Double-clicking** a Sound Cue asset.
 
 
  
@@ -115,7 +114,7 @@ You can add an audio node by dragging it into the graph from the **Palette**.
 
 ![](SoundCueEditor_Looping3.png)
 
-You can also add a node by **right-clicking** in the graph and choosing a node from the context menu that appears.
+You can also add a node by **Right-clicking** in the graph and choosing a node from the context menu that appears.
 
 ![](SoundCueEditor_ContextMenu.png)
 
