#UE4doc #docs:
Fixed capitalization on a term

diff --git a/Engine/Documentation/Source/Shared/Editors/Common/EditorMenuItems/EditorMenuItems.INT.udn b/Engine/Documentation/Source/Shared/Editors/Common/EditorMenuItems/EditorMenuItems.INT.udn
--- a/Engine/Documentation/Source/Shared/Editors/Common/EditorMenuItems/EditorMenuItems.INT.udn
+++ b/Engine/Documentation/Source/Shared/Editors/Common/EditorMenuItems/EditorMenuItems.INT.udn
@@ -129,6 +129,6 @@ Crumbs:
 
 [EXCERPT:UITabTips]
 [REGION:tip]
-You can close any panel by clicking the small "X" in the upper-right corner of the tab. You can also hide any panel by **right-clicking** on the tab, and then clicking **Hide Tab** on the context menu that appears. To once again display a panel that you have closed, click that panel's name on the **Window** menu.
+You can close any panel by clicking the small "X" in the upper-right corner of the tab. You can also hide any panel by **Right-clicking** on the tab, and then clicking **Hide Tab** on the context menu that appears. To once again display a panel that you have closed, click that panel's name on the **Window** menu.
 [/REGION]
 [/EXCERPT:UITabTips]
