#UE4doc #docs:
Added engine version field
Removed Overview header
Changed font settings for multiple terms

diff --git a/Engine/Documentation/Source/Engine/Audio/SoundCues/Editor/Interface/SoundCueEditorInterface.INT.udn b/Engine/Documentation/Source/Engine/Audio/SoundCues/Editor/Interface/SoundCueEditorInterface.INT.udn
--- a/Engine/Documentation/Source/Engine/Audio/SoundCues/Editor/Interface/SoundCueEditorInterface.INT.udn
+++ b/Engine/Documentation/Source/Engine/Audio/SoundCues/Editor/Interface/SoundCueEditorInterface.INT.udn
@@ -2,12 +2,11 @@ Availability: Public
 Title:Sound Cue Editor UI
 Crumbs: %ROOT%, Engine, Engine/Audio, Engine/Audio/SoundCues/Editor
 Description:Guide to using the tools available for setting up and editing SoundCue node-based audio assets.
-
+version: 4.9
 [TOC(start:2)]
 
 
 
-## Overview
 
 The Sound Cue Editor is divided into the following regions:
 
@@ -85,5 +84,5 @@ For information about the basics of using this panel, see the [Editor Details Pa
 The **Palette** lists all the available nodes for the Audio Node Graph. To add a node, you can drag it into the graph from the **Palette**. For information about other ways of adding sound nodes, see [Audio Node Graph](Engine/Audio\SoundCues\Editor#AudioNodeGraph).
 
 [REGION:tip]
-You can add a comment to the Audio Node Graph either by dragging **Add Comment** into the graph, or by dragging a box around nodes with the **left mouse button**, **right-clicking** a selected node, and then clicking **Create Comment from Selection** on the context menu that appears.
+You can add a comment to the Audio Node Graph either by dragging **Add Comment** into the graph, or by dragging a box around nodes with the **Left Mouse Button**, **Right-clicking** a selected node, and then clicking **Create Comment from Selection** on the context menu that appears.
 [/REGION]
