#UE4doc #docs:
Didn't submit all changes on last submit

diff --git a/Engine/Documentation/Source/Engine/Audio/SoundCues/Editor/Interface/SoundCueEditorInterface.INT.udn b/Engine/Documentation/Source/Engine/Audio/SoundCues/Editor/Interface/SoundCueEditorInterface.INT.udn
--- a/Engine/Documentation/Source/Engine/Audio/SoundCues/Editor/Interface/SoundCueEditorInterface.INT.udn
+++ b/Engine/Documentation/Source/Engine/Audio/SoundCues/Editor/Interface/SoundCueEditorInterface.INT.udn
@@ -8,7 +8,7 @@ version: 4.9
 
 
 
-The Sound Cue Editor is divided into the following regions:
+The **Sound Cue Editor** is divided into the following regions:
 
 ![soundcueeditor.png](soundcueeditor.png)(w:900)
 
@@ -20,7 +20,7 @@ The Sound Cue Editor is divided into the following regions:
 
 [Include:Shared/Editors/Common/EditorMenuItems#UITabTips]
 
-For more information about the Sound Cue Editor, see [](Engine/Audio\SoundCues\Editor).
+For more information about the **Sound Cue Editor**, see [](Engine/Audio\SoundCues\Editor).
 
 ## Menu Bar
 
