#UE4doc #docs - Added Engine Version tag. Added Skill Level tag.

diff --git a/Engine/Documentation/Source/Engine/Levels/Levels.INT.udn b/Engine/Documentation/Source/Engine/Levels/Levels.INT.udn
--- a/Engine/Documentation/Source/Engine/Levels/Levels.INT.udn
+++ b/Engine/Documentation/Source/Engine/Levels/Levels.INT.udn
@@ -6,6 +6,9 @@ Related: Engine/UI/LevelEditor
 Related: Engine/Actors
 Related: Engine/LevelStreaming
 Related: Resources/ContentExamples/LevelDesign
+SkillLevel: Beginner
+Version: 4.9
+
 
 [VAR:Topic]
 [OBJECT:Topic]
