#UE4doc #docs - Added LODColoration view mode added.

diff --git a/Engine/Documentation/Source/Engine/LevelStreaming/WorldBrowser/WorldBrowser.INT.udn b/Engine/Documentation/Source/Engine/LevelStreaming/WorldBrowser/WorldBrowser.INT.udn
--- a/Engine/Documentation/Source/Engine/LevelStreaming/WorldBrowser/WorldBrowser.INT.udn
+++ b/Engine/Documentation/Source/Engine/LevelStreaming/WorldBrowser/WorldBrowser.INT.udn
@@ -2,6 +2,8 @@ Availability:Public
 Title:World Composition User Guide
 Crumbs: %ROOT%, Engine, Engine/LevelStreaming
 Description:System for managing large worlds including origin shifting technology.
+SkillLevel: Intermediate
+Version: 4.9
 
 [TOC (start:2 end:3)]
 
diff --git a/Engine/Documentation/Source/Engine/UI/LevelEditor/Viewports/ViewModes/ViewModes.INT.udn b/Engine/Documentation/Source/Engine/UI/LevelEditor/Viewports/ViewModes/ViewModes.INT.udn
--- a/Engine/Documentation/Source/Engine/UI/LevelEditor/Viewports/ViewModes/ViewModes.INT.udn
+++ b/Engine/Documentation/Source/Engine/UI/LevelEditor/Viewports/ViewModes/ViewModes.INT.udn
@@ -52,7 +52,6 @@ Description:Explanations for the available View Modes within the viewports.
 
 [TOC(start:2 end:2)]
 
-## Overview
 
 
 [EXCERPT:Intro]
@@ -219,7 +218,35 @@ Skeletal Meshes will appear in light brown and are not considered in this calcul
 
 **Reflections** view mode overrides all materials with a flat normal and a roughness of 0, which is a mirror. This is useful for diagnosing the detail of reflections and allowing you to place more Reflection Capture Actors in areas where you need more detail.
 
+## LOD Coloration
 
+[REGION:fullwidth]
+![](ViewMode_LODColoration.png)
+[/REGION]
+
+* Console command: `viewmode LODColoration`
+
+**LOD Coloration** view mode displays the current LOD index of a primitive. This is useful for diagnosing any LOD issues or to see at which distance your LODs are switching. 
+
+|LOD Primitive Coloration|||||||||
+|:---|:---:|:---:|:---:|:---:|:---:|:---:|:---:|:---:|
+|**Color**|![](LODColoration_0.png)|![](LODColoration_1.png)|![](LODColoration_2.png)|![](LODColoration_3.png)|![](LODColoration_4.png)|![](LODColoration_5.png)|![](LODColoration_6.png)|![](LODColoration_7.png)|
+|**LOD Primitive Color**|**0**|**1**|**2**|**3**|**4**|**5**|**6**|**7**|
+
+[REGION:code]
++LODColorationColors=(R=1.0,G=1.0,B=1.0,A=1.0)     
++LODColorationColors=(R=1.0,G=0.0,B=0.0,A=1.0)     
++LODColorationColors=(R=0.0,G=1.0,B=0.0,A=1.0)     
++LODColorationColors=(R=0.0,G=0.0,B=1.0,A=1.0)     
++LODColorationColors=(R=1.0,G=1.0,B=0.0,A=1.0)     
++LODColorationColors=(R=1.0,G=0.0,B=1.0,A=1.0)     
++LODColorationColors=(R=0.0,G=1.0,B=1.0,A=1.0)     
++LODColorationColors=(R=0.5,G=0.0,B=0.5,A=1.0)     
+[/REGION]
+
+[Region: note]
+By default the engine uses only four LODs, but this can be extended in the source code. 
+[/Region]
 
 ## Buffer Visualization
 
