#UE4doc #docs:
Added engine version field
Removed Overview header
Changed font settings for multiple terms

diff --git a/Engine/Documentation/Source/Engine/Audio/SoundCues/NodeReference/SoundCueReference.INT.udn b/Engine/Documentation/Source/Engine/Audio/SoundCues/NodeReference/SoundCueReference.INT.udn
--- a/Engine/Documentation/Source/Engine/Audio/SoundCues/NodeReference/SoundCueReference.INT.udn
+++ b/Engine/Documentation/Source/Engine/Audio/SoundCues/NodeReference/SoundCueReference.INT.udn
@@ -2,21 +2,20 @@ Availability: Public
 Title: Sound Cue Reference
 Crumbs: %ROOT%, Engine, Engine/Audio
 Description:Reference describing the nodes available for use in SoundCues and their properties.
-
+version: 4.9
 [TOC (start:2)]
 
-## Overview
 
-The Unreal Engine provides various sound nodes whose properties can be modified to create interesting and complex Sound Cues in the Sound Cue Editor. You can find these sound nodes in the **Palette** panel.
+The Unreal Engine provides various sound nodes whose properties can be modified to create interesting and complex Sound Cues in the **Sound Cue Editor**. You can find these sound nodes in the **Palette** panel.
 
 ![Sound node list](SoundCueEditor_NodeList.png)
 
 This page describes the sound nodes available for use, along with their properties.
 
-For more information about the Sound Cue Editor, see [](Engine/Audio\SoundCues\Editor). For more information about the Sound Cue Editor UI, see [](Engine/Audio\SoundCues\Editor\Interface).
+For more information about the **Sound Cue Editor**, see [](Engine/Audio\SoundCues\Editor). For more information about the **Sound Cue Editor** UI, see [](Engine/Audio\SoundCues\Editor\Interface).
 
 [REGION:tip]
-For some nodes, you can add additional inputs by **right-clicking** the node and then clicking **Add Input**.
+For some nodes, you can add additional inputs by **Right-clicking** the node and then clicking **Add Input**.
 [/REGION]
 
 ## Attenuation Node
@@ -24,7 +23,7 @@ For some nodes, you can add additional inputs by **right-clicking** the node and
 The **Attenuation** node is used to control spatialization, attenuation, and radii properties.
 
 [REGION:note]
-If you want to override the Sound Cue's attenuation settings by using the Attenuation node, the Attenuation node's **Override Attenuation** check box must be set to _True_ in the Sound Cue Editor, and the Sound Cue asset's **Override Attenuation** check box must be set to _False_ in the Level Editor's **Details** panel.
+If you want to override the Sound Cue's attenuation settings by using the Attenuation node, the Attenuation node's **Override Attenuation** check box must be set to True in the **Sound Cue Editor**, and the Sound Cue asset's **Override Attenuation** check box must be set to False in the Level Editor's **Details** panel.
 [/REGION]
 
 ![Attenuation Node](SoundCueEditor_Attenuation.png)
@@ -33,20 +32,20 @@ If you want to override the Sound Cue's attenuation settings by using the Attenu
 | --- | --- |
 | [REGION:tablesection]Attenuation[/REGION] ||
 | **Attenuation Settings** | The attenuation settings copied from a Sound Attenuation asset. Use this property if you want to use pre-existing attenuation settings. |
-| **Attenuate** | If _true_, fades the input sound out over distance. |
-| **Spatialize** | If _true_, positions the input sound in 3D space. |
+| **Attenuate** | If true, fades the input sound out over distance. |
+| **Spatialize** | If true, positions the input sound in 3D space. |
 | **Distance Algorithm** | The interpolation method to use when attenuating the input sound. See [Distance Model Attenuation Curves](Engine/Audio/DistanceModelAttenuation) for descriptions of the various types of distance algorithms. |
 | **Attenuation Shape** | The shape of the attenuation area from a drop-down list. |
 | **Radius** | The attenuation radius. |
 | **Falloff Distance** | The falloff distance for the attenuation. |
 | **Non-Spatialized Radius** | The distance from the sound's origin where the sound source is treated as spatialized. |
-| **Attenuate With LPF** | If _true_, enables attenuation via low pass filter (LPF). |
+| **Attenuate With LPF** | If true, enables attenuation via low pass filter (LPF). |
 | **LPFRadius Min** | The distance from the sound's origin where the low pass filter should begin to be applied. |
 | **LPFRadius Max** | The distance from the sound's origin where the maximum amount of low pass filter should be applied. |
-| **Override Attenuation** | If _true_, overrides existing attenuation settings. |
+| **Override Attenuation** | If true, overrides existing attenuation settings. |
 
 
-Stereo files play evenly on the Left and Right channels and attenuate allowing volume fades between the **Radius min** and **Radius Max** values, but they are not spatialized. Sound Cues that do not contain the Attenuation node are not localized and play consistently at the Sound Cue volume setting.
+Stereo files play evenly on the Left and Right channels and attenuate allowing volume fades between the **Radius Min** and **Radius Max** values, but they are not spatialized. Sound Cues that do not contain the Attenuation node are not localized and play consistently at the Sound Cue volume setting.
 
 ## Branch Node
 
@@ -69,7 +68,7 @@ The **Concatenator** node is used to play sounds sequentially (e.g. creating a s
 | Property | Description |
 | -------- | ----------- |
 | [REGION:tablesection]Sound Node Concatenator[/REGION] ||
-| **Input Volume** | Allows for individual volume settings for each sound asset attached to the **Concatenator**. |
+| **Input Volume** | Allows for individual volume settings for each sound asset attached to the Concatenator. |
 
 ## Continuous Modulator Node
 
@@ -102,7 +101,7 @@ The pitch and volume of the input sound default to a uniform distribution, but y
 | **Param Mode** | [INCLUDE:Engine/Audio/SoundCues/NodeReference/#ParamMode] |
 
 [REGION:note]
-**Min** and **Max** fields are provided for randomization. 
+Min and Max fields are provided for randomization. 
 [/REGION]
 
 
@@ -110,7 +109,7 @@ The pitch and volume of the input sound default to a uniform distribution, but y
 [EXCERPT:ParamMode]
 The mode with which the parameter is applied: 
 * **MPM_Normal** - Clamps input value to the range (MinInput, MaxInput) then maps to the range (MinOutput, MaxOutput). 
-* **MPM_Abs** - Same as **MPM_Normal**, except that the input value is treated as an absolute value.
+* **MPM_Abs** - Same as MPM_Normal, except that the input value is treated as an absolute value.
 * **MPM_Direct** - Uses the input value directly, without scaling or reference to Min or Max input or output values.
 [/EXCERPT:ParamMode]
 --->
@@ -119,7 +118,7 @@ The mode with which the parameter is applied:
 
 ## Crossfade by Distance Node
 
-The **Crossfade by Distance** node provides the ability to fade between two or more input sounds based on the distance from the Sound Cue's origin. By default, there are two inputs on this node, but you can add more inputs by **right-clicking** on the Crossfade by Distance node and clicking **Add Input** on the context menu. Sound nodes can be connected directly to the Crossfade by Distance node inputs but you can also add nodes between them for independent control per layer.
+The **Crossfade by Distance** node provides the ability to fade between two or more input sounds based on the distance from the Sound Cue's origin. By default, there are two inputs on this node, but you can add more inputs by **Right-clicking** on the Crossfade by Distance node and clicking **Add Input** on the context menu. Sound nodes can be connected directly to the Crossfade by Distance node inputs but you can also add nodes between them for independent control per layer.
 
 ![Crossfade by Distance Node](SoundCueEditor_CrossfadeDistance.png)
 
@@ -177,7 +176,7 @@ The **DialoguePlayer** node sets up a Dialogue Wave sound asset to be played und
 | [REGION:tablesection]Dialogue Player[/REGION] ||
 | **Dialogue Wave** | The Dialogue Wave that the node will play. |
 | **Context** | Who is speaking, and to whom. |
-| **Looping** | If _true_, plays the dialogue on a loop. |
+| **Looping** | If true, plays the dialogue on a loop. |
 
 ## Doppler Node
 
@@ -207,7 +206,7 @@ You can use the properties in the **Modulation** category to randomize predefine
 | **Loop End** | The time in seconds where the envelope's loop ends. |
 | **Duration After Loop** | The time in seconds it takes the envelope to fade out after the last loop is completed. |
 | **Loop Count** | The number of times the envelope should loop if looping is enabled and the envelope is not set to loop indefinitely. |
-| **Loop Indefinitely** | If enabled, the envelope will continue to loop indefinitely regardless of the **Loop Count** value. |
+| **Loop Indefinitely** | If enabled, the envelope will continue to loop indefinitely regardless of the Loop Count value. |
 | **Loop** | If enabled, the envelope will loop using the loop settings. |
 | [REGION:tablesection]Envelope[/REGION] ||
 | **Volume Curve** | The distribution defining the volume envelope.  |
@@ -254,7 +253,7 @@ The **Looping** node is used to set a Sound Wave to looping. You can loop multip
 ![Looping Node](SoundCueEditor_LoopingNode.png)
 
 [REGION:warning]
-The Looping node should only be used for logical or procedural looping such as introducing a delay. These sounds will not be played seamlessly. If you want a sound to loop seamlessly and indefinitely, use the **Looping** flag on the Wave Player node for that sound.
+The Looping node should only be used for logical or procedural looping such as introducing a delay. These sounds will not be played seamlessly. If you want a sound to loop seamlessly and indefinitely, use the Looping flag on the Wave Player node for that sound.
 [/REGION]
 
 ## Mature Node
@@ -311,13 +310,13 @@ The **Oscillator** node is used to add continuous pitch and volume oscillation o
 | **Center Max** | The maximum value around which the center oscillates. |
 
 [REGION:note]
-**Min** and **Max** fields are provided for randomization. 
+Min and Max fields are provided for randomization. 
 [/REGION]
 
 
 ## Random Node
 
-The **Random** node is used to randomly trigger a Sound Wave node from within a group of possible Sound Wave nodes. Weight controls the probability a Sound Wave node will be triggered relative to other Sound Waves nodes in the Sound Cue. The check box **RandomWithoutReplacement** will exhaust the entire list of possible sound nodes before repetition. Inputs are added for each audio file by **right-clicking** on the Random node and selecting **Add Input**. Sound Wave nodes can be connected directly to a Random node, but you can also add nodes between them for additional control.
+The **Random** node is used to randomly trigger a Sound Wave node from within a group of possible Sound Wave nodes. Weight controls the probability a Sound Wave node will be triggered relative to other Sound Waves nodes in the Sound Cue. The check box **RandomWithoutReplacement** will exhaust the entire list of possible sound nodes before repetition. Inputs are added for each audio file by **Right-clicking** on the Random node and selecting **Add Input**. Sound Wave nodes can be connected directly to a Random node, but you can also add nodes between them for additional control.
 
 
 ![Random Node](SoundCueEditor_Random.png)
@@ -327,7 +326,7 @@ The **Random** node is used to randomly trigger a Sound Wave node from within a
 | [REGION:tablesection]Sound Node Random[/REGION] ||
 | **Weights** | The list of weights for each input sound determining the odds that a particular sound will be chosen. |
 | **Preselect at Level Load** | A random Sound Wave node, selected when the level is loaded, that is always used during that play session. Note that this setting does not apply during Play In Editor or Simulate; it applies only in game builds outside the Unreal Editor. |
-| **Randomize Without Replacement** | If _true_, prevents an input Sound Wave node from being played again until all input Sound Wave nodes have been played. This ensures that all sounds are played before any are repeated. |
+| **Randomize Without Replacement** | If true, prevents an input Sound Wave node from being played again until all input Sound Wave nodes have been played. This ensures that all sounds are played before any are repeated. |
 
 <!-- get more info -->
 
@@ -392,7 +391,7 @@ The **Wave Player** node contains a Sound Wave asset from the **Content Browser*
 | -------- | ----------- |
 | [REGION:tablesection]Wave Player[/REGION] ||
 | **Sound Wave** | The Sound Wave asset to use in this node. |
-| [INCLUDE:Engine/Audio/SoundCues/NodeReference/#TableSpacing] **Looping** | If _true_, plays the Sound Wave on a loop. | 
+| [INCLUDE:Engine/Audio/SoundCues/NodeReference/#TableSpacing] **Looping** | If true, plays the Sound Wave on a loop. | 
 
 
 [EXCERPT:TableSpacing]
@@ -413,6 +412,6 @@ This Sound Cue changes the Sound Classes of the Sound Waves.
 
 ![Example 3](SoundCueEditor_Example3.png)(w:900)
 
-This Sound Cue plays one sound if a parameter is set to _True_, and another if it is set to _False_.
+This Sound Cue plays one sound if a parameter is set to True, and another if it is set to False.
 
 
