#UE4doc #docs:
Added engine version field

diff --git a/Engine/Documentation/Source/Engine/Audio/WAV/WAV.INT.udn b/Engine/Documentation/Source/Engine/Audio/WAV/WAV.INT.udn
--- a/Engine/Documentation/Source/Engine/Audio/WAV/WAV.INT.udn
+++ b/Engine/Documentation/Source/Engine/Audio/WAV/WAV.INT.udn
@@ -2,6 +2,7 @@ Availability: Public
 Title: Audio Files
 Crumbs: %ROOT%, Engine, Engine/Audio
 Description:The basic requirements for getting sound files into Unreal Engine 4.
+version: 4.9
 
 [VAR:Topic]
 [OBJECT:Topic]
