#UE4doc #docs:
Added engine version field
Removed Overview header
Changed font settings for multiple terms
Expanded on a step to explain further

diff --git a/Engine/Documentation/Source/Engine/Basics/AssetsAndPackages/AssetsAndPackages.INT.udn b/Engine/Documentation/Source/Engine/Basics/AssetsAndPackages/AssetsAndPackages.INT.udn
--- a/Engine/Documentation/Source/Engine/Basics/AssetsAndPackages/AssetsAndPackages.INT.udn
+++ b/Engine/Documentation/Source/Engine/Basics/AssetsAndPackages/AssetsAndPackages.INT.udn
@@ -2,6 +2,7 @@ Availability:Public
 Title: Assets and Packages
 Crumbs: %ROOT%, Engine
 Description:Overview of the asset and package system used in Unreal Engine 4.
+version: 4.9
 
 [VAR:Topic]
 [OBJECT:Topic]
@@ -25,7 +26,6 @@ Description:Overview of the asset and package system used in Unreal Engine 4.
 
 [TOC (start:2 end:3)]
 
-## Overview
 
 When a workflow includes package files that each contain several assets, source-control contention can be an issue. A workflow that instead supports and encourages creating many individual asset files means that each asset can be checked out individually. This reduces source-control contention and bottlenecks in workflow.
 
@@ -53,9 +53,9 @@ For more information on using the **Content Browser** to move, rename, add, and
 
 ### Asset Naming
 
-You name assets at the time that you create or import them. Your asset's path is set automatically by the folder in which you create or import it, which is directly associated with the asset's location on disk. For example, a path of `Content/Characters/MyCharacter` would save to `/UE4/MyProject/Content/Characters/MyCharacter.MyCharacter.uasset`.
+You name assets at the time that you create or import them. Your asset's path is set automatically by the folder in which you create or import it, which is directly associated with the asset's location on disk. For example, a path of Content/Characters/MyCharacter would save to /UE4/MyProject/Content/Characters/MyCharacter.MyCharacter.uasset.
 
-References to assets contain the path in the reference, for example: `Blueprint'/MyProject/Content/Characters/MyCharacter.MyCharacter'`.
+References to assets contain the path in the reference, for example: Blueprint'/MyProject/Content/Characters/MyCharacter.MyCharacter'.
 
 ### Copying Assets to Other Projects
 
@@ -69,7 +69,7 @@ If you have downloaded asset packs through the Launcher, you can add the entire
 
 To add an asset pack to your project:
 
-1. Find the asset pack in your Vault.
+1. Find the asset pack in your **Vault** which is located at the bottom of the **Learn** tab.
 	![](Vault.png)
 1. Click on **Add to Project**.
 1. Select your project in the list that appears.
