#UE4doc #docs:
Resubmitting as not all changes were included with last submission

diff --git a/Engine/Documentation/Source/Engine/Basics/AssetsAndPackages/Saving/Saving.INT.udn b/Engine/Documentation/Source/Engine/Basics/AssetsAndPackages/Saving/Saving.INT.udn
--- a/Engine/Documentation/Source/Engine/Basics/AssetsAndPackages/Saving/Saving.INT.udn
+++ b/Engine/Documentation/Source/Engine/Basics/AssetsAndPackages/Saving/Saving.INT.udn
@@ -1,13 +1,13 @@
-Availability: Public
+Availability: docs
 Title:Package Saving Procedure
 Crumbs: %ROOT%, Engine, Engine/Basics/AssetsAndPackages
 Description:Guide to saving content packages inside of Unreal Editor.
+version: 4.9
 
 [TOC (start:2 end:3)]
 
 
 
-## Overview
 
 There are two dialogs that can appear when you save a package.  The dialogs allow you to easily control the packages saved by the editor.  This document describes the dialogs in the order that they appear.
 
