#UE4doc #docs:
Added engine version field

diff --git a/Engine/Documentation/Source/Engine/Basics/CoordinateSpace/CoordinateSpace.INT.udn b/Engine/Documentation/Source/Engine/Basics/CoordinateSpace/CoordinateSpace.INT.udn
--- a/Engine/Documentation/Source/Engine/Basics/CoordinateSpace/CoordinateSpace.INT.udn
+++ b/Engine/Documentation/Source/Engine/Basics/CoordinateSpace/CoordinateSpace.INT.udn
@@ -2,6 +2,7 @@ Availability:Public
 Title:Coordinate Space Terminology
 Crumbs:%ROOT%, Programming, Programming/Rendering
 Description:Explanations of the various coordinate spaces used in Unreal Engine.
+version: 4.9
 
 ## Spaces
 
