#UE4doc #docs:
Added engine version field
Changed font settings for multiple terms

diff --git a/Engine/Documentation/Source/Engine/Basics/DerivedDataCache/DerivedDataCache.INT.udn b/Engine/Documentation/Source/Engine/Basics/DerivedDataCache/DerivedDataCache.INT.udn
--- a/Engine/Documentation/Source/Engine/Basics/DerivedDataCache/DerivedDataCache.INT.udn
+++ b/Engine/Documentation/Source/Engine/Basics/DerivedDataCache/DerivedDataCache.INT.udn
@@ -2,11 +2,12 @@ Availability:Public
 Title:Derived Data Cache
 Crumbs:%ROOT%, Engine
 Description: Overview of the Derived Data Cache and the storage of assets in formats used by UE4 and its target platforms. 
+version: 4.9
 
 The **Derived Data Cache** (DDC) stores versions of assets in the formats used by the engine and its 
 target platforms, as opposed to the source formats artists create that are imported into the editor 
-and stored in `.uasset` files. Content stored in the DDC is _disposable_ in that it can always be 
-regenerated at any time using the data stored in the `.uasset` file. Storing these _derived_ formats 
+and stored in .uasset files. Content stored in the DDC is disposable in that it can always be 
+regenerated at any time using the data stored in the .uasset file. Storing these derived formats 
 externally makes it possible to easily add or change the formats used by the engine without needing to 
 modify the source asset file.
 
@@ -23,8 +24,8 @@ not do any harm, but it is a waste of time as it generally takes longer to trans
 in the DDC than it would to simply generate it from scratch locally.
 [/REGION]
 
-To set up a shared DDC, override the paths for the `[DerivedDataBackendGraph]` by declaring this 
-section in your game's `DefaultEngine.ini`. This section is originally declared in the `BaseEngine.ini` 
+To set up a shared DDC, override the paths for the [DerivedDataBackendGraph] by declaring this 
+section in your game's **DefaultEngine.ini**. This section is originally declared in the **BaseEngine.ini** 
 where the paths are set to Epic's internal DDC share. Re-declare the section and change the paths to point 
 to a share on your network (e.g. `Path=\\mystudio.net\DDC`):
 
@@ -41,7 +42,7 @@ to a share on your network (e.g. `Path=\\mystudio.net\DDC`):
 	EnginePak=(Type=ReadPak, Filename=../../../Engine/DerivedDataCache/DDC.ddp)
 
 
-Satellite studios working out of a single shared code base can set the `UE-SharedDataCachePath` environment 
+Satellite studios working out of a single shared code base can set the **UE-SharedDataCachePath** environment 
 variable to a path that all users at each location can read and write to. This allows each location to 
 have its own shared DDC. 
 
@@ -72,7 +73,7 @@ You can fill your derived data caches at any time by passing the arguments below
 	UE4\Engine\Binaries\Win64\UE4Editor.exe MyGame -run=DerivedDataCache -fill 
 
 We do this on a nightly basis to ensure the DDC is always primed, but it is not necessary in general 
-as the automatic caching should _just work_.
+as the automatic caching should just work.
 
 ## Distributing with a DDC
  
@@ -81,10 +82,10 @@ DDCs can be packaged for distribution if the need arises.
 
 **To package a DDC:**
  
-1. Run the `UE4Editor.exe` from the `UE4/Engine/Binaries/Win64` directory passing it the arguments shown below:
+1. Run the **UE4Editor.exe** from the UE4/Engine/Binaries/Win64 directory passing it the arguments shown below:
 
 		UE4Editor.exe [GameName] -run=DerivedDataCache -fill -DDC=CreatePak
  
-1. This creates a `DDC.ddp` file in the `UE4\[GameName]\DerivedDataCache`
+1. This creates a **DDC.ddp** file in the UE4\[GameName]\DerivedDataCache
 
-1. The engine automatically detects and uses the `.ddp` file.
+1. The engine automatically detects and uses the .ddp file.
