#UE4doc #docs:
Added engine version field
Changed font settings for multiple terms

diff --git a/Engine/Documentation/Source/Engine/Basics/DirectoryStructure/DirectoryStructure.INT.udn b/Engine/Documentation/Source/Engine/Basics/DirectoryStructure/DirectoryStructure.INT.udn
--- a/Engine/Documentation/Source/Engine/Basics/DirectoryStructure/DirectoryStructure.INT.udn
+++ b/Engine/Documentation/Source/Engine/Basics/DirectoryStructure/DirectoryStructure.INT.udn
@@ -2,6 +2,7 @@ Availability:Public
 Title:Directory Structure
 Crumbs: %ROOT%, Engine
 Description:Overview of the directories that make up the engine and game projects.
+version: 4.9
 
 [TOC(start:2)]
 
@@ -24,8 +25,8 @@ Some subdirectories are common amongst both the Engine and game project director
  * **Config** - Configuration files for setting values that control engine behavior. Values set in the game project Config files override the values set in the Engine\Config directory.
  * **Content** - Holds content for the engine or game, including asset packages and maps.
  * **DerivedDataCache** - Contains derived data files generated on-load for referenced content. Not having cache files present for referenced content can increase load times dramatically.
- * **Intermediate** - Contains temporary files generated during building the engine or game. In game directories, **Shaders** are stored in the **Intermediate** directory.
- * **Saved** - Contains autosaves, configuration (`*.ini`) files, and log files. Additionally, the **Engine > Saved** directory contains crash logs, hardware information, and Swarm options and data. 
+ * **Intermediate** - Contains temporary files generated during building the engine or game. In game directories, Shaders are stored in the Intermediate directory.
+ * **Saved** - Contains autosaves, configuration (.ini) files, and log files. Additionally, the Engine > Saved directory contains crash logs, hardware information, and Swarm options and data. 
  * **Source** - Contains all of the source files for engine or game, including the engine source code, tools, gameplay classes, etc.
     * **Engine** - Source files in the Engine directory are categorized into the following:
         * **Developer** - Files used by both the editor and engine.
@@ -48,7 +49,7 @@ Some subdirectories are specific to the Engine directory.
  * **Extras** - Additional helper and utility files.
  * **Plugins** - Contains plugins used in the engine.
  * **Programs** - Contains configuration files and log files for projects stored in the UE4 root directory and other Unreal programs such as UnrealFrontend and UnrealHeaderTool.
- * **Shaders** - Holds the shader source files (_.usf_) for the engine.
+ * **Shaders** - Holds the shader source files (.usf) for the engine.
 
 
 ## Game Project Directories
@@ -67,9 +68,9 @@ Some subdirectories are specific to the Engine directory.
 [EXCERPT:Source]
 [REGION:simpletable]
 | --- | --- |
-| Classes | Contains game object class definitions (`.h` files). |
-| Private | Contains private game object class implementation files (`.cpp` files). |
-| Public | Contains public game object class implementation files (`.cpp` files). |
+| Classes | Contains game object class definitions (.h files). |
+| Private | Contains private game object class implementation files (.cpp files). |
+| Public | Contains public game object class implementation files (.cpp files). |
 [/REGION]
 [/EXCERPT:Source]
 -->
@@ -78,11 +79,11 @@ Some subdirectories are specific to the Engine directory.
  
 | Directory | Description |
 | --------- | ----------- |
-| **Classes** | Contains game object class definitions (`.h` files). |
+| **Classes** | Contains game object class definitions (.h files). |
 | **Config** | Default project settings for your game. |
 | **External dependencies** | Displays the public engine header files (only visible in Visual Studio). |
-| **Private** | Contains private game object class implementation files (`.cpp` files). |
-| **Public** | Contains public game object class implementation files (`.cpp` files). |
+| **Private** | Contains private game object class implementation files (.cpp files). |
+| **Public** | Contains public game object class implementation files (.cpp files). |
 
 
 [/EXCERPT:Main]
