#UE4doc #docs:
Added engine version field
Changed font settings for a few terms

diff --git a/Engine/Documentation/Source/Engine/Basics/Distributions/Distributions.INT.udn b/Engine/Documentation/Source/Engine/Basics/Distributions/Distributions.INT.udn
--- a/Engine/Documentation/Source/Engine/Basics/Distributions/Distributions.INT.udn
+++ b/Engine/Documentation/Source/Engine/Basics/Distributions/Distributions.INT.udn
@@ -2,14 +2,14 @@ Availability: Public
 Title: Distributions
 Crumbs: %ROOT%, Engine
 Description:Data types that provide flexibility by allowing for constant values, random values within a range, values interpolated along a curve, and values driven by parameters.
-
+version: 4.9
 
 [TOC (start:2 end:3)]
 
 
 
 
-Distributions are a group of data types that provide flexibility by allowing for constant values, random values within a range, values interpolated along a curve, and values driven by parameters. These are generally used in properties of Particle Systems and nodes within SoundCues.
+**Distributions** are a group of data types that provide flexibility by allowing for constant values, random values within a range, values interpolated along a curve, and values driven by parameters. These are generally used in properties of Particle Systems and nodes within SoundCues.
 
 [REGION:note]
 	Although regions are listed in alphabetical order within the editor, we will list them here in order of complexity. This is done simply to aid in understanding for those new to the concept.
@@ -17,7 +17,7 @@ Distributions are a group of data types that provide flexibility by allowing for
 
 ## Distribution Baking
 
-No matter which type of distribution you use, you will always see the option _Can be Baked_ just beneath it. This determines whether the numeric results of the distribution can be baked out into a lookup table. Baking values out to a lookup table is far more efficient than having to calculate random values or curve interpolations on the fly. For this reason, this option is on by default.
+No matter which type of distribution you use, you will always see the option **Can be Baked** just beneath it. This determines whether the numeric results of the distribution can be baked out into a lookup table. Baking values out to a lookup table is far more efficient than having to calculate random values or curve interpolations on the fly. For this reason, this option is on by default.
 
 Should you run into edge cases in which the lookup table results lose critical information, you may uncheck this option so that the system checks the actual value on each tick. Naturally, this is a bit less efficient, but in rare cases where you have a critical look to achieve and the lookup tables are failing, the option is worth considering.
 
@@ -25,7 +25,7 @@ Should you run into edge cases in which the lookup table results lose critical i
 
 ## Distribution Types
 
-Most particle properties utilize the DistributionFloat and DistributionVector types to allow for the maximum flexibility when creating particle systems. There are a variety of options available when using distribution types which will be detailed in the following sections.
+Most particle properties utilize the **DistributionFloat** and **DistributionVector** types to allow for the maximum flexibility when creating particle systems. There are a variety of options available when using distribution types which will be detailed in the following sections.
 
 Note that the actual usage of the value will vary, depending on the property that is using it. For example, the Initial Size of a particle is determined at the time that a particle spawns. If the distribution type is set to a constant curve for the property, then the emitter time at the spawning will be used to retrieve the desired value from the curve.
 
@@ -41,7 +41,7 @@ This type is used to supply a value for the property that is constant. Below is
 
 ![FloatConstant_Graph.png](FloatConstant_Graph.png)
 
-**1.** Constant float value which will not change.
+1. Constant float value which will not change.
 
 
 
