#UE4 Docs: Adding a new section in the Lighting Modes that covers the 2 Sided Lighting Model for Foliage.

diff --git a/Engine/Documentation/Source/Engine/Rendering/Materials/MaterialProperties/LightingModels/LightingModels.INT.udn b/Engine/Documentation/Source/Engine/Rendering/Materials/MaterialProperties/LightingModels/LightingModels.INT.udn
--- a/Engine/Documentation/Source/Engine/Rendering/Materials/MaterialProperties/LightingModels/LightingModels.INT.udn
+++ b/Engine/Documentation/Source/Engine/Rendering/Materials/MaterialProperties/LightingModels/LightingModels.INT.udn
@@ -2,6 +2,8 @@ Availability: Public
 Title:Shading Models
 Crumbs: %ROOT%, Engine, Engine/Rendering/Materials, Engine/Rendering/Materials/IntroductionToMaterials,  Engine/Rendering/Materials/MaterialProperties
 Description:This document breaks down the Material Shading Models property options.
+Version: 4.9
+SkillLevel:Intermediate
 
 [TOC(start:2 end:2)]
 
@@ -137,3 +139,35 @@ When using the Subsurface Profile Shading Model, you are given access to the fol
 * Normal
 * World Position Offset
 * Ambient Occlusion
+
+## Two Sided Foliage
+
+The Two Sided Foliage shading model was introduced to give foliage a more natural look when being lit. 
+Normally when Foliage is being lit the sides facing away from the light source will be very dark to almost black. 
+The Two Sided Foliage shading model helps to eliminate this issue by simulating a very cheap Subsurface scattering to give the illusion that the foliage is receiving more light on the polygons that are facing away from the light source than it actually is. 
+
+[OBJECT:ComparisonSlider]
+	[PARAM:before]
+	![Two Sided Lighting Enabled](2_Sided_Enabled.png)
+	[/PARAM]
+	[PARAM:after]
+	![Two Sided Lighting Disabled](2_Sided_Disabled.png)
+	[/PARAM]
+[/OBJECT]
+
+When using the Two Sided Foliage Model, you are give access to the following inputs.
+* Base Color
+* Specular
+* Roughness
+* Emissive Color
+* Opacity
+* Normal
+* World Position Offset
+* Subsurface Color
+* Ambient Occlusion
+
+
+
+
+
+
