#UE4doc #docs - Updated information in page to reflect change in functionality. Updated images. Addresses UEDOC-1750

diff --git a/Engine/Documentation/Source/Engine/UI/LevelEditor/Viewports/ActorLocking/ActorLocking.INT.udn b/Engine/Documentation/Source/Engine/UI/LevelEditor/Viewports/ActorLocking/ActorLocking.INT.udn
--- a/Engine/Documentation/Source/Engine/UI/LevelEditor/Viewports/ActorLocking/ActorLocking.INT.udn
+++ b/Engine/Documentation/Source/Engine/UI/LevelEditor/Viewports/ActorLocking/ActorLocking.INT.udn
@@ -1,7 +1,8 @@
 Availability:Public
-Title:Locking Viewports to Actors
+Title:Pilot Actors in the Viewport
 Crumbs: %ROOT%, Engine, Engine/UI, Engine/UI/LevelEditor/Viewports
-Description:A look at how to lock Actors to the viewport camera for intuitive placement.
+Description:A look at how to Pilot/Lock Actors to the viewport camera for intuitive placement.
+Version: 4.9
 
 [VAR:Topic]
 [OBJECT:Topic]
@@ -47,32 +48,53 @@ Description:A look at how to lock Actors to the viewport camera for intuitive pl
 
 
 
-Actor Locking allows for the perspective Viewport camera to be used as a placement mechanism for Actors within your levels. Simply lock to a given Actor, position the view where you want that Actor to be placed, and then unlock. The position and orientation of the camera will be used to position the given Actor.
+Actor Piloting allows for the perspective Viewport camera to be used as a placement mechanism for Actors within your levels. Simply lock to a given Actor, position the view where you want that Actor to be placed, and then unlock. The position and orientation of the camera will be used to position the given Actor.
 
-## Locking the Viewport to an Actor
 
-To lock the view to an Actor, click the Actor Locking option in the Viewport Options menu, which is located in the upper-left corner of the Viewport. 
+## Piloting an Actor in the Viewport
 
-![](ActorLocking.png)(w:780)
+To Pilot a Actor in the Viewport, right-click on the Actor and choose the option from the context menu for **Pilot**. 
 
-Locking is especially useful for Camera Actors and Lights. As such, all cameras and lights will be immediately available in the sub-menu for Actor Locking. However, locking also works with any selected Actor. With an Actor selected, you will see it available in the Actor Locking sub-menu.
+![](PilotActorViewport.png)(w:700)
 
-While locking is active, you will see a padlock icon in the upper left of the Viewport. This indicates that you have some Actor attached to the view.
+Piloting is especially useful for Camera Actors and Lights. 
 
-![](LockIcon.png)
+While piloting is active, you will see the piloting text and icons in the upper left of the Viewport. This indicates that you have some Actor attached to the view.
 
-Once locked, you will look down the X-Axis of the given object. You may then move the view as normal and the Actor will come along automatically. 
+![](PilotViewportOptions.png)
 
+Once piloting, you will look down the X-Axis of the given object. You may then move the view as normal and the Actor will come along automatically. 
+
+![](PilotActive.png)(w:700)
 
-![](SpotlightLocked.png)(w:780)
 [REGION:caption]
 In this image, we have locked a SpotLight Actor to the view. Now, wherever we move the view, the light will follow as if attached to the Viewport's camera.
 [/REGION]
 
-Once you are finished positioning the Actor by way of the view, you can unlock it by choosing the Unlock option in the Actor Locking sub-menu:
+### Camera Actors
+
+When using a Camera Actor there is one additional feature that allows you to see exactly what the camera sees. This can be toggled on or off when **Pilot** mode is active.
+
+| Camera View | Button | Description |
+|:---:|:---:|:---:|
+| **Enabled** | ![](CameraViewEnabled.png) | This mode will use any settings specific to the Camera Actor. |
+| **Disabled** | ![](CameraViewDisabled.png) | This will only use the basic view from the viewport. |
+
+
+[OBJECT:ComparisonSlider]
+	[PARAM:before]
+	![Camera View Enabled](CameraPilotingViewEnabled.png)(w:400)
+	[/PARAM]
+	[PARAM:after]
+	![Camera View Disabled](CameraPilotingViewDisabled.png)(w:400)
+	[/PARAM]
+[/OBJECT]
+
+
+Once you are finished positioning the Actor by way of the view, you can stop piloting it by choosing the option in the upper left of the Viewport:
 
-![](Unlocking.png)(w:780)
+![](StopPiloting.png)
 
 You can now continue working and the Actor will maintain the position previously set by the view.
 
-![](LockComplete.png)(w:780)
+![](PilotComplete.png)(w:500)
\ No newline at end of file
