#UE4doc #docs - Added Engine Version Field

diff --git a/Engine/Documentation/Source/Engine/Rendering/LightingAndShadows/LightFunctions/LightFunctions.INT.udn b/Engine/Documentation/Source/Engine/Rendering/LightingAndShadows/LightFunctions/LightFunctions.INT.udn
--- a/Engine/Documentation/Source/Engine/Rendering/LightingAndShadows/LightFunctions/LightFunctions.INT.udn
+++ b/Engine/Documentation/Source/Engine/Rendering/LightingAndShadows/LightFunctions/LightFunctions.INT.udn
@@ -2,6 +2,7 @@ Availability: Public
 Title:Light Functions
 Crumbs: %ROOT%, Engine, Engine/Rendering/LightingAndShadows
 Description:Lighting projected using a Material as the source.
+Version: 4.9
 
 [VAR:Topic]
 [OBJECT:Topic]
@@ -69,6 +70,6 @@ You can create a Light Function by setting the Material Domain of a Material to
 
 ## Applying a Light Function 
 
-You can apply a Light Function to a light from inside the **Details** panel under the Light Function Material setting. 
+You can apply a Light Function to a light from inside the **Details** panel under the **Light Function Material** setting. 
 
 ![](AddLightFunction.png)(w:640)
