#UE4doc #docs - Added Engine Version Field. Replaced image tables with comparison slider.

diff --git a/Engine/Documentation/Source/Engine/Rendering/LightingAndShadows/BumpMappingWithoutTangentSpace/BumpMappingWithoutTangentSpace.INT.udn b/Engine/Documentation/Source/Engine/Rendering/LightingAndShadows/BumpMappingWithoutTangentSpace/BumpMappingWithoutTangentSpace.INT.udn
--- a/Engine/Documentation/Source/Engine/Rendering/LightingAndShadows/BumpMappingWithoutTangentSpace/BumpMappingWithoutTangentSpace.INT.udn
+++ b/Engine/Documentation/Source/Engine/Rendering/LightingAndShadows/BumpMappingWithoutTangentSpace/BumpMappingWithoutTangentSpace.INT.udn
@@ -2,6 +2,7 @@ Availability:Public
 Title:Bump Mapping w/o Tangent Space
 Crumbs: %ROOT%, Engine, Engine/Rendering/LightingAndShadows
 Description:How to achieve bump mapping from 3D procedural shaders where classic tangent space does not work.
+Version: 4.9
 
 [VAR:Topic]
 [OBJECT:Topic]
@@ -80,24 +81,28 @@ ddx and ddy are computed for a 2x2 block and therefore show some blocky artifact
 
 ## Bump Mapping Material Functions
 
-Several [Material Functions](Engine/Rendering/Materials/Functions) are provided
+Several [**Material Functions**](Engine/Rendering/Materials/Functions) are provided
 to implement bump mapping in your materials without relying on tangent-space normal maps.
 
 ### ComputeFilterWidth
 
 ![FilterWidth.png](FilterWidth.png)(convert:false)
 
-This function makes use of [ddx and ddy](#ddxandddy) to compute how quickly a value is changing over the screen. This allows to 
+This function makes use of [**ddx and ddy**](#ddxandddy) to compute how quickly a value is changing over the screen. This allows to 
 fade out a procedural shader in the distance where it would start getting noisy. The fading results in less shimmering which is 
 even more noticeable in motion and extremely important for bump mapping as specularity on bumpy surfaces can produce 
 extreme aliasing artifacts.
 
 The following example images show how a procedural bump map function can be faded out in the distance.
-[REGION:imagetable]
-| ![DetailSize0.png](DetailSize0.png)(convert:false) | ![DetailSizeNot0.png](DetailSizeNot0.png)(convert:false) |
-| :---: | :---: |
-| Not Using FilterWidth | Using FilterWidth |
-[/REGION]
+
+[OBJECT:ComparisonSlider]
+	[PARAM:before]
+	![Not Using FilterWidth](DetailSize0.png)(convert:false)
+	[/PARAM]
+	[PARAM:after]
+	![Using FilterWidth](DetailSizeNot0.png)(convert:false)
+	[/PARAM]
+[/OBJECT]
 
 | Item | Description |
 | -------- | ----------- |
@@ -109,14 +114,17 @@ The following example images show how a procedural bump map function can be fade
 ### PerturbNormalLQ
 
 The **PerturbNormalLQ** function converts a grayscale bump map input into a world-space normal. 
-However, because it uses [ddx and ddy](#ddxandddy), which have the aforementioned 2x2 block artifacts, 
+However, because it uses [**ddx and ddy**](#ddxandddy), which have the aforementioned 2x2 block artifacts, 
 the output world-space normal is of low quality.
 
-[REGION:imagetable]
-| ![](bump_lq.png)(convert:false) | ![](bump_hq.png)(convert:false) |
-| ------ | ------ |
-| Low Quality | High Quality |
-[/REGION]
+[OBJECT:ComparisonSlider]
+	[PARAM:before]
+	![Low Quality](bump_lq.png)(convert:false)
+	[/PARAM]
+	[PARAM:after]
+	![High Quality](bump_hq.png)(convert:false)
+	[/PARAM]
+[/OBJECT]
 
 | Item | Description |
 | -------- | ----------- |
@@ -132,7 +140,7 @@ the output world-space normal is of low quality.
 
 [REGION:warning]
 	This function exists only as a reference and is not exposed to the Material Function Library. 
-	Use the [PerturbNormalHQ](#PerturbNormalHQ) function instead.
+	Use the [**PerturbNormalHQ**](#PerturbNormalHQ) function instead.
 [/REGION]
 
 ### PerturbNormalHQ
