#UE4doc #docs - Added Engine Version field. Bolded a couple of links where needed.

diff --git a/Engine/Documentation/Source/Engine/Rendering/LightingAndShadows/IESLightProfiles/IESLightProfiles.INT.udn b/Engine/Documentation/Source/Engine/Rendering/LightingAndShadows/IESLightProfiles/IESLightProfiles.INT.udn
--- a/Engine/Documentation/Source/Engine/Rendering/LightingAndShadows/IESLightProfiles/IESLightProfiles.INT.udn
+++ b/Engine/Documentation/Source/Engine/Rendering/LightingAndShadows/IESLightProfiles/IESLightProfiles.INT.udn
@@ -2,6 +2,7 @@ Availability:Public
 Title: IES Light Profiles
 Crumbs: %ROOT%, Engine, Engine/Rendering/LightingAndShadows
 Description:Using IES Light Profiles in Unreal Engine 4.
+Version: 4.9
 
 [VAR:Topic]
 [OBJECT:Topic]
@@ -61,10 +62,10 @@ Description:Using IES Light Profiles in Unreal Engine 4.
 
 ## Where can I get them?
 
-The easiest way to acquire IES Light Profiles is to simply head to a manufacturer's web site and download them, as almost all major lighting manufactures provide them free of charge. [Lithonia Lighting](http://www.lithonia.com/photometrics.aspx) and [Phillips](http://www.usa.lighting.philips.com/connect/tools_literature/photometric_data_1.wpd) for example, give a wide selection of light profiles for single lights. And they all document their lights very well, going so far as to base the profile on the type and color of the light emitter used and any variations in the light fixture.
+The easiest way to acquire IES Light Profiles is to simply head to a manufacturer's web site and download them, as almost all major lighting manufactures provide them free of charge. [**Lithonia Lighting**](http://www.lithonia.com/photometrics.aspx) and [**Phillips**](http://www.usa.lighting.philips.com/connect/tools_literature/photometric_data_1.wpd) for example, give a wide selection of light profiles for single lights. And they all document their lights very well, going so far as to base the profile on the type and color of the light emitter used and any variations in the light fixture.
 
 [REGION:tip]
-The nice thing about [Lithonia Lighting](http://www.lithonia.com/photometrics.aspx) is the website has a web browser profile viewer so you can see what you are getting, once you get down to the actual light you want.
+The nice thing about [**Lithonia Lighting**](http://www.lithonia.com/photometrics.aspx) is the website has a web browser profile viewer so you can see what you are getting, once you get down to the actual light you want.
 [/REGION]
 
 If you examine a profile in the Unreal Editor, they will be represented as a 1D texture (a gradient), however, they are not textures. They are ASCII files containing a wide array of meta data: 
@@ -115,7 +116,7 @@ _The IES Profile for a Lithonia TPG_1000S_PG21_M_
 IES Light profiles basically turn point lights into spotlights (visually speaking) and using an IES profile on a spot light will operate the same way as on a point light except the spotlight cones will mask the IES profile effect. However, since the projection area of a spot light is an 179 degree arc at its maximum, you lose any IES data past that point.
 
 [REGION:lightbox]
-[![](IES_SpotVsPoint_02.png)(h:220)](IES_SpotVsPoint_02.png) [![](IES_SpotVsPoint_01.png)(h:220)](IES_SpotVsPoint_01.png)
+[![](IES_SpotVsPoint_02.png)(h:270)](IES_SpotVsPoint_02.png) [![](IES_SpotVsPoint_01.png)(h:270)](IES_SpotVsPoint_01.png)
 [/REGION]
 
 _In these images, the light on the left is a spot light and the light on the right is a point light, notice the lack of lighting above the spot light. The spot light's cone clips the IES Light Profile effect as seen in the image on the right._
