#UE4doc #docs - Added Engine Version Field. Added Skill Level Field.

diff --git a/Engine/Documentation/Source/Engine/Rendering/LightingAndShadows/Lightmass/Basics/Lightmass.INT.udn b/Engine/Documentation/Source/Engine/Rendering/LightingAndShadows/Lightmass/Basics/Lightmass.INT.udn
--- a/Engine/Documentation/Source/Engine/Rendering/LightingAndShadows/Lightmass/Basics/Lightmass.INT.udn
+++ b/Engine/Documentation/Source/Engine/Rendering/LightingAndShadows/Lightmass/Basics/Lightmass.INT.udn
@@ -2,6 +2,8 @@ Availability: Public
 Title:Lightmass Basics
 Crumbs: %ROOT%, Engine, Engine/Rendering/LightingAndShadows
 Description:High level introduction to Lightmass.
+SkillLevel: Beginner
+Version: 4.9
 
 [VAR:Topic]
 [OBJECT:Topic]
