#UE4doc #docs - Added Engine Version Field.

diff --git a/Engine/Documentation/Source/Engine/Rendering/LightingAndShadows/LightMobility/DynamicLights/DynamicLights.INT.udn b/Engine/Documentation/Source/Engine/Rendering/LightingAndShadows/LightMobility/DynamicLights/DynamicLights.INT.udn
--- a/Engine/Documentation/Source/Engine/Rendering/LightingAndShadows/LightMobility/DynamicLights/DynamicLights.INT.udn
+++ b/Engine/Documentation/Source/Engine/Rendering/LightingAndShadows/LightMobility/DynamicLights/DynamicLights.INT.udn
@@ -2,6 +2,7 @@ Availability:Public
 Title:Movable Lights
 Crumbs: %ROOT%, Engine, Engine/Rendering/LightingAndShadows,Engine/Rendering/LightingAndShadows/LightMobility
 Description:A fully dynamic light capable of changing all of its properties during runtime.
+Version: 4.9
 
 [VAR:Topic]
 [OBJECT:Topic]
