#UE4doc #docs:
Removed Monospace from a few words
Fixed caption regions as they were formatted incorrectly.

diff --git a/Engine/Documentation/Source/Programming/Slate/Architecture/SlateArchitecture.INT.udn b/Engine/Documentation/Source/Programming/Slate/Architecture/SlateArchitecture.INT.udn
--- a/Engine/Documentation/Source/Programming/Slate/Architecture/SlateArchitecture.INT.udn
+++ b/Engine/Documentation/Source/Programming/Slate/Architecture/SlateArchitecture.INT.udn
@@ -81,9 +81,9 @@ The STextBlock uses a delegate called **Text** for this purpose.
 +------------+                                               
 ~~~
 
-	[REGION:caption]
-		The STextBlock reads the framerate as a string.
-	[/REGION]
+[REGION:caption]
+	The STextBlock reads the framerate as a string.
+[/REGION]
 [/REGION]
 
 Consider also that in the example above the Text is a string, while the framerate is most likely stored as a `float` or an `integer`. The use of delegates offers us the flexibility to perform the conversion whenever the value is read.
@@ -108,9 +108,9 @@ The assumption is that the programmer has attached the functionality appropriate
 |                 |      \-----------/      +--------------+
 +-----------------+
 ~~~
-	[REGION:caption]
-		The SEditable text reads the `item name`. When the user hits enter, the new text is sent to OnTextChanged where it will be validated and assigned to `item name` if appropriate.
-	[/REGION]
+[REGION:caption]
+	The SEditable text reads the item name. When the user hits enter, the new text is sent to OnTextChanged where it will be validated and assigned to item name if appropriate.
+[/REGION]
 [/REGION]
 
 During the next frame, the SEditableText will read from the Model's data. In the example above, `item name` will have been mutated by the OnTextChanged delegate and will be read for visualization via the Text delegate.
@@ -205,9 +205,9 @@ In the below example, we examine a horizontal box that arranges two children - a
 
 ~~~
 
-	[REGION:caption]
+[REGION:caption]
 	A horizontal box arranges two children - a text block and an image.
-	[/REGION]
+[/REGION]
 [/REGION]
 
 An STextBlock widget would compute its desired size by measuring the string that it is displaying. The SImage widget would
@@ -240,9 +240,9 @@ The recursion continues until all the children are arranged.
 .     |<---14--->|<-----11----->|
 ~~~
 
-	[REGION:caption]
+[REGION:caption]
 	A horizontal panel arranges two children, a text block and an image.
-	[/REGION]
+[/REGION]
 
 In the example above, the panel was allotted 25 units by its parent.
 The first slot indicates that it wants to use the desired size of the child, and is allotted 14 units of space.
