Fixing the wording and description of the Two Sided Foliage text that was added. Also bolding the first use of the shading model to match the style guide.

diff --git a/Engine/Documentation/Source/Engine/Rendering/Materials/MaterialProperties/LightingModels/LightingModels.INT.udn b/Engine/Documentation/Source/Engine/Rendering/Materials/MaterialProperties/LightingModels/LightingModels.INT.udn
--- a/Engine/Documentation/Source/Engine/Rendering/Materials/MaterialProperties/LightingModels/LightingModels.INT.udn
+++ b/Engine/Documentation/Source/Engine/Rendering/Materials/MaterialProperties/LightingModels/LightingModels.INT.udn
@@ -21,8 +21,8 @@ The following examples are all using Materials with their Blend Modes set to Opa
 
 ## Default Lit
 
-Default Lit is, as the name suggests, the default shading model and will likely be the one that you will use most often. This shading model makes use of direct and indirect lighting, as well as 
-specularity for reflections.
+**Default Lit** is, as the name suggests, the default shading model and will likely be the one that you will use most often. 
+This shading model makes use of direct and indirect lighting, as well as specularity for reflections.
 
 
 ![](DefaultLitObjects.png)(w:500)
@@ -41,7 +41,7 @@ When using the Default Lit Shading Model, you are given access to the following
 
 ## Unlit
 
-The Unlit Shading Model only outputs Emissive for color, making this perfect for special effects such as fire or illuminating objects. Note that in this example, the Material is _not_ casting light into the scene. Instead, its high Emissive value results in a glow effect, which is also picked up by the Dirt Mask applied to the camera. It _seems_ to illuminate, but no lighting or shadows will be cast by this object.
+The **Unlit** Shading Model only outputs Emissive for color, making this perfect for special effects such as fire or illuminating objects. Note that in this example, the Material is _not_ casting light into the scene. Instead, its high Emissive value results in a glow effect, which is also picked up by the Dirt Mask applied to the camera. It _seems_ to illuminate, but no lighting or shadows will be cast by this object.
 
 ![](UnlitExample.png)(w:500)
 
@@ -54,8 +54,7 @@ When using the Unlit Shading Model, you are given access to the following inputs
 
 
 ## Subsurface
-The Subsurface Shading Model simulates the effect of Subsurface Scattering. This is a real-world phenomenon in which light penetrates a surface and then diffuses throughout it. It can be most readily seen on such objects as ice, wax candles, and skin. 
-
+The **Subsurface** Shading Model simulates the effect of Subsurface Scattering. This is a real-world phenomenon in which light penetrates a surface and then diffuses throughout it. It can be most readily seen on such objects as ice, wax candles, and skin.
 The Subsurface model (as well as the Preintegrated Skin model described below) relies on the _Subsurface Color_ input. You can think of this as the color of the matter just beneath the surface of the object, such that when light scatters through the surface, this color will be seen. For human skin, you can often get good results from using a dark red color. In the Ice Elemental below, a dark blue-green color (with a variety of calculations based on lighting) is used to give a sense of translucent depth to the surface.
 
 For more information on the Subsurface Shading model, please see [the Subsurface Shading Model documentation](Engine/Rendering/Materials/LightingModels/SubSurface).
@@ -78,7 +77,7 @@ When using the Subsurface Shading Model, you are given access to the following i
 
 ## Preintegrated Skin
 
-The Preintegrated Skin Shading Model is very similar in nature to the Subsurface model, but geared toward low performance cost skin rendering on human characters. While not physically perfect, 
+The **Preintegrated Skin** Shading Model is very similar in nature to the Subsurface model, but geared toward low performance cost skin rendering on human characters. While not physically perfect, 
 this Shading Model will often give nice results on your characters while being a bit cheaper to render than the Subsurface method. In the image below, the character's flesh has been set to use 
 the Preintegrated Skin Shading Model.
 
@@ -98,7 +97,7 @@ When using the Preintegrated Skin Shading Model, you are given access to the fol
 * Ambient Occlusion
 
 ## Clear Coat 
-The Clear Coat Shading Model can be used to better simulate multilayer materials that have a thin translucent layer of film over the surface of a standard material.
+The **Clear Coat** Shading Model can be used to better simulate multilayer materials that have a thin translucent layer of film over the surface of a standard material.
 In addition to this, the Clear Coat Shading Model can also be used with either a metal or nonmetal surfaces. In fact, it was specifically designed to model this second class of smooth colored films over a non-colored metal. 
 Some examples include acrylic or lacquer clear coats, and colored films over metals such as soda cans and car paint. 
 
@@ -142,20 +141,22 @@ When using the Subsurface Profile Shading Model, you are given access to the fol
 
 ## Two Sided Foliage
 
-The Two Sided Foliage shading model was introduced to give foliage a more natural look when being lit. 
-Normally when Foliage is being lit the sides facing away from the light source will be very dark to almost black. 
-The Two Sided Foliage shading model helps to eliminate this issue by simulating a very cheap Subsurface scattering to give the illusion that the foliage is receiving more light on the polygons that are facing away from the light source than it actually is. 
+The **Two Sided Foliage** shading model was introduced to help give foliage a more natural and uniform look when being lit. 
+When using the Default Lit shading model for foliage parts of the foliage can look almost black.
+The Default Lit shading model also does not simulate any kind of light transmission which is key to having believable looking foliage. 
+The Two Sided Foliage shading model helps to eliminate these issues by simulating a very cheap Subsurface scattering to give the illusion that light is being transferred through the foliage. 
+This also helps to eliminate the dark parts of the foliage and give the foliage a more uniform look to it's lighting.
 
 [OBJECT:ComparisonSlider]
 	[PARAM:before]
-	![Two Sided Lighting Enabled](2_Sided_Enabled.png)
+	![Two Sided Foliage Enabled](2_Sided_Enabled.png)
 	[/PARAM]
 	[PARAM:after]
-	![Two Sided Lighting Disabled](2_Sided_Disabled.png)
+	![Two Sided Foliage Disabled](2_Sided_Disabled.png)
 	[/PARAM]
 [/OBJECT]
 
-When using the Two Sided Foliage Model, you are give access to the following inputs.
+When using the Two Sided Foliage Model, you are given access to the following inputs.
 * Base Color
 * Specular
 * Roughness
