#UE4doc #docs:
Added engine version field

diff --git a/Engine/Documentation/Source/Engine/Basics/Projects/UnrealGameProjects.INT.udn b/Engine/Documentation/Source/Engine/Basics/Projects/UnrealGameProjects.INT.udn
--- a/Engine/Documentation/Source/Engine/Basics/Projects/UnrealGameProjects.INT.udn
+++ b/Engine/Documentation/Source/Engine/Basics/Projects/UnrealGameProjects.INT.udn
@@ -2,6 +2,7 @@ Availability:Public
 Title:Unreal Game Projects
 Crumbs: %ROOT%, Engine
 Description:Overview of game projects and instructions for working with them.
+version: 4.9
 
 [VAR:Topic]
 [OBJECT:Topic]
