#UE4doc #docs:
Added engine version field
Changed font settings for multiple terms
Changed Tip region to a Note region due to the nature of the content

diff --git a/Engine/Documentation/Source/Engine/Basics/Projects/CreatingTemplates/CreatingTemplates.INT.udn b/Engine/Documentation/Source/Engine/Basics/Projects/CreatingTemplates/CreatingTemplates.INT.udn
--- a/Engine/Documentation/Source/Engine/Basics/Projects/CreatingTemplates/CreatingTemplates.INT.udn
+++ b/Engine/Documentation/Source/Engine/Basics/Projects/CreatingTemplates/CreatingTemplates.INT.udn
@@ -2,6 +2,7 @@ Availability:Public
 Title:Converting a Project to a Template
 Crumbs: %ROOT%, Engine, Engine/Basics/Projects
 Description:Steps for converting an existing project to a template.
+version: 4.9
 
 [EXCERPT:ConvertDesc]
 To convert an existing project into a template:
@@ -10,26 +11,26 @@ To convert an existing project into a template:
 1.	Move the project folder into the appropriate Templates directory:
 |Engine Version | Path | 
 |---------------|------|
-| Launcher | `C:\Program Files\Unreal Engine\[Version]\Templates` |
-| Source | `[ForkLocation]\UE4\Templates` |
-1.	Open `Templates\\[ProjectName]\Config\DefaultGame.ini` and add or update the `ProjectName` variable.
+| Launcher | C:\Program Files\Unreal Engine\[Version]\Templates |
+| Source | [ForkLocation]\UE4\Templates |
+1.	Open Templates\\[ProjectName]\Config\DefaultGame.ini and add or update the **ProjectName** variable.
 
 		[/Script/EngineSettings.GeneralProjectSettings]
 		ProjectID=E6468D0243A591234122E38F92DB28F4
 		ProjectName=MyTestTemplate
-1.	Copy `Templates\TP_FirstPerson\Config\TemplateDefs.ini` to `Templates\\[ProjectName]\Config\TemplateDefs.ini`. You can use any existing template folder in place of TP_FirstPerson.
-1.	Open `Templates\\[ProjectName]\Config\TemplateDefs.ini` and update the `LocalizedDisplayNames` and `LocalizedDescriptions` variables.
+1.	Copy Templates\TP_FirstPerson\Config\TemplateDefs.ini to Templates\\[ProjectName]\Config\TemplateDefs.ini. You can use any existing template folder in place of TP_FirstPerson.
+1.	Open Templates\\[ProjectName]\Config\TemplateDefs.ini and update the **LocalizedDisplayNames** and **LocalizedDescriptions** variables.
 
 		[/Script/GameProjectGeneration.TemplateProjectDefs]
 		LocalizedDisplayNames=(Language="en",Text="My Test Template")
 		LocalizedDescriptions=(Language="en",Text="A project including starter content and a first person character set up in Blueprints.")
-1.	You may add an icon file in `Templates\\[ProjectName]\\[ProjectName].png`
+1.	You may add an icon file in Templates\\[ProjectName]\\[ProjectName].png
 
 You should now see the new template in the New Project dialog. 
 
-[REGION:tip]
-Note that to see new templates, you only need to close and open the New Project dialog. However, if you make any changes to an existing template (for example LocalizedDisplayNames), the change will not appear 
+[REGION:note]
+To see new templates, you only need to close and open the New Project dialog. However, if you make any changes to an existing template (for example LocalizedDisplayNames), the change will not appear 
 until you restart the editor.
-[/REGION:tip]
+[/REGION]
 
 [/EXCERPT:ConvertDesc]
