#UE4doc #docs:
Added engine version field
Changed font settings for multiple terms
Removed section that was severly out of date.

diff --git a/Engine/Documentation/Source/Engine/Basics/Projects/Packaging/PackagingProjects.INT.udn b/Engine/Documentation/Source/Engine/Basics/Projects/Packaging/PackagingProjects.INT.udn
--- a/Engine/Documentation/Source/Engine/Basics/Projects/Packaging/PackagingProjects.INT.udn
+++ b/Engine/Documentation/Source/Engine/Basics/Projects/Packaging/PackagingProjects.INT.udn
@@ -2,6 +2,7 @@ Availability:Public
 Title:Packaging Projects
 Crumbs: %ROOT%, Engine, Engine/Basics/Projects
 Description:Packaging Unreal game projects for distribution.
+version: 4.9
 
 [TOC(start:2)]
 
@@ -55,11 +56,15 @@ The status indicator provides a way to cancel the active packaging process by cl
 For less experienced users, the most important log messages, such as errors and warnings, are logged to the regular **Message Log** window:
 
 ![](message_log.png)
+
+[REGION:note]
+If neither of these windows are visible, they can be accessed by the **Window** > **Developer Tools** > **Output Log** / **Message Log** options.
+[/REGION]
  
 ## Running Packaged Games
 
 When you packaged, you selected a directory for the output. If packaging completes successfully, your packaged game will be in a platform-specific subdirectory. For instance, if you selected TappyChicken/Final, then the iOS build would be 
-in `TappyChicken/Final/IOS` and an Android build would be in `TappyChicken/Final/Android`. When you go to that subdirectory, you will see the packaged game, in a format appropriate for the platform. 
+in TappyChicken/Final/IOS and an Android build would be in TappyChicken/Final/Android. When you go to that subdirectory, you will see the packaged game, in a format appropriate for the platform. 
 For Android, you will see .apk, .obb, and .bat files (run the .bat to install the game on your device). For IOS, you will see an .ipa file. This can be installed via iTunes or Xcode.
 The number and type of files created will vary depending on the target platform. The following image shows example output for a Windows project:
 
@@ -85,10 +90,10 @@ To submit an IOS or Android game to the App Store or Google Play Store, you need
 and check the Distribution check box. This will sign the final package appropriately.
 
 On IOS, you will need to create a distribution Certificate and MobileProvision on Apple's developer website. Install the certificate the same way as your Development certificate, and name your 
-distribution provision with a `Distro_` prefix, next to your other one (so you would have `Distro_MyProject.mobileprovision` and `MyProject.mobileprovision`).
+distribution provision with a Distro_ prefix, next to your other one (so you would have Distro_MyProject.mobileprovision and MyProject.mobileprovision).
 
-On Android, you will need to create a key to sign the .apk, and give some information to our build tools with a file called `SigningConfig.xml`. This file exists in the installed Engine directory 
-(`Engine/Build/Android/Java`). If you edit this file, it will apply to all your projects. However, you can copy this file to your project's `Build/Android` directory (no Java directory), and it 
+On Android, you will need to create a key to sign the .apk, and give some information to our build tools with a file called SigningConfig.xml. This file exists in the installed Engine directory 
+(Engine/Build/Android/Java). If you edit this file, it will apply to all your projects. However, you can copy this file to your project's Build/Android directory (no Java directory), and it 
 will be used for just that project. Directions for how to generate the key and fill out the file are found in the file itself.
 
 
@@ -115,14 +120,6 @@ Clicking **File > Package Project > Packaging Settings...** or **Edit > Project
 | Directories to Always Cook | The list of directories with files that should always be cooked. By default, the packaging feature automatically detects and cooks all content that is referenced by your game. Some required content may not be referenced directly, such as textures for a custom Slate UI. The directories in which such content resides should be listed here in order to also include it in the package. |
 [/PUBLISH:Licensee]
 
-[PUBLISH:Licensee]
-## Supported Platforms
-
-At the time of this writing, the only support platform is Windows. We are in the process of adding support for other platforms, such as iOS, PS4, and Android. Although these platforms may already be listed in the Package Project menu, they are currently grayed out and cannot be selected.
-
-Creating an installer for Windows is not yet supported. The packaged files are simply copied out to the staging directory for the time being.
-[/PUBLISH:Licensee]
-
 
 
 ## Content Cooking
