#UE4doc #docs:
Added engine version field
Removed Overview header
Changed font settings for a few terms
Fixed capitalization

diff --git a/Engine/Documentation/Source/Engine/Basics/Redirectors/Redirectors.INT.udn b/Engine/Documentation/Source/Engine/Basics/Redirectors/Redirectors.INT.udn
--- a/Engine/Documentation/Source/Engine/Basics/Redirectors/Redirectors.INT.udn
+++ b/Engine/Documentation/Source/Engine/Basics/Redirectors/Redirectors.INT.udn
@@ -2,20 +2,17 @@ Availability: Public
 Title:Redirectors
 Crumbs:%ROOT%, Engine
 Description:Objects that are placed in packages to redirect references to moved assets to their current locations.
-
+version: 4.9
 [TOC (start:2 end:3)]
 
- 
-
 
-## Overview
 
-Moving or renaming an asset in UE4 leaves a 'Redirector' in its old location. This is so that packages that are not currently loaded, but reference this asset, will know where to find it in its new location. Choosing a naming system early on and sticking with it will avoid many of the problems experienced with redirectors.
+Moving or renaming an asset in UE4 leaves a **Redirector** in its old location. This is so that packages that are not currently loaded, but reference this asset, will know where to find it in its new location. Choosing a naming system early on and sticking with it will avoid many of the problems experienced with Redirectors.
 
 
 ## FixupRedirects Commandlet
 
-Be sure to run the 'fixupredirects' commandlet on a regular basis. On previous projects, our Level Designers ran this commandlet on all packages every two weeks. It will clean up as much as it can (so it does not require that everyone check in every package). 
+Be sure to run the FixupRedirects commandlet on a regular basis. On previous projects, our Level Designers ran this commandlet on all packages every two weeks. It will clean up as much as it can (so it does not require that everyone check in every package). 
 
 
 ## Gotchas
