#UE4doc #docs:
Added engine version field
Removed Overview header
Changed a Note region into a Warning region due to the nature of the content
Changes a Tip region into a Note region due to the nature of the content.

diff --git a/Engine/Documentation/Source/Engine/Basics/Screenshots/Screenshots.INT.udn b/Engine/Documentation/Source/Engine/Basics/Screenshots/Screenshots.INT.udn
--- a/Engine/Documentation/Source/Engine/Basics/Screenshots/Screenshots.INT.udn
+++ b/Engine/Documentation/Source/Engine/Basics/Screenshots/Screenshots.INT.udn
@@ -2,11 +2,10 @@ Availability: Public
 Title:Taking Screenshots
 Crumbs:%ROOT%, Engine
 Description:Guide to taking in-game screenshots of environments and gameplay.
-
+version: 4.9
 [TOC (start:2 end:3)]
 
 
-## Overview
 
 Unreal Engine 4 (UE4) offers users a multitude of different ways of taking in-game screenshots of your Unreal Engine game project. 
 This document will not only cover which tools you have available for taking screenshots but also how you can use the tools. 
@@ -35,7 +34,7 @@ It can be used to generate screenshots of any size that are based off of the scr
 The high resolution screenshot console command does this by rendering the game frame multiple times, one tile at a time at full resolution and concatenating all tiles into a single image file. 
 
 The default behavior is to take a screenshot with a resolution 4 times greater than the game resolution and create an image file in the Screenshot folder. 
-When running the game through the editor (PIE) or as standalone, open the command console by hitting either &lt;Tab&gt; or &lt;~&gt; (_tilde_) and type:
+When running the game through the editor (PIE) or as standalone, open the command console by hitting **~** (_tilde_) and type:
 
     HighResShot
     
@@ -67,7 +66,7 @@ Once you have selected the **High Resolution Screenshot**, the HRSST will be dis
 
 ![](HRSS_HRSST.png)
 
-[region:tip]
+[region:note]
 You can move the HRSST anywhere you want to in side of UE4. However be careful as sometime the HRSST can become hidden by other floating windows like the **Content Browser**. 
 [/region]
 
@@ -83,7 +82,7 @@ You can move the HRSST anywhere you want to in side of UE4. However be careful a
 | 7: | Clicking on this will make the capture region the entire editor window.|
 | 8: | Clicking on this this will take the screenshot.|
 
-[region:note]
+[region:warning]
 Due to the high system requirements of a high resolution screenshot, very large multipliers might cause the graphics driver to be unresponsive and possibly crash. In these circumstances, please try using a lower number. 
 [/region]
 
