#UE4doc #docs:
Added engine version field
Removed Overview header

diff --git a/Engine/Documentation/Source/Engine/Basics/SourceControl/SourceControlLandingPage.INT.udn b/Engine/Documentation/Source/Engine/Basics/SourceControl/SourceControlLandingPage.INT.udn
--- a/Engine/Documentation/Source/Engine/Basics/SourceControl/SourceControlLandingPage.INT.udn
+++ b/Engine/Documentation/Source/Engine/Basics/SourceControl/SourceControlLandingPage.INT.udn
@@ -2,9 +2,8 @@ Availability: Public
 Title: Collaboration in Unreal Engine 4
 Crumbs: %ROOT%, Engine
 Description: How to setup Perforce or SVN so that you can share assets with other on your team. 
+version: 4.9
 
-
-## Overview 
 Unreal Engine 4 comes with two fully-integrated methods that allow people to work together on their projects using version control software like Perforce and SVN. 
 Version control allows users on your team to share assets and code with one another, as well as providing a backup and history of changes so any files can 
 be rolled back to an earlier version if something went wrong with the file.
