#UE4doc #docs:
Added engine version field
Removed Overview header
Changed font settings for multiple terms
Fixed capitalization in some spots

diff --git a/Engine/Documentation/Source/Engine/Basics/SourceControl/InEditor/SourceControlInsideUE4.INT.udn b/Engine/Documentation/Source/Engine/Basics/SourceControl/InEditor/SourceControlInsideUE4.INT.udn
--- a/Engine/Documentation/Source/Engine/Basics/SourceControl/InEditor/SourceControlInsideUE4.INT.udn
+++ b/Engine/Documentation/Source/Engine/Basics/SourceControl/InEditor/SourceControlInsideUE4.INT.udn
@@ -7,11 +7,11 @@ Related: Engine/Content/Browser
 Related: Engine/Basics/Redirectors
 Related: Engine/UI/ProjectSettings
 Related: Programming/Basics/ConfigurationFiles
-
+version: 4.9
 
 [TOC(start:2)]
 
-## Overview 
+
 No matter which Source Control software you end up using, the manner in which you interface with that software from inside of UE4 is the same. The following guide will detail everything 
 you need to know about using Source Control inside of UE4. 
 
@@ -21,8 +21,8 @@ you need to know about using Source Control inside of UE4.
 Unreal Engine 4 makes it very easy for multiple user collaboration on projects. In the following section, we will be going over how all of the various collaboration tools work inside of the Unreal Editor. 
 
 ### Checking Out Files in Unreal Editor
-The main way to share content with collaborators is via the [**Content Browser**](Engine/Content/Browser). All source control operations are accessible via the context menu for files and folders. 
-For instance, to **Check Out** a file (i.e. lock it for editing), **Right-Click** on the asset in the **Content Browser** and then select the **Check Out** option.
+The main way to share content with collaborators is via the [Content Browser](Engine/Content/Browser). All source control operations are accessible via the context menu for files and folders. 
+For instance, to **Check Out** a file (i.e. lock it for editing), **Right-click** on the asset in the **Content Browser** and then select the **Check Out** option.
 
 ![](SC_UE4_Checkout_1.png)
 
@@ -33,7 +33,7 @@ You can find these options by going to **Editor Preferences** -> **Loading & Sav
 
 You can also set the Unreal Editor to automatically add new files to Source Control when they are first saved inside of the Editor. You can find these options by going to 
 **Editor Preferences** -> **Loading & Saving** -> **Source Control**. To do this, all you have to do is enable the **Add New Files when Modified** option which is on by default.
-When you start editing an asset, you will be presented with a notification in the **Bottom-Right Hand Corner** of the screen asking you whether you want to check the file(s) out:
+When you start editing an asset, you will be presented with a notification in the bottom-right hand corner of the screen asking you whether you want to check the file(s) out:
 
 ![](SC_UE4_Checkout_3.png)
 
@@ -44,11 +44,11 @@ This allows you to check out assets in advance of editing them. If you click the
 
 | Number | Property Name|Description
 |--------|------------|------------|
-| 1:     | Check Out Selected:|Checks out the files that are checked in the list.|
-| 2:     | Make Writable:|Clears the read-only flag (if one exists) allowing local edits to a file.|
-| 3:     | Ask Me Later:|Allows you to carry on working and check files out at a later time.|
+| 1:     | **Check Out Selected**|Checks out the files that are checked in the list.|
+| 2:     | **Make Writable**|Clears the read-only flag (if one exists) allowing local edits to a file.|
+| 3:     | **Ask Me Later**|Allows you to carry on working and check files out at a later time.|
 
-When you **Right-Click** on an object in the **Content Browser** and then go under the **Source Control** option, you will see various symbols that show the current state of the assets. 
+When you **Right-click** on an object in the **Content Browser** and then go under the **Source Control** option, you will see various symbols that show the current state of the assets. 
 
 | Icon | Description
 |--------|-------------|
@@ -63,19 +63,19 @@ When you **Right-Click** on an object in the **Content Browser** and then go und
 ### Submitting to Source Control
 Submitting to Source Control inside of Unreal Engine 4 is a vital part of the Source Control process as if you do not submit anything to Source Control, other users will not be able to see 
 what you have been working on/created. A workflow that is convenient if you are working on many changes, or you are not tracking changes you are making too closely, is to make use of the 
-**Submit to Source Control** option in the **File Menu** of the Editor.
+**Submit to Source Control** option in the **File** menu of the Editor.
 
 ![](SC_UE4_Submit_To_Source.png)
 
-This will scan your project and the Engine for changes made to assets that can be submitted and then present you with the Submit Files dialog.
+This will scan your project and the Engine for changes made to assets that can be submitted and then present you with the **Submit Files** dialog.
 
 ![](SC_UE4_Submit_To_Source_CL_Description.png)
 
 You can then enter a description of your changes and then click **OK** and you will then commit/submit the files to Source Control.
 
 ### Syncing Files 
-If you want to get changes that other people have made, you can **Right-Click** on files or folders in the **Content Browser** and select the **Sync** option. 
-Alternatively, you can **Sync** using an external program such as P4V. However, it is recommended that you shut the Editor down if you do this as it may lock some files for editing while open.
+If you want to get changes that other people have made, you can **Right-click** on files or folders in the **Content Browser** and select the **Sync** option. 
+Alternatively, you can **Sync** using an external program such as Perforce (P4V). However, it is recommended that you shut the Editor down if you do this as it may lock some files for editing while open.
 
 [region:note]
 When a file is **Out of Sync** (i.e. a newer version exists on the Source Control server to the one you have locally), the following icon, ![](SC_UE4_Checkout_New_Version.png),
@@ -91,20 +91,20 @@ Files that are added appear with a  ![](SC_UE4_Checkout_Add.png) symbol on them
 
 ### Deleting Files 
 Deleting assets that are already in Source Control is more complex as they are usually referenced by other assets. 
-If you select a file in the **Content Browser** and press **Delete** (or select it from the **Right-Click** menu), you will be presented with the **Delete Assets** dialog. 
+If you select a file in the **Content Browser** and press **Delete** (or select it from the **Right-click** menu), you will be presented with the **Delete Assets** dialog. 
 
 ![](SC_UE4_Delete_Assets.png)
 
-| Number | Description|
-| -------| ---------- |
-| 1:     | This is the asset(s) that you are going to delete.|
-| 2:     | This shows asset(s) that are referencing the assets you are about to delete.|
-| 3:     | Here you specify what assets you want other assets that reference this asset(s) to reference when this asset is deleted.|
-| 4:     | Here you can force delete an asset. However be very careful doing this as it can have a very negative impact on your project. You should only use this option as a last resort if nothing else works.|
-| 5:     | This will cancel the delete operation.|
+| Number | Name |Description|
+| ---| ---- | ---------- |
+| 1| **Pending Deleted Assets** |This is the asset(s) that you are going to delete.|
+| 2| **Assets Referencing the Pending Deleted Assets** |This shows asset(s) that are referencing the assets you are about to delete.|
+| 3| **Replace References** |Here you specify what assets you want other assets that reference this asset(s) to reference when this asset is deleted.|
+| 4| **Force Delete** |Here you can force delete an asset. However be very careful doing this as it can have a very negative impact on your project. You should only use this option as a last resort if nothing else works.|
+| 5| **Cancel** |This will cancel the delete operation.|
 
-Once you have chosen how to handle the deletion, the file will generally be replaced by a hidden **Asset** known as a [**Redirector**](Engine/Basics/Redirectors). If you want to clean up the [**Redirector**](Engine/Basics/Redirectors) 
-that is created when deleting assets, you can **Right-Click** on a **folder** in the **Content Browser** and select the **Fix Up Redirectors in Folders**. This will fix up the references 
+Once you have chosen how to handle the deletion, the file will generally be replaced by a hidden **Asset** known as a [Redirector](Engine/Basics/Redirectors). If you want to clean up the [Redirector](Engine/Basics/Redirectors) 
+that is created when deleting assets, you can **Right-click** on a **folder** in the **Content Browser** and select the **Fix Up Redirectors in Folders**. This will fix up the references 
 in your project and permanently delete the redirectors, marking the files for delete in source control. Presently, to commit these deletions, you need to use an external tool such as P4V or TortoiseSVN.
 
 ### Level Files
