#UE4doc #docs:
Added engine version field
Removed Overview header
Changed font settings for multiple terms
Fixed capitalization in multiple areas
Replaced multiple pictures due to being out of date compared to Perforce's latest version

diff --git a/Engine/Documentation/Source/Engine/Basics/SourceControl/Perforce/P4SourceControl.INT.udn b/Engine/Documentation/Source/Engine/Basics/SourceControl/Perforce/P4SourceControl.INT.udn
--- a/Engine/Documentation/Source/Engine/Basics/SourceControl/Perforce/P4SourceControl.INT.udn
+++ b/Engine/Documentation/Source/Engine/Basics/SourceControl/Perforce/P4SourceControl.INT.udn
@@ -7,13 +7,13 @@ Related: Engine/Content/Browser
 Related: Engine/Basics/Redirectors
 Related: Engine/UI/ProjectSettings
 Related: Programming/Basics/ConfigurationFiles
-
+version: 4.9
 
 [TOC(start:2)]
 
-## Overview 
+
 Unreal Engine 4 comes with an integrated method that allows people to work together on their projects using version control software called Perforce. 
-Version control allows users on your team to share assets and code with one another, as well as providing a backup and history of changes so any files can 
+Source Control allows users on your team to share assets and code with one another, as well as providing a backup and history of changes so any files can 
 be rolled back to an earlier version if something went wrong with the file.
 
 ## Workflow 
@@ -32,8 +32,8 @@ they check in, or commit those files, uploading the changes to the server and re
 You can download the Perforce server tools from [here](http://www.perforce.com/p4d). Once downloaded, run the installer and choose a location and name for your server.
 [/region]
 
-There are two ways of running the Perforce server, **P4D** and **P4S**. **P4D** runs while an instance of the command line is running, so it is usually only used for spot-maintenance tasks. 
-**P4S** is the service-based equivalent of **P4D** as it allows you to run the server in the background. The service is usually installed and started when you install the Perforce server tools. 
+There are two ways of running the Perforce server, **P4D** and **P4S**. P4D runs while an instance of the command line is running, so it is usually only used for spot-maintenance tasks. 
+P4S is the service-based equivalent of P4D as it allows you to run the server in the background. The service is usually installed and started when you install the Perforce server tools. 
 
 [Region:note]
 Sometimes the service does not get run automatically. To turn the service on first go into the **Control Panel -> Administrative Tools ->Services applet**.
@@ -115,11 +115,11 @@ A quick search online for **Perforce Hosting** should turn up a few results. Eac
 that you can use to connect. See also **Connecting from Unreal Editor** below for a note about connecting to Perforce cloud providers.
 
 ### Setting up a Local Workspace
-To work with files that are under Perforce control, you will need to set up a Workspace so that the files can be managed on your local machine. You can do this via P4V by opening up the Workspaces tab:
+To work with files that are under Perforce control, you will need to set up a Workspace so that the files can be managed on your local machine. You can do this via P4V by opening up the **Workspaces** tab in the **View** dropdown:
 
 ![](SC_View_Workspace.png)
 
-To create a new workspace, **Right-Click** in the Workspace window and select the **New Workspace** option. Then enter a name for your workspace and a location where you want the files to be on your 
+To create a new workspace, **Right-click** in the Workspace window and select the **New Workspace** option. Then enter a name for your workspace and a location where you want the files to be on your 
 local machine, the Workspace Root.
 
 ![](SC_New_Workspace.png)
@@ -138,7 +138,7 @@ Note that the server typemap setup above should mean that coders can compile the
 
 Adding files to Perforce is something that can be accomplished in the few following steps. 
 
-1. To add files/folders to Perforce, first select the files/folders in P4V and then **Right-Click** and select the **Mark For Add** option.
+1. To add files/folders to Perforce, first select the files/folders in P4V and then **Right-click** and select the **Mark For Add** option.
 
 	![](SC_P4_Mark_For_Add.png)
 
@@ -146,16 +146,16 @@ Adding files to Perforce is something that can be accomplished in the few follow
 
 	![](SC_View_Pending_Changelists.png)
 
-1. To submit something, first **Right-Click** on the change list and then select **Submit**. Then enter a change list description and click the Submit button. 
+1. To submit something, first **Right-click** on the change list and then select **Submit**. Then enter a change list description and click the Submit button. 
 
 	![](SC_P4_Submit_Files.png)
 
 If you want to start work on a project that someone else has already started, you will need to get the latest revision from Perforce using P4V. 
-To do this, find the project in the Depot View in P4V and then **Right-Click** on the project you want to get and select **Get Latest Revision**. 
+To do this, find the project in the Depot View in P4V and then **Right-click** on the project you want to get and select **Get Latest Revision**. 
 This will download all the files to the location you specified for your Workspace Root above.
 
 ### Connecting from Unreal Editor
-You can connect to your Perforce server right from the Editor by clicking on the **Source Control** icon in the top-right of the editor:
+You can connect to your Perforce server right from the Editor by clicking on the **Source Control** icon on the **Toolbar**:
 
 ![](SC_UE4_SC_Enable.png)
 
