#UE4doc #docs:
Added engine version field
Removed Overview header
Changed font settings for multiple terms
Fixed capitalization in some places
Replaced a few images due to out of date UI and incorrect reference

diff --git a/Engine/Documentation/Source/Engine/Basics/SourceControl/SVN/SVNSourceControl.INT.udn b/Engine/Documentation/Source/Engine/Basics/SourceControl/SVN/SVNSourceControl.INT.udn
--- a/Engine/Documentation/Source/Engine/Basics/SourceControl/SVN/SVNSourceControl.INT.udn
+++ b/Engine/Documentation/Source/Engine/Basics/SourceControl/SVN/SVNSourceControl.INT.udn
@@ -7,11 +7,10 @@ Related: Engine/Content/Browser
 Related: Engine/Basics/Redirectors
 Related: Engine/UI/ProjectSettings
 Related: Programming/Basics/ConfigurationFiles
-
+version: 4.9
 
 [TOC(start:2)]
 
-## Overview 
 Unreal Engine 4 comes with an integrated method that allows people to work together on their projects using version control software called Subversion or SVN for short. 
 Version control allows users on your team to share assets and code with one another, as well as providing a backup and history of changes so any files can 
 be rolled back to an earlier version if something went wrong with the file.
@@ -30,7 +29,7 @@ You can download SVN for Windows by clicking [here](https://www.visualsvn.com/se
 
 ### Setting up a SVN server on your Local Machine
 
-1. After SVN had been installed, open SVN up and in VisualSVN, **Right-Click** on the **Repositories** tree item and select **Create New Repository**:
+1. After SVN had been installed, open SVN up and in VisualSVN, **Right-click** on the **Repositories** tree item and select **Create New Repository**:
 
 	![](SC_SVN_Create_New_Repository.png)
 
@@ -58,7 +57,7 @@ To do this, follow the steps above on the server and connect using the server's
 ### Secure Access
 Many users like to use a secure connection to access their repositories. These connections have URLs starting with **https://** and have a number of things to be aware of when setting them up.
 To use HTTPS, the server needs to provide a **trusted certificate** so that clients know that they are connecting to a secure server. VisualSVN can issue a local certificate so that internal users can connect securely.
-You can find details of this by **Right-Clicking** on the root of the tree view, selecting **Properties**, and navigating to the **Certificate** tab:
+You can find details of this by **Right-clicking** on the root of the tree view, selecting **Properties**, and navigating to the **Certificate** tab:
 
 ![](SC_SVN_Certificate.png)
 
@@ -94,7 +93,7 @@ If you are trying to edit a project that someone else has already created, you c
 [/region]
 
 
-To create your working copy, you should **Check Out** the (assumed empty) root of your repository. To do this, **Right-Click** on your project's folder in Windows Explorer and select **SVN Checkout**.
+To create your working copy, you should **Check Out** the (assumed empty) root of your repository. To do this, **Right-click** on your project's folder in Windows Explorer and select **SVN Checkout**.
 
 ![](SC_SVN_Checkout.png)
 
@@ -106,12 +105,12 @@ Next click **OK** and it should create a hidden **.svn** directory that SVN uses
 
 
 ### Adding files to Subversion
-1. To add the files to your repository, select the items in Windows Explorer and then **Right-Click** to add and select  the **TortoiseSVN->Add** option.
+1. To add the files to your repository, select the items in Windows Explorer and then **Right-click** to add and select  the **TortoiseSVN->Add** option.
 
 	![](SC_SVN_Add.png)
 
 1. When you do this, SVN will mark the files you selected for addition to the repository.
-	Next, **Right-Click** in the directory containing your project and select **TortoiseSVN->SVN Commit**.
+	Next, **Right-click** in the directory containing your project and select **TortoiseSVN->SVN Commit**.
 
 	![](SC_SVN_Commit.png)
 
@@ -124,7 +123,7 @@ Next click **OK** and it should create a hidden **.svn** directory that SVN uses
 	set for your checkout depth. This should download the all the existing files to your local machine.
 
 ### Connecting to Subversion from Unreal Editor
-You can connect to your SVN server from the Unreal Editor by clicking on the **Source Control** icon in the top-right-hand corner of the Editor.
+You can connect to your SVN server from the Unreal Editor by clicking on the **Source Control** icon in the **Toolbar**.
 
 ![](SC_SVN_Connect_From_UE4_1.png)
 
