#UE4doc #docs:
Added engine version field
Removed italics from a term

diff --git a/Engine/Documentation/Source/Engine/Blueprints/Blueprints.INT.udn b/Engine/Documentation/Source/Engine/Blueprints/Blueprints.INT.udn
--- a/Engine/Documentation/Source/Engine/Blueprints/Blueprints.INT.udn
+++ b/Engine/Documentation/Source/Engine/Blueprints/Blueprints.INT.udn
@@ -2,6 +2,7 @@ Availability:Public
 Title: Blueprints Visual Scripting
 Crumbs: %ROOT%, Engine
 Description:Overview of using the Blueprint visual scripting system for gameplay.
+version: 4.9
 
 
 [VAR:Topic]
@@ -78,7 +79,7 @@ a node-based interface to create gameplay elements from within Unreal Editor. Th
 provides the ability for designers to use virtually the full range of concepts and tools generally only available to programmers.
 [/EXCERPT:Intro]
 
-Through the use of _Blueprints_, designers can prototype, implement, or modify virtually any gameplay element, such as:
+Through the use of Blueprints, designers can prototype, implement, or modify virtually any gameplay element, such as:
 
 * **Games** - set up game rules, tweak gameplay conditions, etc.
 * **Players** - create variants with different meshes and _Materials_ or character customization.
