#UE4doc #docs:
Added engine version field
Removed Overview header
Changed font settings for multiple terms
Replaced an image due to not being up to date with the editor's UI

diff --git a/Engine/Documentation/Source/Engine/Blueprints/BestPractices/BestPractices.INT.udn b/Engine/Documentation/Source/Engine/Blueprints/BestPractices/BestPractices.INT.udn
--- a/Engine/Documentation/Source/Engine/Blueprints/BestPractices/BestPractices.INT.udn
+++ b/Engine/Documentation/Source/Engine/Blueprints/BestPractices/BestPractices.INT.udn
@@ -2,13 +2,13 @@ Availability:Public
 Title:Blueprint Best Practices
 Crumbs:%ROOT%, Engine/Blueprints
 Description:Tips and tricks to help you make decisions about when to use Blueprints and how best to set them up.
-
+version: 4.9
 
 [TOC(start:2)]
 
-## Overview
 
-You can do many things with Blueprints. Everything from making a small game or procedural content tools, to prototyping new functionality, to tweaking and polishing things made by programmers,
+
+You can do many things with **Blueprints**. Everything from making a small game or procedural content tools, to prototyping new functionality, to tweaking and polishing things made by programmers,
 is possible with the Blueprints visual scripting system.
 
 However, there are certain things that will impact your performance more if they are done in Blueprints. 
@@ -39,11 +39,11 @@ operate on Actors (say, calling Set Brightness on a light), and paste them into
 ## Class Variables vs. Local Variables
 
 When you're working in your Blueprint, you have the ability to add [variables](Engine\Blueprints\UserGuide\Variables) using the [](Engine\Blueprints\Editor\UIComponents\MyBlueprint)
-tab. When you're inside a function or macro, though, you'll see an additional section of the **My Blueprint** tab for **Local Variables**.
+tab. When you're inside a function though, you'll see an additional section of the **My Blueprint** tab for **Local Variables**.
 
 ![](VariablesLocal.png)
 
-Local variables are scoped, which means that they only exist where you define them.  So, a local variable in a function or macro is only visible to that function, and not to other functions, or the 
+Local variables are scoped, which means that they only exist where you define them.  So, a local variable in a function is only visible to that function, and not to other functions, or the 
 Event Graph.  This helps to reduce clutter for things that are only relevant within the context of that function.  You can consider local variables a "scratch pad" to work with in a function.  You 
 can use them to help perform what the function needs to, but they get tossed away after the function has completed.  
 
@@ -52,7 +52,7 @@ that need to be public so you can access them from other Blueprints, should be c
 
 ## Functions vs. Macros
 
-Functions and macros both accomplish the same goal externally:  they let you send input into a node, something happens to it, and then outputs come out.  They accomplish these in two different ways, 
+**Functions** and **Macros** both accomplish the same goal externally:  they let you send input into a node, something happens to it, and then outputs come out.  They accomplish these in two different ways, 
 but they have some similarities:
 
 -	Both have a central place that you can edit the functionality of what the node does (macros in their macro graph, and functions in their function graph), and every node that calls the 
