#UE4doc #docs:
Added engine version field

diff --git a/Engine/Documentation/Source/Engine/Blueprints/BP_HowTo/BlueprintHowTo.INT.udn b/Engine/Documentation/Source/Engine/Blueprints/BP_HowTo/BlueprintHowTo.INT.udn
--- a/Engine/Documentation/Source/Engine/Blueprints/BP_HowTo/BlueprintHowTo.INT.udn
+++ b/Engine/Documentation/Source/Engine/Blueprints/BP_HowTo/BlueprintHowTo.INT.udn
@@ -5,6 +5,7 @@ Description:The Blueprint How To page provides several short step-by-step guides
 Related: Engine/Blueprints
 Related: Engine/Blueprints/UserGuide
 Related: Gameplay
+version: 4.9
 
 
 If you are new to [](Engine/Blueprints) in Unreal Engine 4, these Blueprint How To's provide information on the core aspects needed to work with Blueprints including Creating Blueprints, Adding and Connecting Nodes, Creating Variables and others as seen in the tables below. 
