#UE4doc #docs - Added Engine Version Field.

diff --git a/Engine/Documentation/Source/Engine/Physics/Constraints/DampingAndFriction/DampingAndFriction.INT.udn b/Engine/Documentation/Source/Engine/Physics/Constraints/DampingAndFriction/DampingAndFriction.INT.udn
--- a/Engine/Documentation/Source/Engine/Physics/Constraints/DampingAndFriction/DampingAndFriction.INT.udn
+++ b/Engine/Documentation/Source/Engine/Physics/Constraints/DampingAndFriction/DampingAndFriction.INT.udn
@@ -2,6 +2,7 @@ Availability: Public
 Title:Physics Damping
 Crumbs: %ROOT%, Engine, Engine/Physics
 Description:Properties and methods used to cause physical objects to resist movement.
+Version: 4.9
 
 
 [TOC(start:2)]
@@ -48,10 +49,10 @@ separate properties, the properties all have the same effect on the Motions they
 The two values work together to create different effects. A high Stiffness and low Damping will look like a hard limit, while having them both high will appear to be move through 
 tar, coming to a stop gently but quickly. A good deal of experimentation is necessary to find the right values for your specific needs.
 
-For more information on Constraints, see: [](Engine/Physics/Constraints/ConstraintsUserGuide)
+For more information on Constraints, see: **[](Engine/Physics/Constraints/ConstraintsUserGuide)**
 
 Or for a reference on the properties available to Physics Constraints,
- see: [](Engine/Physics\Constraints\ConstraintsReference)
+ see: **[](Engine/Physics\Constraints\ConstraintsReference)**
 
 <!--
 ## PhysicalMaterial Damping
