#UE4doc #docs - Added Engine Version Field. Filled in blank fields that didn't describe the properties listed.

diff --git a/Engine/Documentation/Source/Engine/Physics/PhAT/Properties/Properties.INT.udn b/Engine/Documentation/Source/Engine/Physics/PhAT/Properties/Properties.INT.udn
--- a/Engine/Documentation/Source/Engine/Physics/PhAT/Properties/Properties.INT.udn
+++ b/Engine/Documentation/Source/Engine/Physics/PhAT/Properties/Properties.INT.udn
@@ -2,6 +2,7 @@ Availability: Public
 Title:Physics Asset Properties Reference
 Crumbs: %ROOT%, Engine, Engine/Physics, Engine/Physics/PhAT
 Description:A listing of properties and descriptions for Physics Bodies, Physics Constraints, and Physics Assets accessible in the Physics Asset Tool.
+Version: 4.9
 
 [VAR:Topic]
 [OBJECT:Topic]
@@ -87,18 +88,17 @@ Description:A listing of properties and descriptions for Physics Bodies, Physics
 | **Angular Damping** | %doxygen:FBodyInstance::AngularDamping% |
 | **Enable Gravity** | %doxygen:FBodyInstance::bEnableGravity% |
 | **Physics Type** | %doxygen:UBodySetup::PhysicsType% |
-| **Double Sided Geometry** |  |
+| **Double Sided Geometry** | If true, the physics triangle mesh will use double sided faces when doing scene queries. Thsi is useful for planes and single sided meshes that need traces to work on both sides. |
 | **Simple Collision Physical Material** | %doxygen:UBodySetup::PhysMaterial% |
 | **Start Awake** | %doxygen:FBodyInstance::bStartAwake% |
 | **Center of Mass Offset** | %doxygen:FBodyInstance::COMNudge% |
 | **Mass Scale** | %doxygen:FBodyInstance::MassScale% |
-| **Max Angular Velocity** |  |
+| **Max Angular Velocity** | The maximum angular velocity for this instance. |
 | **Use Async Scene** | %doxygen:FBodyInstance::bUseAsyncScene% |
-| **Override Max Depenetration Velocity** |  |
-| **Max Depenetration Velocity** |  |
-| **Override Walkable Slope on Instance** |  |
+| **Max Depenetration Velocity** | The maximum velocity used to depenetrate this object.  |
+| **Override Walkable Slope on Instance** | Overrides the Walkable Slope settings for this physics body. Also, allows for the angle at which the slope is determined to be walkable to be customized.  |
 | **Walkable Slope Override** | %doxygen:UBodySetup::WalkableSlopeOverride% |
-| **Walkable Slope Behavior** | %doxygen:FWalkableSlopeOverride::WalkableSlopeBehavior% |
+| **Walkable Slope Behavior** | Behavior of this surface (whether we affect the walkable slope). Determines how the Walkable Slope Angle operates on the existing walkable angle, whether it increases or decreases it. |
 | **Walkable Slope Angle** | %doxygen:FWalkableSlopeOverride::WalkableSlopeAngle% |
 | **Sleep Family** | %doxygen:FBodyInstance::SleepFamily% |
 | **Position Solver Iteration Count** | %doxygen:FBodyInstance::PositionSolverIterationCount% |
@@ -117,7 +117,7 @@ Description:A listing of properties and descriptions for Physics Bodies, Physics
 | **Joint Name** | %doxygen:FConstraintInstance::JointName% |
 | **Constraint Bone 1** | %doxygen:FConstraintInstance::ConstraintBone1% |
 | **Constraint Bone 2** | %doxygen:FConstraintInstance::ConstraintBone2% |
-| **Disable Collision** | %doxygen:FConstraintInstance::bDisableCollision% |
+| **Disable Collision** | This will disable collision between the two Physics Bodies this Physics Constraint affects. This is most often used in the Physics Asset Tool when it is desirable to get Physics Bodies to match as close as possible to the SkeletalMesh but not collide with each other at the joints where a little overlap is desirable.  |
 |[REGION:tablesection]Linear Limits[/REGION]||
 | **Linear XMotion** | %doxygen:FConstraintInstance::LinearXMotion% |
 | **Linear YMotion** | %doxygen:FConstraintInstance::LinearYMotion% |
