#UE4doc #docs - Added Engine Version Field.

diff --git a/Engine/Documentation/Source/Engine/Physics/PhAT/PhAT.INT.udn b/Engine/Documentation/Source/Engine/Physics/PhAT/PhAT.INT.udn
--- a/Engine/Documentation/Source/Engine/Physics/PhAT/PhAT.INT.udn
+++ b/Engine/Documentation/Source/Engine/Physics/PhAT/PhAT.INT.udn
@@ -2,6 +2,8 @@ Availability: Public
 Title:Physics Asset Tool (PhAT)
 Crumbs: %ROOT%, Engine, Engine/Physics
 Description:Tool to set up physics bodies and joints to use for physical simulation and collision for SkeletalMeshes.
+Version: 4.9
+
 
 [VAR:Topic]
 [OBJECT:Topic]
@@ -28,8 +30,6 @@ Description:Tool to set up physics bodies and joints to use for physical simulat
 [/REGION]
 
 
-## Overview
-
 The Unreal **Physics Asset Tool** (**PhAT**) is an integrated editor in the Unreal Editor designed specifically to manipulate Physics Assets for Skeletal Meshes.
 
 ## Opening PhAT
@@ -53,7 +53,6 @@ The PhAT Physics Editor can be opened by **double-clicking** any Physics Asset o
 
 
 
-
 <!--
 
 ## Opening PhAT
