#UE4doc #docs - Filled in two missing properties.

diff --git a/Engine/Documentation/Source/Engine/Physics/PhAT/Properties/Properties.INT.udn b/Engine/Documentation/Source/Engine/Physics/PhAT/Properties/Properties.INT.udn
--- a/Engine/Documentation/Source/Engine/Physics/PhAT/Properties/Properties.INT.udn
+++ b/Engine/Documentation/Source/Engine/Physics/PhAT/Properties/Properties.INT.udn
@@ -58,8 +58,8 @@ Version: 4.9
 |[REGION:tablesection]Simulation[/REGION]||
 |**Floor Gap** | %doxygen:UPhATSimOptions::FloorGap% |
 |**Grav Scale** | %doxygen:UPhATSimOptions::GravScale% |
-|**Max FPS** | |
-|**Time Dilation** |  |
+|**Max FPS** | Max FPS for simulation in PhAT. This is helpful for targeting the same FPS as your game. -1 means this is disabled.  |
+|**Time Dilation** | This will dilate the time by scale.  |
 |[REGION:tablesection]Mouse Spring[/REGION]||
 |**Handle Linear Damping** | %doxygen:UPhATSimOptions::HandleLinearDamping% |
 |**Handle Linear Stiffness** | %doxygen:UPhATSimOptions::HandleLinearStiffness% |
