#UE4 Docs: Updating the foliage docs to remove old images and text that was misleading.

diff --git a/Engine/Documentation/Source/Engine/Foliage/Foliage.INT.udn b/Engine/Documentation/Source/Engine/Foliage/Foliage.INT.udn
--- a/Engine/Documentation/Source/Engine/Foliage/Foliage.INT.udn
+++ b/Engine/Documentation/Source/Engine/Foliage/Foliage.INT.udn
@@ -2,6 +2,8 @@ Availability: Public
 Title:Foliage Instanced Meshes
 Crumbs:%ROOT%, Engine
 Description:System for rendering instanced meshes on the surfaces of other geometry for use as foliage or other ground cover effects.
+Version: 4.9
+SkillLevel:Intermediate
 
 [VAR:Topic]
 [OBJECT:Topic]
@@ -47,11 +49,12 @@ Description:System for rendering instanced meshes on the surfaces of other geome
 ![Foliage Tool](Foliage_Banner.png)(w:900)
 [/REGION]
 
-[TOC(start:2 end:2)]
+[TOC(start:2 end:3)]
 
 [EXCERPT:Intro]
 
-The Foliage system allows you to quickly paint or erase sets of Static Meshes on Landscape Actors, other Static Mesh Actors, and legacy terrain. These meshes are automatically grouped together into batches that are rendered using hardware instancing, meaning that many instances can be rendered with only a single draw call.
+The Foliage Instanced Mesh system allows you to quickly paint or erase sets of Static Meshes on Landscape Actors, other Static Mesh Actors, and BSP. 
+These meshes are automatically grouped together into batches that are rendered using hardware instancing, meaning that many instances can be rendered with only a single draw call.
 
 [/EXCERPT:Intro]
 
@@ -68,51 +71,40 @@ The Foliage Edit Mode window appears in the **Modes** panel.
 ![FoliageWindowInitial.png](FoliageWindowInitial.png)
 
 
-
 ## Mesh List
 
-At first, there will be no foliage meshes listed. Before you can paint foliage, you must drag a Static Mesh from the **Content Browser** directly onto the black bar that says Drag Static Meshes from the Content Browser into this area.
+The Mesh List is the area in which Static Meshes that are to be used as foliage are placed. 
+To add a Static Mesh to the Mesh List you must drag the Static Mesh from the Content Browser into the area that says **Drop Foliage Here**.
 
 ![](Foliage_TextBar.png)
 
 [REGION:tip]
-Do not try dropping the first Static Mesh into the empty space underneath the black bar; you must drop it onto the black bar itself.
+The yellow Foliage Type will only be displayed when you are trying to add the correct actor type. If you do not see this then what you are trying to add will not work.
 [/REGION]
 
-[REGION:note]
-After you have added the first Static Mesh to the list, you can add more Static Meshes by dragging them into the Foliage Edit Mode window, and dropping them on top of existing meshes in the list.
-[/REGION]
+To adjust the properties of a Static Mesh that is used as foliage first select the Static Mesh in the Mesh List and then press the small **Magnifying Glass** icon in the lower right hand corner of the foliage tool to expose the selected Static Meshes details. 
 
-<!----
-![FoliageWindowDrag.png](FoliageWindowDrag.png)
----->
+![](Expose_Instanc_Options.png)
 
-A Static Mesh in the mesh list can be displayed in one of three modes:
-* **Hide Details** - Hides all parameters and properties, to make the entire list easier to see.
-	![FoliageWindowAfterDrag.png](FoliageWindowAfterDrag.png)
-* **Show Painting Settings** - Displays parameters that affect how the mesh will be placed in the level.  (This is the default mode when a Static Mesh is added.)
+Once pressed you should see something that looks similar to the image below.
 
-	![Moliage Painting Settings](Foliage_MeshParams.png)
-* **Show Instance Settings** - Displays parameters that control the instances that are already placed in the level.
-	![Mesh Instance Settings](Foliage_ShowInstance.png)
+![](Foliage_Mesh_Options.png)
 
-When the mesh is in either the Paint Settings or Instance Settings mode, three buttons also appear above the mesh icon.
 
-| Item | Description |
-| --- | --- |
-| ![](Shared/Icons/icon_assign_left_16x.png) | This button replaces all instances of the mesh already placed in the current level with the Static Mesh that is currently selected in the **Content Browser**. |
-| ![](Shared/Icons/icon_delete.png)(w:16) | This button removes the mesh from the mesh list as well as all instances of this mesh already placed in the current level. |
-| ![](Shared/Icons/icon_MeshPaint_Find_16x.png) | This button locates the Static Mesh in the **Content Browser**. |
+### Selecting Foliage Meshes in the Mesh List
 
-### Selected Meshes in the Mesh List
+Foliage Meshes in the Mesh list can be enabled or disabled by hovering the mouse over the mesh in the list and then checking or unchecking the checkmark box in the upper left-hand corner. 
 
-Meshes can be selected and deselected by clicking on them in the **Meshes** list (orange is selected). When creating or erasing meshes in the Level Editor viewport, only the meshes currently selected in the Meshes list will be affected. Foliage instances of other Static Meshes already in the level will not be affected in any way, and no instances of the deselected Meshes list items will be added to the level.
+![](Foliage_Selecting_Meshes.png)
 
-The following example shows two meshes selected. The SM_Rock mesh will not be painted or erased from the level as it is not selected.
+You can also select Foliage Meshes by clicking on them in the Mesh list. 
+This will put an orange highlight around them and then in the option panel below you can adjust various properties. 
 
-![Two Meshes Selected](Foliage_MeshList.png)
+![](Foliage_Selected_Meshes.png)
 
-#### Brush Settings
+## Brush Settings
+
+![](Foliage_Brush_Settings.png)
 
 Above the **Meshes** list, the Mesh Edit Tool displays properties related to the foliage tool that is currently selected. (Not every setting is displayed with every foliage tool.)
 
@@ -121,11 +113,14 @@ Above the **Meshes** list, the Mesh Edit Tool displays properties related to the
 | **Brush Size** | The size of the brush in Unreal units. | 
 | **Paint Density** |  [INCLUDE:Engine/Foliage#PaintDensity] |
 | **Erase Density** |  [INCLUDE:Engine/Foliage#EraseDensity] |
-| **Filter** | [INCLUDE:Engine/Foliage#Filter] |
+
+The checkmark boxes you see with the names **Landscape**, **Static Meshes**, **BSP** and **Translucent** determine which type of Actor type the Foliage meshes will be placed on when the tool is used. 
+Checking a box will indicate that a Foliage mesh should be added should you use the tool on that Actor type.
+Un-checking a box will indicate that a foliage mesh should not be added to that Actor type.
 
 <!---
 [EXCERPT:PaintDensity]
-The target density at which you will add instances when using **Ctrl+click**.
+The target density at which you will add instances when using **Left Mouse** button.
 * This is a value from 0 to 1 where 1 paint mesh instances at the maximum density listed in the Mesh Properties for each mesh (see below).
 * If the density of meshes in the brush is already greater than this, no meshes will be added.
 [/EXCERPT:PaintDensity]
@@ -153,8 +148,6 @@ The kinds of objects you want to place instances on.
 Tablet pressure sensitivity is also supported. The pen pressure will modulate the number of foliage mesh instances added or removed, after all other parameters are considered.
 [/REGION]
 
-
-
 ## Foliage Tools
 There are five tools available by selecting the appropriate button in the Toolbar at the top of the foliage window. 
 
@@ -167,72 +160,71 @@ There are five tools available by selecting the appropriate button in the Toolba
 | ![Fill Tool](ToolFill.png) | **Fill Tool**, used to determine the number of meshes to place with the Paint tool. |
 
 
-
 ### Paint Tool
 
 ![Paint Tool](ToolPaint.png)
 
 The Paint tool is Foliage Edit Mode's default tool. You can use it to paint instances of the Static Meshes in the Meshes list in your level.
-
 When Foliage Edit Mode is active, a transparent sphere brush is drawn in the level indicating where the foliage brush will operate.
+In the image below you can see an example of what the Foliage paint brush looks like.
+
+![](Foliage_Paint_Tool_Gimzo.png)
 
 **To add foliage to the brush area:**  
-* Holding down **Ctrl**, click where you want to add foliage.  All currently selected meshes in the mesh list will be added according to their current parameters and settings.
+* Holding down **Left Mouse** button and clicking where you want to add foliage. All currently selected meshes in the mesh list will be added according to their current parameters and settings.
 
 When painting meshes, the engine performs a number of line traces inside the sphere parallel to your viewing direction, meaning any surface you can see inside the sphere will be a potential target for foliage instances.
 
 **To erase foliage from the brush area:**  
-* Holding down **Ctrl+Shift**, click where you want to remove foliage. Note that only instances of meshes that are currently selected in the mesh list will be erased.
+* Holding down **Left Mouse Button + Shift**, click where you want to remove foliage. Note that only instances of meshes that are currently selected in the mesh list will be erased.
 
 
-#### Paint Settings
+### Paint Settings
 
-Each mesh in the Mesh List has a number of properties that control how instances are placed for this mesh when painting. The properties are accessible by pressing the middle paint brush icon.
+Each Static mesh in the Mesh List has a number of properties that control how instances are placed, scale or rotated when the selected Static mesh is painted. 
+The properties are accessible by selecting the Static Mesh or Static Meshes you want in the Mesh List and then adjusting the properties displayed below.
 
 ![Paint Settings](Foliage_PaintSettings.png)
 
+Here is a high level description of what each section is used for.
+
 | **Property** | **Description** |
 | --- | --- |
-| **Density / 1Kuu&sup2;** | The number of instances that will be placed for this mesh, for each 1000x1000 Unreal unit area. |
-| **Radius** | Minimum distance between instances of this mesh.  |
-| **Align to Normal** | Whether or not the foliage mesh should be oriented vertically or follow the direction of the surface it is placed upon. |
-| **Max Angle** | If non-zero, this specifies the maximum angle to follow when aligning to the normal. For example, if ground is 30&deg; slope, and you set the Max Angle to 10&deg;, the foliage will try to follow the direction of the ground but it placed at 10&deg; to the vertical.  If you place an instance on 5&deg; surface, the angle is less than 10&deg; so the foliage will be oriented to match the 5&deg; slope. |
-| **Random Yaw** | Whether or not meshes will be placed rotated randomly along the Z axis. |
-| **Uniform Scale** | Whether or not you wish to scale instances' X, Y and Z scales independently. |
-| **Scale Min/Max** | Instances will be randomly scaled between the minimum and maximum scale range. |
-| **Random Pitch &plusmn;** | Adds variation up to &plusmn; the angle you specify, so the instances are not all pointing exactly the same direction. |
-| **Ground Slope** | If non-zero, if the surface you are painting on is steeper than this angle, no foliage instances will be placed. A negative value will place foliage instances only on surfaces steeper than the angle specified. |
-| **Height Min/Max** | If specified, if the World space Z (height) value is outside this range, no instances will be applied. |
-| **Landscape Layer** | If entered, foliage instances will be painted proportional in density to the Landscape layer specified, after all the other parameters have been applied. |
-
+| **Mesh** | This is where you input which type of Static Mesh should be used. |
+| **Painting** | This is where you can adjust properties like the minimum distance between foliage meshes, the density at which meshes will be added and if the Foliage meshes should be scaled.|
+| **Placement** | This controls how the Foliage meshes will be placed when painted in the level.|
+| **Instance Settings** | This allows you to adjust the behavior of this instance and not use properties derived from the Static Mesh in the Content Browser.|
 
 
 ### Reapply Tool
 
 ![Reapply Tool](ToolReapply.png)
 
-The Reapply tool allows you to selectively change parameters for instances already placed in the world. The settings available are very similar to those of the Paint tool, but there is an extra checkbox to the left, to choose which parameters to reapply:
+The Reapply tool allows you to selectively change certain parameters for Foliage meshes already placed in the world. 
+When you paint over those Foliage meshes with the Reapply tool active the Foliage meshes that are selected in the Mesh List will be updated to reflect the changes made in the Reapply tool.
 
-![](Foliage_Reapply.png)(w:400)
+![](Foliage_Reapply.png)
 
-If a checkbox is set, the corresponding parameter will be reapplied when painting the brush over existing instances. Most settings are the same as when using the Paint Tool, however there are a few differences:
+If a checkbox is set, the corresponding parameter will be reapplied when painting the brush over existing instances. 
+Most settings are the same as when using the Paint Tool, however there are a few differences:
 
 | Item | Description |
 | --- | --- |
-| **Density Adjust** |This is a density multiplier for existing instances. If you set the slider to 2.0, the number of instances in a particular area will be doubled. A value of 0.5 will leave only 50% of the instances remaining after painting.|
-| **Ground Slope** |This will remove instances that do not meet the ground slope criteria. No new instances will be added.|
-| **Height Min/Max** |This will remove instances that do not meet the height range criteria. No new instances will be added.|
-| **Landscape Layer** |This will remove instances in proportion to the specified layer density. No new instances will be added.|
+| **Mesh** | This is where you input which type of Static Mesh should be used. |
+| **Painting** | This is where you override what should happen when you are adding Foliage meshes to the level. |
+| **Placement** | This is where you override how Foliage meshes are placed in the level. |
+
 
 
 ### Select Tool
 ![Select Tool](ToolSelect.png)
 
-When the Select tool is active, you can select individual instances of foliage meshes by clicking on them. **Ctrl+click** will multi-select.
+When the Select tool is active, you can select individual instances of foliage meshes by clicking on them. 
+You can select multiple Foliage meshes at a time by holding down the **Ctrl** button while clicking on Foliage meshes in the level.
 
-![Selected Instance](Foliage_SelectInstance.png)(w:600)
+![Selected Instance](Foliage_SelectInstance.png)
 
-Once instances are selected, a number of actions are available:
+Once instances have been selected you can perform the following actions on them.
 
 | **Action** | **Result** |
 | --- | --- |
@@ -241,42 +233,37 @@ Once instances are selected, a number of actions are available:
 | **Delete key** | The selected instances will be deleted. |
 | **End key** | The selected instances will attempt to snap to the floor. They will be aligned to the normal if originally placed with this setting enabled. |
 
+Using the various **Select** buttons you can select all Foliage Mesh Actors, select any invalid Foliage Meshes or Deselect anything that is selected.
 
+![](Foliage_Select_Options.png)
 
 ### Lasso Tool
 
 ![Lasso Tool](ToolPaintSelect.png)
 
-The Lasso tool allows you to select a number of instances simultaneously by using the sphere brush that is also used with the Paint tool. The Meshes list selection and filter settings apply, just as they do for regular mesh painting. Holding down **Shift** while painting deselects instances. Once you have a selection, you can switch to the regular Select tool and perform some of the actions listed above.
+The Lasso tool allows you to select a number large or small number of Foliage meshes simultaneously by using the sphere brush that is also used with the Paint tool. 
+You can use the filter settings to ensure that you do not select Foliage meshes that are placed on certain objects. 
+To select Foliage meshes hold down the Right Mouse button and move the mouse around the scene.
+To de-select Foliage meshes hold down the **Shift** key while painting. 
+Once you have a selection, you can switch to the regular Select tool and do things like duplicate, move or remove Foliage meshes.
 
-![Lasso Tool](Foliage_Lasso.png)(w:600)
+![Lasso Tool](Foliage_Lasso.png)
 
 ### Fill Tool
 
 ![Fill Tool](ToolFill.png)
 
-The Fill tool is used to cover an entire Static Mesh Actor in your level with the selected foliage meshes, with the fill density and other properties set according to the paint settings.
-
-To use the Fill tool, **Ctrl+click** on the Static Mesh you want covered with foliage.
-
-![Foliage Tool](Foliage_Fill.png)(w:600)
+The Fill tool is used to cover an entire Static Mesh Actor in your level with the selected Foliage meshes. 
+The Fill tool works a lot like the Painting tool with the key difference being that you can cover an entire Actor in Foliage in a single click.
+To use the Fill tool, press the **Left Mouse Button** on the Static Mesh you want covered with foliage. 
+The more mouse button clicks you make the more Foliage will be added.
+TO remove Foliage hold down the **Shift** key and use the **Left Mouse Button** to click on the Static Mesh you want to remove Foliage from.
 
+![Foliage Tool](Foliage_Fill.png)
 
-## Instance Settings Mode
-
-You can use Instance Settings Mode to modify mesh instances that you have already placed in your level with the Foliage tool.
-
-![Instance Settings Mode](Foliage_Instance.png)(w:500)
-
-| **Property** | **Description** |
-| --- | --- |
-| **Instance Count** | This shows the number of instances of this mesh currently placed in the current streaming level. |
-| **Cluster Count** | This shows the number of clusters used to render instances of this mesh in the current streaming level. Groups of instances are rendered together in a single draw call, called a cluster, to improve rendering efficiency. The cluster allocation is determined automatically based on the number of instances already in the cluster and the radius of the cluster in world space. The following two parameters can adjust this. |
-| **Max Instances Per Cluster** | This sets the maximum number of instances per cluster. Remember that all the instances in a single cluster are always rendered in a single draw call with the total number of triangles equal to the polygons per mesh multiplied by the instances per cluster. The ideal number depends on the polygon count of the mesh. |
-| **Max Cluster Radius** | This is the maximum size a cluster can grow before instances are allocated to a new cluster. Reducing this number increases the number of clusters but improves occlusion as the cluster has a smaller bounding box. |
-| **Start Cull Distance** | This is the distance in world units at which the instances will begin to fade out. Note that this needs to be set up in the Material. See the "Culling" section below. |
-| **End Cull Distance** | This is the distance in world units at which the instances will be completely culled. If the Material is not set up to fade individual instances, entire clusters will disappear and reappear together. See the "Culling" section below. |
-
+[region:tip]
+If you are having issues getting Foliage meshes to erase with the Fill tool make sure that you are clicking on the Actor that they are placed and not on the Foliage meshes that were added.
+[/region]
 
 
 ## Culling
