#UE4 Docs: Applying feedback from the SME.

diff --git a/Engine/Documentation/Source/Engine/Rendering/Materials/MaterialProperties/LightingModels/LightingModels.INT.udn b/Engine/Documentation/Source/Engine/Rendering/Materials/MaterialProperties/LightingModels/LightingModels.INT.udn
--- a/Engine/Documentation/Source/Engine/Rendering/Materials/MaterialProperties/LightingModels/LightingModels.INT.udn
+++ b/Engine/Documentation/Source/Engine/Rendering/Materials/MaterialProperties/LightingModels/LightingModels.INT.udn
@@ -143,9 +143,9 @@ When using the Subsurface Profile Shading Model, you are given access to the fol
 
 The **Two Sided Foliage** shading model was introduced to help give foliage a more natural and uniform look when being lit. 
 When using the Default Lit shading model for foliage parts of the foliage can look almost black.
-The Default Lit shading model also does not simulate any kind of light transmission which is key to having believable looking foliage. 
-The Two Sided Foliage shading model helps to eliminate these issues by simulating a very cheap Subsurface scattering to give the illusion that light is being transferred through the foliage. 
-This also helps to eliminate the dark parts of the foliage and give the foliage a more uniform look to it's lighting.
+The Default Lit shading model also does not simulate any kind of light transmission which is key to having believable looking foliage.
+The Two Sided Foliage shading model helps to eliminate these issues by using a variation of the Subsurface scattering lighting model that accounts for objects that are very thin helping to give the illusion that light is being transferred through the foliage.
+This also helps to eliminate the dark parts of the foliage and give the foliage a more uniform look to it's lighting. 
 
 [OBJECT:ComparisonSlider]
 	[PARAM:before]
