#UE4doc #docs - Updated PhAT image. Now shows updated Details panel options. Added Engine Version Field.

diff --git a/Engine/Documentation/Source/Engine/Physics/PhAT/Reference/PhAT_Reference.INT.udn b/Engine/Documentation/Source/Engine/Physics/PhAT/Reference/PhAT_Reference.INT.udn
--- a/Engine/Documentation/Source/Engine/Physics/PhAT/Reference/PhAT_Reference.INT.udn
+++ b/Engine/Documentation/Source/Engine/Physics/PhAT/Reference/PhAT_Reference.INT.udn
@@ -2,6 +2,7 @@ Availability: Public
 Title:Physics Asset Tool Reference
 Crumbs: %ROOT%, Engine, Engine/Physics, Engine/Physics/PhAT
 Description:Tool to set up physics bodies and joints to use for physical simulation and collision for SkeletalMeshes.
+Version: 4.9
 
 [VAR:Topic]
 [OBJECT:Topic]
