#UE4doc #docs - Added Engine Version Field. Updated header image to match latest UI settings. Added missing properties to the descriptions.

diff --git a/Engine/Documentation/Source/Engine/Physics/Collision/Collision.INT.udn b/Engine/Documentation/Source/Engine/Physics/Collision/Collision.INT.udn
--- a/Engine/Documentation/Source/Engine/Physics/Collision/Collision.INT.udn
+++ b/Engine/Documentation/Source/Engine/Physics/Collision/Collision.INT.udn
@@ -2,6 +2,7 @@ Availability: Public
 Title:Collision Responses
 Crumbs: %ROOT%, Engine, Engine/Physics
 Description:All about collision, collision responses, and collision presets.
+Version: 4.9
 
 [VAR:Topic]
 [OBJECT:Topic]
@@ -38,13 +39,14 @@ block it or ignore it based on _their_ Trace Responses.
 ![](COL_collisionBasics.png)
 
 In the properties list, there are a large number of physics related key words.  The two that are not really clear are **Simple** and **Complex Collision**. In short, 
-Simple Collision shapes are your [boxes, KDOP, spheres, capsules, and convex shapes](Engine/Physics\Collision\Reference), all of which can be generated in the Static
+Simple Collision shapes are your [**boxes, KDOP, spheres, capsules, and convex shapes**](Engine/Physics\Collision\Reference), all of which can be generated in the Static
 Mesh editor. Complex Collision is per-poly, and it is never used for an Actor that is simulating physics (it will just fall through the world).
 
 
 | **Property** | **Description** |
 | -------- | ----------- |
 | **Simulation Generates Hit Events** | This flag indicates if **Event Hit** and **On Component Hit** should be called for physics collisions involving this PrimitiveComponent. In more simple terms, it lets this object "know" it has actually had a physics collision. However, it is just a notification, and a dynamic object set to collide with the static environment will still do so regardless of this setting; it just will not fire those events. |
+| **Phys Material Override** | This allows you to override the Phsyical Material to use for simple collision on this body. |
 | **Generate Overlap Events** | If you want an object to generate overlap events, like **Event Actor Begin Overlap** or **Event Actor End Overlap**, then this flag needs to be **true**. Further, for it to be of any use, the object in question will need to be able to overlap another. |
 | [REGION:tablesection] &nbsp; [/REGION]||
 | **Collision Presets** | These are a collection of presets for the **Collision Responses** below. There is an INI file you can edit to change or add profiles. **To edit anything below this point, you need to set this to "Custom..."**|
@@ -95,7 +97,7 @@ There are a few rules to keep in mind with how collisions are handled:
 * Overlap events can be generated even if an object Blocks another, especially if traveling at high speeds.
 	* It is not recommended for an object to have both collision and overlap events. Though possible, there is much that needs manual handling.
 * If one object is set to ignore and the other is set to overlap, no overlap events will be fired.
-* The Object Responses listed are just the defaults and have no actual bearing on what an object _is_. See below for more information on setting up custom [Object Types and Trace Channels](#customobjectchannelsandtracechannels).
+* The Object Responses listed are just the defaults and have no actual bearing on what an object _is_. See below for more information on setting up custom [**Object Types and Trace Channels**](#customobjectchannelsandtracechannels).
 
 For purposes of testing levels and looking around your worlds:
 
