#UE4doc #docs - Removed Overview header. Added Engine Version Field.

diff --git a/Engine/Documentation/Source/Engine/Physics/Collision/Reference/Reference.INT.udn b/Engine/Documentation/Source/Engine/Physics/Collision/Reference/Reference.INT.udn
--- a/Engine/Documentation/Source/Engine/Physics/Collision/Reference/Reference.INT.udn
+++ b/Engine/Documentation/Source/Engine/Physics/Collision/Reference/Reference.INT.udn
@@ -2,6 +2,8 @@ Availability: Public
 Title:Static Mesh Collision Reference
 Crumbs: %ROOT%, Engine, Engine/Physics, Engine/Physics/Collision
 Description:Content guide to creating and setting up collision geometry.
+Version: 4.9
+SkillLevel: Beginner
 
 [VAR:Topic]
 [OBJECT:Topic]
@@ -25,7 +27,6 @@ Description:Content guide to creating and setting up collision geometry.
 
 [TOC(start:2 end:2)]
 
-## Overview
 
 As graphical detail has increased, it is important to ensure that collision within the game remains simple. This is because of two reasons; it is much faster to calculate collision 
 against a few polygons and it is a much smoother experience for the player as then the player does not get snagged on a small piece geometry. 
