#UE4doc #docs - Added Engine Version field.

diff --git a/Engine/Documentation/Source/Engine/Physics/Constraints/PhysicsConstraints.INT.udn b/Engine/Documentation/Source/Engine/Physics/Constraints/PhysicsConstraints.INT.udn
--- a/Engine/Documentation/Source/Engine/Physics/Constraints/PhysicsConstraints.INT.udn
+++ b/Engine/Documentation/Source/Engine/Physics/Constraints/PhysicsConstraints.INT.udn
@@ -2,6 +2,7 @@ Availability: Public
 Title:Physics Constraints
 Crumbs: %ROOT%, Engine, Engine/Physics
 Description:Constraining physics objects together and to the world.
+Version: 4.9
 
 [OBJECT:TopicList]
 	[PARAM:icon]
