#UE4doc #docs - Added Skill Level field. Added Engine Version Field.

diff --git a/Engine/Documentation/Source/Engine/Physics/Constraints/ConstraintsUserGuide/ConstraintsUserGuide.INT.udn b/Engine/Documentation/Source/Engine/Physics/Constraints/ConstraintsUserGuide/ConstraintsUserGuide.INT.udn
--- a/Engine/Documentation/Source/Engine/Physics/Constraints/ConstraintsUserGuide/ConstraintsUserGuide.INT.udn
+++ b/Engine/Documentation/Source/Engine/Physics/Constraints/ConstraintsUserGuide/ConstraintsUserGuide.INT.udn
@@ -2,6 +2,8 @@ Availability: Public
 Title:PhysX Constraint User Guide
 Crumbs: %ROOT%, Engine, Engine/Physics
 Description:Physical joints used to connect two rigid body simulated objects together.
+SkillLevel: Beginner
+Version:  4.9
 
 [REGION:fullwidth raw]
 ![](Engine/Physics\physics_banner.png)
