#UE4doc #docs - Added Engine Version Field. Updated naming for some of the properties that have changed.

diff --git a/Engine/Documentation/Source/Engine/Physics/Constraints/ConstraintsReference/Constraints.INT.udn b/Engine/Documentation/Source/Engine/Physics/Constraints/ConstraintsReference/Constraints.INT.udn
--- a/Engine/Documentation/Source/Engine/Physics/Constraints/ConstraintsReference/Constraints.INT.udn
+++ b/Engine/Documentation/Source/Engine/Physics/Constraints/ConstraintsReference/Constraints.INT.udn
@@ -2,6 +2,7 @@ Availability: Public
 Title:Physics Constraint Reference
 Crumbs: %ROOT%, Engine, Engine/Physics
 Description:A reference for all the properties available on Physics Constraints in Unreal Editor.
+Version: 4.9
 
 [TOC(start:2)]
 
@@ -116,19 +117,20 @@ There are 3 ways to create PhysX Constraints: **Blueprint**, **Actor**, and in t
 | **Linear YPosition Drive** | %doxygen:FConstraintInstance::bLinearYPositionDrive% |
 | **Linear ZPosition Drive** | %doxygen:FConstraintInstance::bLinearZPositionDrive% |
 | **Linear Position Target** | %doxygen:FConstraintInstance::LinearPositionTarget% |
+| **Linear Position Strength** | Spring to apply to the linear drive.  |
 | **Linear Velocity Drive** | %doxygen:FConstraintInstance::bLinearVelocityDrive% |
 | **Linear Velocity Target** | %doxygen:FConstraintInstance::LinearVelocityTarget% |
-| **Linear Drive Spring** | %doxygen:FConstraintInstance::LinearDriveSpring% |
-| **Linear Drive Damping** | %doxygen:FConstraintInstance::LinearDriveDamping% |
-| **Linear Drive Force Limit** | %doxygen:FConstraintInstance::LinearDriveForceLimit% |
+| **Linear Velocity Strength** | Damping to apply to the linear drive. |
+| **Max Linear Force** | Limit to the force the linear drive can apply. |
 |[REGION:tablesection]Angular Motor[/REGION]||
 | **Angular Orientation Drive** | %doxygen:FConstraintInstance::bAngularOrientationDrive% |
 | **Angular Orientation Target** | %doxygen:FConstraintInstance::AngularOrientationTarget% |
+| **Angular Position Strength** | Spring value to apply for the angular drive. |
 | **Angular Velocity Drive** | %doxygen:FConstraintInstance::bAngularVelocityDrive% |
 | **Angular Velocity Target** | %doxygen:FConstraintInstance::AngularVelocityTarget% |
-| **Angular Drive Force Limit** | %doxygen:FConstraintInstance::AngularDriveForceLimit% |
-| **Angular Drive Spring** | %doxygen:FConstraintInstance::AngularDriveSpring% |
-| **Angular Drive Damping** | %doxygen:FConstraintInstance::AngularDriveDamping% |
+| **Angular Velocity Strength** | Damping value to apply for the angular drive. |
+| **Max Angular Force ** | Limit to the force the angular drive can apply. |
+
 
 <!--
 
