#UE4doc #docs - Added Engine Version Field.

diff --git a/Engine/Documentation/Source/Engine/Physics/PhysicalMaterials/PhysicalMaterials.INT.udn b/Engine/Documentation/Source/Engine/Physics/PhysicalMaterials/PhysicalMaterials.INT.udn
--- a/Engine/Documentation/Source/Engine/Physics/PhysicalMaterials/PhysicalMaterials.INT.udn
+++ b/Engine/Documentation/Source/Engine/Physics/PhysicalMaterials/PhysicalMaterials.INT.udn
@@ -2,6 +2,7 @@ Availability: Public
 Title:Physical Materials
 Crumbs: %ROOT%, Engine, Engine/Physics
 Description:Assets applied to physically simulated primitives directly or via materials used to configure and control physical properties used by the simulation.
+Version: 4.9
 
 [VAR:Topic]
 [OBJECT:Topic]
