#UE4 Docs: added skill and version tags

diff --git a/Engine/Documentation/Source/Engine/Blueprints/BP_HowTo/ActorReference/ActorReference.INT.udn b/Engine/Documentation/Source/Engine/Blueprints/BP_HowTo/ActorReference/ActorReference.INT.udn
--- a/Engine/Documentation/Source/Engine/Blueprints/BP_HowTo/ActorReference/ActorReference.INT.udn
+++ b/Engine/Documentation/Source/Engine/Blueprints/BP_HowTo/ActorReference/ActorReference.INT.udn
@@ -4,6 +4,10 @@ Crumbs: %ROOT%, Engine, Engine\Blueprints, Engine\Blueprints\BP_HowTo
 Description:Using Actor Reference to move different Actors around in a level.
 Related: Gameplay/HowTo/ReferenceAssets/Blueprints
 version: 4.9
+skilllevel:Intermediate
+type:howto
+parent:Engine/Blueprints/BP_HowTo
+order:1
 
 [INCLUDE:Gameplay\HowTo\ReferenceAssets#referenceoverview]
 
diff --git a/Engine/Documentation/Source/Engine/Blueprints/BP_HowTo/Casting/Casting.INT.udn b/Engine/Documentation/Source/Engine/Blueprints/BP_HowTo/Casting/Casting.INT.udn
--- a/Engine/Documentation/Source/Engine/Blueprints/BP_HowTo/Casting/Casting.INT.udn
+++ b/Engine/Documentation/Source/Engine/Blueprints/BP_HowTo/Casting/Casting.INT.udn
@@ -5,6 +5,8 @@ Description:Using Blueprint Casting to communicate from one Blueprint to multipl
 Related: Engine/Blueprints/UserGuide/CastNodes
 Related: Engine/Blueprints/UserGuide/BlueprintCommsUsage
 Related: Engine/Blueprints/UserGuide
+version:4.8
+skilllevel:Advanced
 
 This page covers how to use Blueprint Casting to communicate with different Blueprints and share information. In this example we allow the player to shoot an object, making it smaller each time it is hit. Each time the object is hit, it also notifies the Character Blueprint and increases the player's movement speed. 
 
diff --git a/Engine/Documentation/Source/Engine/Blueprints/BP_HowTo/CollapsingGraphs/CollapsingGraphs.INT.udn b/Engine/Documentation/Source/Engine/Blueprints/BP_HowTo/CollapsingGraphs/CollapsingGraphs.INT.udn
--- a/Engine/Documentation/Source/Engine/Blueprints/BP_HowTo/CollapsingGraphs/CollapsingGraphs.INT.udn
+++ b/Engine/Documentation/Source/Engine/Blueprints/BP_HowTo/CollapsingGraphs/CollapsingGraphs.INT.udn
@@ -5,6 +5,11 @@ Description:Collapsing nodes to a new Graph, a Function, or a Macro.
 Related: Engine/Blueprints/UserGuide
 Related: Engine/Blueprints/UserGuide/Functions
 Related: Engine/Blueprints/UserGuide/Macros
+version:4.8
+skilllevel:Beginner
+parent:Engine/Blueprints/BP_HowTo
+type:howto
+order:4
 
 This page covers several ways in which you can collapse your **Blueprint** graphs into a single node, into a **Function** or into a **Macro** which can help clean up any large graphs you may have by consolidating them into a linked graph. 
 
diff --git a/Engine/Documentation/Source/Engine/Blueprints/BP_HowTo/ConnectingNodes/ConnectingNodes.INT.udn b/Engine/Documentation/Source/Engine/Blueprints/BP_HowTo/ConnectingNodes/ConnectingNodes.INT.udn
--- a/Engine/Documentation/Source/Engine/Blueprints/BP_HowTo/ConnectingNodes/ConnectingNodes.INT.udn
+++ b/Engine/Documentation/Source/Engine/Blueprints/BP_HowTo/ConnectingNodes/ConnectingNodes.INT.udn
@@ -6,6 +6,11 @@ Related: Engine/Blueprints/BP_HowTo/PlacingNodes
 Related: Engine/Blueprints/BP_HowTo/BasicUsage
 Related: Engine/Blueprints/UserGuide
 Related: Engine/Blueprints/UserGuide/CheatSheet
+version:4.8
+skilllevel:Beginner
+parent:Engine/Blueprints/BP_HowTo
+order:2
+type:howto
 
 This page covers the various methods in which you can connect nodes together on your graphs in Blueprints. 
 
diff --git a/Engine/Documentation/Source/Engine/Blueprints/BP_HowTo/Debugging/Debugging.INT.udn b/Engine/Documentation/Source/Engine/Blueprints/BP_HowTo/Debugging/Debugging.INT.udn
--- a/Engine/Documentation/Source/Engine/Blueprints/BP_HowTo/Debugging/Debugging.INT.udn
+++ b/Engine/Documentation/Source/Engine/Blueprints/BP_HowTo/Debugging/Debugging.INT.udn
@@ -4,6 +4,11 @@ Crumbs: %ROOT%, Engine, Engine\Blueprints, Engine\Blueprints\BP_HowTo
 Description:Illustrates the methods on how to Debug your Blueprints. 
 Related: Engine/Blueprints/UserGuide/Debugging
 Related: Engine/Blueprints
+version:4.8
+skilllevel:Beginner
+parent:Engine/Blueprints/BP_HowTo
+order:5
+type:howto
 
 
 [TOC(start:2)]
diff --git a/Engine/Documentation/Source/Engine/Blueprints/BP_HowTo/DirectBlueprintComs/DirectBlueprintComs.INT.udn b/Engine/Documentation/Source/Engine/Blueprints/BP_HowTo/DirectBlueprintComs/DirectBlueprintComs.INT.udn
--- a/Engine/Documentation/Source/Engine/Blueprints/BP_HowTo/DirectBlueprintComs/DirectBlueprintComs.INT.udn
+++ b/Engine/Documentation/Source/Engine/Blueprints/BP_HowTo/DirectBlueprintComs/DirectBlueprintComs.INT.udn
@@ -5,6 +5,8 @@ Description:A demonstration of passing information directly from one Blueprint t
 Related: Engine/Blueprints/UserGuide/BlueprintComms
 Related: Engine/Blueprints/UserGuide/BlueprintCommsUsage
 Related: Engine/Blueprints/UserGuide
+version:4.8
+skilllevel:Advanced
 
 This page covers using [Direct Blueprint Communication](Engine/Blueprints/UserGuide/BlueprintComms) to communicate between two Blueprints and pass information between them. 
 
diff --git a/Engine/Documentation/Source/Engine/Blueprints/BP_HowTo/EventDispatcher/EventDispatcher.INT.udn b/Engine/Documentation/Source/Engine/Blueprints/BP_HowTo/EventDispatcher/EventDispatcher.INT.udn
--- a/Engine/Documentation/Source/Engine/Blueprints/BP_HowTo/EventDispatcher/EventDispatcher.INT.udn
+++ b/Engine/Documentation/Source/Engine/Blueprints/BP_HowTo/EventDispatcher/EventDispatcher.INT.udn
@@ -5,6 +5,8 @@ Description:A sample of how an event called in one Blueprint can execute events
 Related: Engine/Blueprints/UserGuide/EventDispatcher
 Related: Engine/Blueprints/UserGuide/BlueprintCommsUsage
 Related: Engine/Blueprints/UserGuide
+version:4.8
+skilllevel:Advanced
 
 When an [Event Dispatcher](Engine/Blueprints/UserGuide/EventDispatcher) is called, any Blueprints that implement and have Events bound to the Event Dispatcher will also execute when the Event Dispatcher is called. In other words, whenever the Event Dispatcher is called, any Events that are bound to it are also called allowing you to execute multiple Events at once from a single source. 
 
diff --git a/Engine/Documentation/Source/Engine/Blueprints/BP_HowTo/Interfaces/Interfaces.INT.udn b/Engine/Documentation/Source/Engine/Blueprints/BP_HowTo/Interfaces/Interfaces.INT.udn
--- a/Engine/Documentation/Source/Engine/Blueprints/BP_HowTo/Interfaces/Interfaces.INT.udn
+++ b/Engine/Documentation/Source/Engine/Blueprints/BP_HowTo/Interfaces/Interfaces.INT.udn
@@ -5,6 +5,8 @@ Description:Sample of using Blueprint Interfaces to communicate with multiple Bl
 Related: Engine/Blueprints/UserGuide/Types/Interface/UsingInterfaces
 Related: Engine/Blueprints/UserGuide/BlueprintCommsUsage
 Related: Engine/Blueprints/UserGuide
+version:4.8
+skilllevel:Advanced
 
 In this How To, we will set up a [Blueprint Interface](Engine/Blueprints/UserGuide/Types/Interface/UsingInterfaces) which will execute different scripts in multiple Blueprints when called. 
 
diff --git a/Engine/Documentation/Source/Engine/Blueprints/BP_HowTo/MacroLibraries/MacroLibraries.INT.udn b/Engine/Documentation/Source/Engine/Blueprints/BP_HowTo/MacroLibraries/MacroLibraries.INT.udn
--- a/Engine/Documentation/Source/Engine/Blueprints/BP_HowTo/MacroLibraries/MacroLibraries.INT.udn
+++ b/Engine/Documentation/Source/Engine/Blueprints/BP_HowTo/MacroLibraries/MacroLibraries.INT.udn
@@ -5,6 +5,8 @@ Description:Macros inside a Macro Library are used to increase the health and sc
 Related: Engine/Blueprints/UserGuide/Macros
 Related: Engine/Blueprints/UserGuide/Types/MacroLibrary
 Related: Engine/Blueprints
+version:4.8
+skilllevel:Advanced
 
 A **Macro Library** is a container that holds a collection of Macros or self-contained graphs that can be placed as nodes in other Blueprints. These can be time-savers as they can store commonly used sequences of nodes complete with inputs and outputs for both execution and data transfer.
 
diff --git a/Engine/Documentation/Source/Engine/Blueprints/BP_HowTo/MakingFunctions/MakingFunctions.INT.udn b/Engine/Documentation/Source/Engine/Blueprints/BP_HowTo/MakingFunctions/MakingFunctions.INT.udn
--- a/Engine/Documentation/Source/Engine/Blueprints/BP_HowTo/MakingFunctions/MakingFunctions.INT.udn
+++ b/Engine/Documentation/Source/Engine/Blueprints/BP_HowTo/MakingFunctions/MakingFunctions.INT.udn
@@ -4,6 +4,8 @@ Crumbs: %ROOT%, Engine, Engine\Blueprints, Engine\Blueprints\BP_HowTo
 Description:Create a Function which displays text when a key is pressed. 
 Related: Engine/Blueprints/UserGuide/Functions
 Related: Engine/Blueprints/UserGuide/Functions/FunctionCalls
+version:4.8
+skilllevel:Intermediate
 
 [INCLUDE:Engine/Blueprints/UserGuide/Functions#intro]
 
diff --git a/Engine/Documentation/Source/Engine/Blueprints/BP_HowTo/MakingMacros/MakingMacros.INT.udn b/Engine/Documentation/Source/Engine/Blueprints/BP_HowTo/MakingMacros/MakingMacros.INT.udn
--- a/Engine/Documentation/Source/Engine/Blueprints/BP_HowTo/MakingMacros/MakingMacros.INT.udn
+++ b/Engine/Documentation/Source/Engine/Blueprints/BP_HowTo/MakingMacros/MakingMacros.INT.udn
@@ -5,6 +5,8 @@ Description:A Macro is used to check if a player has enough energy to jump.
 Related: Engine/Blueprints/UserGuide/Macros
 Related: Engine/Blueprints/UserGuide/Types/MacroLibrary
 Related: Engine/Blueprints
+version:4.8
+skilllevel:Advanced
 
 **Macros** are essentially the same as collapsed graphs of nodes. They have an entry point and exit point designated by tunnel nodes. Each tunnel can have any number of execution or data pins which are visible on the macro node when used in other Blueprints and graphs.
 
diff --git a/Engine/Documentation/Source/Engine/Blueprints/BP_HowTo/NavigatingGraphs/NavigatingGraphs.INT.udn b/Engine/Documentation/Source/Engine/Blueprints/BP_HowTo/NavigatingGraphs/NavigatingGraphs.INT.udn
--- a/Engine/Documentation/Source/Engine/Blueprints/BP_HowTo/NavigatingGraphs/NavigatingGraphs.INT.udn
+++ b/Engine/Documentation/Source/Engine/Blueprints/BP_HowTo/NavigatingGraphs/NavigatingGraphs.INT.udn
@@ -2,6 +2,8 @@ Availability: Public
 Title:Navigating the Graph
 Crumbs:
 Description:Description of the different graph elements and how to move around the graph.
+version:4.8
+skilllevel:Beginner
 
 [OBJECT:Redirect]
 [PARAM:link]
diff --git a/Engine/Documentation/Source/Engine/Blueprints/BP_HowTo/PlacingNodes/PlacingNodes.INT.udn b/Engine/Documentation/Source/Engine/Blueprints/BP_HowTo/PlacingNodes/PlacingNodes.INT.udn
--- a/Engine/Documentation/Source/Engine/Blueprints/BP_HowTo/PlacingNodes/PlacingNodes.INT.udn
+++ b/Engine/Documentation/Source/Engine/Blueprints/BP_HowTo/PlacingNodes/PlacingNodes.INT.udn
@@ -5,7 +5,11 @@ Description:Illustrates the methods in which you can add nodes to Blueprints.
 Related: Engine/Blueprints/BP_HowTo/ConnectingNodes
 Related: Engine/Blueprints/BP_HowTo/BasicUsage
 Related: Engine/Blueprints/UserGuide
-
+version:4.8
+skilllevel:Beginner
+parent:Engine/Blueprints/BP_HowTo
+type:howto
+order:2
 
 
 On this page, you will learn how you can place nodes on your graphs in Blueprints. 
diff --git a/Engine/Documentation/Source/Engine/Blueprints/BP_HowTo/WorkingWithArrays/WorkingWithArrays.INT.udn b/Engine/Documentation/Source/Engine/Blueprints/BP_HowTo/WorkingWithArrays/WorkingWithArrays.INT.udn
--- a/Engine/Documentation/Source/Engine/Blueprints/BP_HowTo/WorkingWithArrays/WorkingWithArrays.INT.udn
+++ b/Engine/Documentation/Source/Engine/Blueprints/BP_HowTo/WorkingWithArrays/WorkingWithArrays.INT.udn
@@ -4,6 +4,8 @@ Crumbs: %ROOT%, Engine, Engine\Blueprints, Engine\Blueprints\BP_HowTo
 Description:Showcases different methods in which you can use Arrays for a collection of variables. 
 Related: Engine/Blueprints/UserGuide/Arrays
 Related: Engine/Blueprints/UserGuide/Arrays/ArrayNodes
+version:4.8
+skilllevel:Intermediate
 
 [TOC(start:2)]
 
