#UE4doc #docs:
Updated engine version field
Changed font settings for a few terms
Updated some terms due to names being changed in the editor
Replaced picture due to out of date UI

diff --git a/Engine/Documentation/Source/Engine/Blueprints/BP_HowTo/Casting/Casting.INT.udn b/Engine/Documentation/Source/Engine/Blueprints/BP_HowTo/Casting/Casting.INT.udn
--- a/Engine/Documentation/Source/Engine/Blueprints/BP_HowTo/Casting/Casting.INT.udn
+++ b/Engine/Documentation/Source/Engine/Blueprints/BP_HowTo/Casting/Casting.INT.udn
@@ -5,8 +5,8 @@ Description:Using Blueprint Casting to communicate from one Blueprint to multipl
 Related: Engine/Blueprints/UserGuide/CastNodes
 Related: Engine/Blueprints/UserGuide/BlueprintCommsUsage
 Related: Engine/Blueprints/UserGuide
-version:4.8
-skilllevel:Advanced
+skilllevel: Advanced
+version: 4.9
 
 This page covers how to use Blueprint Casting to communicate with different Blueprints and share information. In this example we allow the player to shoot an object, making it smaller each time it is hit. Each time the object is hit, it also notifies the Character Blueprint and increases the player's movement speed. 
 
@@ -20,7 +20,7 @@ For this guide we are using the **Blueprint First Person** template with **Start
 
 1.  At the confirmation window that appears, click **Ok** to use the default name and create the Blueprint. 
 
-1.  In the **Content Browser** under **Content/FirstPersonBP/Blueprints**, open the **FirstPeronProjectile** Blueprint. 
+1.  In the **Content Browser** under **Content/FirstPersonBP/Blueprints**, open the **FirstPersonProjectile** Blueprint. 
 
 	![](Casting2.png)
 
@@ -30,7 +30,7 @@ For this guide we are using the **Blueprint First Person** template with **Start
 
 	![](Casting3.png)(w:800)
 
-	On the **Event Hit** node, the **Other** pin refers to the other Actor that was hit. By using the **Cast To** node, we are checking to see if it is indeed the **LaregCube_Blueprint** that was the other Actor that was hit. If it is, we can access it and execute any Events or Functions within the Blueprint, change Variable values inside of it, or change the properties of the Components that are inside the Blueprint. 
+	On the **Event Hit** node, the **Other** pin refers to the other Actor that was hit. By using the **Cast To** node, we are checking to see if it is indeed the **LargeCube_Blueprint** that was the other Actor that was hit. If it is, we can access it and execute any Events or Functions within the Blueprint, change Variable values inside of it, or change the properties of the Components that are inside the Blueprint. 
 
 1.	Off the **As Large Cube Blueprint** pin, use **Get Static Mesh Component** and then **Set World Scale 3D** and connect as shown.
 
