#UE4doc #docs:
Resubmitting as not all changes were saved

diff --git a/Engine/Documentation/Source/Engine/Blueprints/BP_HowTo/Casting/Casting.INT.udn b/Engine/Documentation/Source/Engine/Blueprints/BP_HowTo/Casting/Casting.INT.udn
--- a/Engine/Documentation/Source/Engine/Blueprints/BP_HowTo/Casting/Casting.INT.udn
+++ b/Engine/Documentation/Source/Engine/Blueprints/BP_HowTo/Casting/Casting.INT.udn
@@ -14,11 +14,11 @@ This page covers how to use Blueprint Casting to communicate with different Blue
 For this guide we are using the **Blueprint First Person** template with **Starter Content** enabled. 
 [/REGION]
 
-1.  Inside the **Viewport** select one of the **LargeCube** meshes, then in the **Details** panel click the **Blueprint/Add Script** button. 
+1.  Inside the **Viewport** select one of the **EditorCube** meshes, then in the **Details** panel click the **Blueprint/Add Script** button. 
 
 	![](Casting1.png)
 
-1.  At the confirmation window that appears, click **Ok** to use the default name and create the Blueprint. 
+1.  At the confirmation window that appears, change the name to LargeCube_Blueprint and click **Ok** to rename and create the Blueprint. 
 
 1.  In the **Content Browser** under **Content/FirstPersonBP/Blueprints**, open the **FirstPersonProjectile** Blueprint. 
 
