#UE4doc #docs:
Updated engine version field
Reworded a few steps

diff --git a/Engine/Documentation/Source/Engine/Blueprints/BP_HowTo/CollapsingGraphs/CollapsingGraphs.INT.udn b/Engine/Documentation/Source/Engine/Blueprints/BP_HowTo/CollapsingGraphs/CollapsingGraphs.INT.udn
--- a/Engine/Documentation/Source/Engine/Blueprints/BP_HowTo/CollapsingGraphs/CollapsingGraphs.INT.udn
+++ b/Engine/Documentation/Source/Engine/Blueprints/BP_HowTo/CollapsingGraphs/CollapsingGraphs.INT.udn
@@ -5,7 +5,7 @@ Description:Collapsing nodes to a new Graph, a Function, or a Macro.
 Related: Engine/Blueprints/UserGuide
 Related: Engine/Blueprints/UserGuide/Functions
 Related: Engine/Blueprints/UserGuide/Macros
-version:4.8
+version:4.9
 skilllevel:Beginner
 parent:Engine/Blueprints/BP_HowTo
 type:howto
@@ -61,7 +61,7 @@ The easiest method of collapsing which is illustrated below is through **Collaps
 
 	![](Collapse9.png)(w:540)
 
-1.	By **Right-clicking** on the collapsed node, you can select **Expand Node** to "un-collapse" the nodes. 
+1.	By **Right-clicking** on the collapsed node, you can select **Expand Node** to revert the nodes to their original state. 
 
 	![](Collapse7.png)
 	
@@ -98,7 +98,7 @@ Below are the steps you can take to collapse your nodes to a Function.
 
 	Here we are calling the **Remove Energy** Function whenever F is pressed. 
 
-1.	You can "un-collapse" the Function by **Right-clicking** on it and selecting the **Expand Node** option. 
+1.	You can revert this collapsing by **Right-clicking** on the Function and selecting the **Expand Node** option. 
 
 For more information about Using or Calling Functions, see the [](Engine/Blueprints/UserGuide/Functions) documentation. 
 
@@ -127,7 +127,7 @@ See below for steps on collapsing nodes to a Macro.
 		
 	![](CollapseMacro5.png)
 
-1.	By **Right-clicking** on your Macro, you can select **Expand Node** to "un-collapse" the node network. 
+1.	By **Right-clicking** on your Macro, you can select **Expand Node** to revert the collpasing of the nodes. 
 
 For more information about Macros, see the [](Engine/Blueprints/UserGuide/Macros) documentation. 
 
