#UE4doc #docs:
Updated engine version field
Changed font settings for a few terms

diff --git a/Engine/Documentation/Source/Engine/Blueprints/BP_HowTo/ConnectingNodes/ConnectingNodes.INT.udn b/Engine/Documentation/Source/Engine/Blueprints/BP_HowTo/ConnectingNodes/ConnectingNodes.INT.udn
--- a/Engine/Documentation/Source/Engine/Blueprints/BP_HowTo/ConnectingNodes/ConnectingNodes.INT.udn
+++ b/Engine/Documentation/Source/Engine/Blueprints/BP_HowTo/ConnectingNodes/ConnectingNodes.INT.udn
@@ -6,7 +6,7 @@ Related: Engine/Blueprints/BP_HowTo/PlacingNodes
 Related: Engine/Blueprints/BP_HowTo/BasicUsage
 Related: Engine/Blueprints/UserGuide
 Related: Engine/Blueprints/UserGuide/CheatSheet
-version:4.8
+version:4.9
 skilllevel:Beginner
 parent:Engine/Blueprints/BP_HowTo
 order:2
@@ -38,7 +38,7 @@ There are instances where you can connect two different pin types together, in w
 
 ![](ConversionNode1.png)
 
-Above we have taken a **Float** output and connected it to a **String** input. You can see it will convert via the tool tip. 
+Above we have taken a Float output and connected it to a String input. You can see it will convert via the tool tip. 
 
 After connecting, a new conversion node is added between the two pins. 
 
@@ -80,13 +80,13 @@ Say for example you have a **Set** variable node (in which you are setting the v
 
 ![](NodeToNode1.png)
 
-Above, we have an **Event Begin Play** which then sets the **Boolean Variable 2** to false and we want to change it out to set a **Float** Variable we have created instead. We can drag the Float variable onto the Set variable and you can see that the tool tip will inform you of the action you are taking which is to change the node to write 'Float Variable' or "Set" it. 
+Above, we have an **Event Begin Play** which then sets the **Boolean Variable 2** to false and we want to change it out to set a Float Variable we have created instead. We can drag the Float variable onto the Set variable and you can see that the tool tip will inform you of the action you are taking which is to change the node to write Float Variable or Set it. 
 
 ![](NodeToNode2.png)
 
 After confirming as indicated above, your node will be swapped with the new node and the connections will remain. 
 
-You can do the same for **Get** nodes which are connected already (below we are swapping one **Bool** variable out for another). 
+You can do the same for **Get** nodes which are connected already (below we are swapping one Bool variable out for another). 
 
 ![](NodeToNode3.png)
 
@@ -125,7 +125,7 @@ You can then drag off the new pin and connect it to your target as shown below.
 
 ![](ReRoute4.png)(w:800)
 
-The **Reroute** node acts as the original pin and can be dragged off of and new nodes can be connected to it. 
+The Reroute node acts as the original pin and can be dragged off of and new nodes can be connected to it. 
 
 ![](ReRoute5.png)
 
