#UE4doc #docs:
Updated engine version field
Changed font settings for a few terms
Replaced a few images due to updates to UI

diff --git a/Engine/Documentation/Source/Engine/Blueprints/BP_HowTo/Debugging/Debugging.INT.udn b/Engine/Documentation/Source/Engine/Blueprints/BP_HowTo/Debugging/Debugging.INT.udn
--- a/Engine/Documentation/Source/Engine/Blueprints/BP_HowTo/Debugging/Debugging.INT.udn
+++ b/Engine/Documentation/Source/Engine/Blueprints/BP_HowTo/Debugging/Debugging.INT.udn
@@ -4,7 +4,7 @@ Crumbs: %ROOT%, Engine, Engine\Blueprints, Engine\Blueprints\BP_HowTo
 Description:Illustrates the methods on how to Debug your Blueprints. 
 Related: Engine/Blueprints/UserGuide/Debugging
 Related: Engine/Blueprints
-version:4.8
+version:4.9
 skilllevel:Beginner
 parent:Engine/Blueprints/BP_HowTo
 order:5
