#UE4doc #docs:
Updated engine version field
Bolded a term

diff --git a/Engine/Documentation/Source/Engine/Blueprints/BP_HowTo/DirectBlueprintComs/DirectBlueprintComs.INT.udn b/Engine/Documentation/Source/Engine/Blueprints/BP_HowTo/DirectBlueprintComs/DirectBlueprintComs.INT.udn
--- a/Engine/Documentation/Source/Engine/Blueprints/BP_HowTo/DirectBlueprintComs/DirectBlueprintComs.INT.udn
+++ b/Engine/Documentation/Source/Engine/Blueprints/BP_HowTo/DirectBlueprintComs/DirectBlueprintComs.INT.udn
@@ -5,7 +5,7 @@ Description:A demonstration of passing information directly from one Blueprint t
 Related: Engine/Blueprints/UserGuide/BlueprintComms
 Related: Engine/Blueprints/UserGuide/BlueprintCommsUsage
 Related: Engine/Blueprints/UserGuide
-version:4.8
+version:4.9
 skilllevel:Advanced
 
 This page covers using [Direct Blueprint Communication](Engine/Blueprints/UserGuide/BlueprintComms) to communicate between two Blueprints and pass information between them. 
@@ -18,7 +18,7 @@ For this guide we are using the **Blueprint Third Person** template with **Start
 
 1.  In the **Content Browser**, open the **Content/StarterContent/Shapes** folder. 
 
-1.  **Right-click** on the **Shape_Cube**, then under _Asset Actions_, choose **Create Blueprint Using This...**
+1.  **Right-click** on the **Shape_Cube**, then under **Asset Actions**, choose **Create Blueprint Using This...**
 
 	![](DirectCom1.png)
 
