#UE4doc #docs - Updated Images for UI changes. Added Engine Version Field.

diff --git a/Engine/Documentation/Source/Engine/Physics/PhAT/UserGuide/UserGuide.INT.udn b/Engine/Documentation/Source/Engine/Physics/PhAT/UserGuide/UserGuide.INT.udn
--- a/Engine/Documentation/Source/Engine/Physics/PhAT/UserGuide/UserGuide.INT.udn
+++ b/Engine/Documentation/Source/Engine/Physics/PhAT/UserGuide/UserGuide.INT.udn
@@ -2,6 +2,7 @@
 Title:Physics Asset Tool User Guide
 Crumbs: %ROOT%, Engine, Engine/Physics, Engine/Physics/PhAT
 Description:How to use the Physics Asset Tool to create Physics Asset for Skeletal Meshes.
+Version: 4.9
 
 [VAR:Topic]
 [OBJECT:Topic]
@@ -63,7 +64,7 @@ However, you can follow these steps if you need to create a Physics Asset for a
 
 1. Find the Skeletal Mesh asset you wish to add a Physics Asset to in the 
 **Content Browser**.
-1. **Right-click** on the Skeletal Mesh in the Content Browser to bring up the **Right Click Context Menu** -> **Create** -> **Create Physics Asset**.
+1. **Right-click** on the Skeletal Mesh in the Content Browser to bring up the **Context Menu** -> **Create** -> **Create Physics Asset**.
 1. Adjust setting to your liking.
 
 	![](Shared\PhAT\PhAT_restetDefault.png)
@@ -200,7 +201,7 @@ Each Physics Body in a Physics Asset can have a Physical Material assigned to it
 ## Constraints
 
 [REGION:note]
-The use of Physics Constraints is covered in the **[Physics Constraints User Guide](Engine/Physics\Constraints\ConstraintsUserGuide)** and their properties are covered in the **[Physics Constraint Reference](Engine/Physics\Constraints\ConstraintsReference)**. This section will only cover workflows that are specific to the Physics Asset Tool or ones that have deviated substantially from the norm.
+The use of Physics Constraints is covered in the [Physics Constraints User Guide](Engine/Physics\Constraints\ConstraintsUserGuide) and their properties are covered in the [Physics Constraint Reference](Engine/Physics\Constraints\ConstraintsReference). This section will only cover workflows that are specific to the Physics Asset Tool or ones that have deviated substantially from the norm.
 [/REGION]
 
 ### Editing Physics Constraints
@@ -285,7 +286,7 @@ This will not overwrite the shape of a Physics Body.
 
 Clicking the Simulation ![](Shared\Icons\icon_PhAT_PlaySim_40x.png) button will allow you to test your Physics Asset.
 
-![](Shared\PhAT\PhAT_Sim.png)
+![](SimPhAT.png)(w:600)
 
 The downward triangle next to the Simulation button will give you access to **Real Simulation** and **No Gravity**:
 
