#UE4doc #docs - Added Engine Version Field.

diff --git a/Engine/Documentation/Source/Engine/Physics/PhysicalMaterials/PhysMatReference/PhysMatReference.INT.udn b/Engine/Documentation/Source/Engine/Physics/PhysicalMaterials/PhysMatReference/PhysMatReference.INT.udn
--- a/Engine/Documentation/Source/Engine/Physics/PhysicalMaterials/PhysMatReference/PhysMatReference.INT.udn
+++ b/Engine/Documentation/Source/Engine/Physics/PhysicalMaterials/PhysMatReference/PhysMatReference.INT.udn
@@ -2,6 +2,7 @@ Availability: Public
 Title:Physical Materials Reference
 Crumbs: %ROOT%, Engine, Engine/Physics, Engine/Physics/PhysicalMaterials
 Description:Assets applied to physically simulated primitives directly or via materials used to configure and control physical properties used by the simulation.
+Version: 4.9
 
 [VAR:Topic]
 [OBJECT:Topic]
