#UE4doc #docs - Updated images. Added Engine Version Field. Added Skill Level Field.

diff --git a/Engine/Documentation/Source/Engine/Physics/PhysicalMaterials/PhysMatUserGuide/PhysMatUserGuide.INT.udn b/Engine/Documentation/Source/Engine/Physics/PhysicalMaterials/PhysMatUserGuide/PhysMatUserGuide.INT.udn
--- a/Engine/Documentation/Source/Engine/Physics/PhysicalMaterials/PhysMatUserGuide/PhysMatUserGuide.INT.udn
+++ b/Engine/Documentation/Source/Engine/Physics/PhysicalMaterials/PhysMatUserGuide/PhysMatUserGuide.INT.udn
@@ -2,6 +2,8 @@ Availability: Public
 Title:Physical Materials User Guide
 Crumbs: %ROOT%, Engine, Engine/Physics, Engine/Physics/PhysicalMaterials
 Description:Physical Materials User Guide: the creation and application of Physical Materials.
+Version: 4.9
+SkillLevel: Intermediate
 
 [VAR:Topic]
 [OBJECT:Topic]
@@ -64,7 +66,7 @@ This document will cover the creation and usage of **Physical Materials** as wel
 
 1. **Click Save** ![](Shared\Icons\icon_SaveAsset_40x.png)
 
-See the [Physical Material Reference](Engine/Physics\PhysicalMaterials\PhysMatReference) for information on the properties in the Physical Material.
+See the [**Physical Material Reference**](Engine/Physics\PhysicalMaterials\PhysMatReference) for information on the properties in the Physical Material.
 
 ### Surface Type
 
@@ -131,7 +133,7 @@ If specific Physics Bodies need differing Physical Materials, you can then adjus
 
 
 [REGION:raw]
-![](SimpleCollision.jpg)(w:920)
+![](SimpleCollision.png)(w:700)
 [/REGION]
 
 The default behavior for Physics interactions with Skeletal Meshes is to only interact with its associated Physics Asset, so the Skeletal Mesh's material's Physical Material will 
