#UE4doc #docs - Removed Overview header. Updated and reorganized the list of settings available.

diff --git a/Engine/Documentation/Source/Engine/Physics/PhysicsBodies/PhysicsBodies.INT.udn b/Engine/Documentation/Source/Engine/Physics/PhysicsBodies/PhysicsBodies.INT.udn
--- a/Engine/Documentation/Source/Engine/Physics/PhysicsBodies/PhysicsBodies.INT.udn
+++ b/Engine/Documentation/Source/Engine/Physics/PhysicsBodies/PhysicsBodies.INT.udn
@@ -2,6 +2,7 @@ Availability: Public
 Title:Physics Bodies Reference
 Crumbs: %ROOT%, Engine, Engine/Physics
 Description:The basis for most physics interactions: the Physics Body.
+Version: 4.9
 
 [VAR:Topic]
 [OBJECT:Topic]
@@ -45,7 +46,6 @@ Description:The basis for most physics interactions: the Physics Body.
 
 [TOC(start:2)]
 
-## Overview
 
 Given the ever increasing complexity of 3D meshes, proxy shapes must be used to drive physics simulations in real time 3D rendering applications; Physics Bodies are the simplified 3D meshes that Unreal Engine 4 uses for physics simulations. They can be formed of boxes, spheres, capsules, or convex hulls.
 
@@ -54,14 +54,14 @@ The following links have information on the various editors and workflows for cr
 
 ### Static Mesh
 
-For creating Physics Bodies using your 3D art program, see: [FBX: Static Meshes](Engine/Content\FBX\StaticMeshes#Collision)
+For creating Physics Bodies using your 3D art program, see: [**FBX: Static Meshes**](Engine/Content\FBX\StaticMeshes#Collision)
 
-For creating Physics Bodies using the Static Mesh Editor, see: [Collision Reference: Static Mesh Editor](Engine/Physics\Collision\Reference#CreatingStaticMeshCollisionHullsintheStaticMeshEditor)
+For creating Physics Bodies using the Static Mesh Editor, see: [**Collision Reference: Static Mesh Editor**](Engine/Physics\Collision\Reference#CreatingStaticMeshCollisionHullsintheStaticMeshEditor)
 
 
 ### Physics Assets/Skeletal Meshes
 
-For documentation over the creation of Physics Assets for use by Skeletal Meshes, see: [](Engine/Physics\PhAT) documentation.
+For documentation over the creation of Physics Assets for use by Skeletal Meshes, see: **[](Engine/Physics\PhAT)** documentation.
 
 
 ### Shape Components
@@ -80,22 +80,27 @@ All Physics Bodies have a shared pool of properties. Below are the properties fr
 | --- | --- |
 |[REGION:tablesection]Physics[/REGION]||
 | **Mass in KG** | Mass of the body in KG. |
-| **Start Awake** | %doxygen:FBodyInstance::bStartAwake% |
+| **Linear Damping** | %doxygen:FBodyInstance::LinearDamping% |
+| **Angular Damping** | %doxygen:FBodyInstance::AngularDamping% |
+| **Enable Gravity** | %doxygen:FBodyInstance::bEnableGravity% |
 | **Physics Type** | %doxygen:UBodySetup::PhysicsType% |
-| **Phys Material** | %doxygen:UBodySetup::PhysMaterial% |
+| **Double Sided Geometry** | If _True_, the physics triangle mesh will use double sided faces when doing scene queries. This is useful for planes and single sided meshes that need traces to work on both sides. |
+| **Simple Collision Physical Material** | Physical material to use for simple collision on this body. Encodes informatrion about density, friction, etc. |
+|[REGION:tablesection]Extras[/REGION]||
+| **Start Awake** | %doxygen:FBodyInstance::bStartAwake% |
+| **Center Of Mass Offset** | User Specified offset for the center of mass of this object, from the calculated location. |
+| **Mass Scale** | %doxygen:FBodyInstance::MassScale% |
+| **Max Angular Velocity** | The maximum angular velocity for this instance. |
+| **Use Async Scene** | %doxygen:FBodyInstance::bUseAsyncScene% |
+| **Generate Wake Events** | Should 'wake/sleep' events fire when this object is woken up or put to sleep by the physics simulation. |
+| **Max Depenetration Velocity** | The maximum velocity used to depenetrate this object. |
 | **Walkable Slope Override** | %doxygen:UBodySetup::WalkableSlopeOverride% |
-| **Walkable Slope Behavior** | %doxygen:FWalkableSlopeOverride::WalkableSlopeBehavior% |
+| **Walkable Slope Behavior** | Behavior of this surface (whether we affect the walkable slope). |
 | **Walkable Slope Angle** | %doxygen:FWalkableSlopeOverride::WalkableSlopeAngle% |
-| **Enable Gravity** | %doxygen:FBodyInstance::bEnableGravity% |
-| **Use Async Scene** | %doxygen:FBodyInstance::bUseAsyncScene% |
-| **Phys Material Override** | %doxygen:FBodyInstance::PhysMaterialOverride% |
-| **COMNudge** | %doxygen:FBodyInstance::COMNudge% |
-| **Sleep Family** | %doxygen:FBodyInstance::SleepFamily% |
-| **Mass Scale** | %doxygen:FBodyInstance::MassScale% |
-| **Angular Damping** | %doxygen:FBodyInstance::AngularDamping% |
-| **Linear Damping** | %doxygen:FBodyInstance::LinearDamping% |
+| **Custom Sleep Threshold Multiplier** | If the SleepFamily is set to custom, multiply the natural sleep threshold by this amount. A higher number will cause the body to sleep sooner. |
 | **Position Solver Iteration Count** | %doxygen:FBodyInstance::PositionSolverIterationCount% |
 | **Velocity Solver Iteration Count** | %doxygen:FBodyInstance::VelocitySolverIterationCount% |
+| **Sleep Family** | %doxygen:FBodyInstance::SleepFamily% |
 
 
 ### Collision
@@ -103,16 +108,17 @@ All Physics Bodies have a shared pool of properties. Below are the properties fr
 | Property | Description |
 | --- | --- |
 |[REGION:tablesection]Collision[/REGION]||
-| **Notify Rigid Body Collision** | %doxygen:FBodyInstance::bNotifyRigidBodyCollision% |
+| **Simulation Generates Hit Events** | Should 'Hit' events fire when this objects collides during physics simulation. |
+| **Phys Material Override** | %doxygen:FBodyInstance::PhysMaterialOverride% |
+|[REGION:tablesection]Collision Responses[/REGION]||
+| | See the [](Engine/Physics\Collision) documentation for more information. |
+|[REGION:tablesection]Extras[/REGION]||
+| **Use CCD** | %doxygen:FBodyInstance::bUseCCD% |
 | **Always Create Physics State** | %doxygen:UPrimitiveComponent::bAlwaysCreatePhysicsState% |
 | **Generate Overlap Events** | %doxygen:UPrimitiveComponent::bGenerateOverlapEvents% |
 | **Multi Body Overlap** | %doxygen:UPrimitiveComponent::bMultiBodyOverlap% | 
 | **Trace Complex On Move** | %doxygen:UPrimitiveComponent::bTraceComplexOnMove% |
 | **Return Material On Move** | %doxygen:UPrimitiveComponent::bReturnMaterialOnMove% |
-|[REGION:tablesection]Collision Responses[/REGION]||
-| | See the [](Engine/Physics\Collision) documentation for more information. |
-|[REGION:tablesection]Extras[/REGION]||
-| **Use CCD** | %doxygen:FBodyInstance::bUseCCD% |
 | **Check Async Scene On Move** | %doxygen:UPrimitiveComponent::bCheckAsyncSceneOnMove% |
 | **Can Ever Affect Navigation** | %doxygen:UPrimitiveComponent::bCanEverAffectNavigation% |
 
@@ -120,3 +126,5 @@ All Physics Bodies have a shared pool of properties. Below are the properties fr
 
 
 
+
+
