#UE4doc #docs - Updated image to latest Project Settings UI. Added Engine Version Field. Added additional property and description that were not previously listed.

diff --git a/Engine/Documentation/Source/Engine/Physics/Substepping/Substepping.INT.udn b/Engine/Documentation/Source/Engine/Physics/Substepping/Substepping.INT.udn
--- a/Engine/Documentation/Source/Engine/Physics/Substepping/Substepping.INT.udn
+++ b/Engine/Documentation/Source/Engine/Physics/Substepping/Substepping.INT.udn
@@ -1,7 +1,8 @@
 Availability: Public
 Title:Physics Sub-Stepping
 Crumbs: %ROOT%, Engine
-Description:
+Description: Explanation of what Physics Sub-Stepping is and when to use it. 
+Version: 4.9
 
 [VAR:Topic]
 [OBJECT:Topic]
@@ -62,6 +63,7 @@ This feature is still a work in progress, and is not fully supported by APEX Des
 | Property | Description |
 | --- | --- |
 | **Substepping** | This enables or disables sub-stepping in your current project. |
+| **Substepping Async** | Whether to substep the async phsyics simulation. |
 | **Max Substep Delta Time** | This is the maximum time, in seconds a sub-step is allowed to take. So, if a full step takes .05 seconds, and **Max Substep Delta Time** is set to .025, the full step will be split into 2 sub-steps. Now, if the full step were to take a substantial amount of time, say 2 full seconds, then the step will be evenly split based on the value in **Max Substeps**, and not **80** `(2.0/0.025=80)` sub-steps. It is worth noting that **Max Physics Delta Time** limits how long a physics step is permitted to take. |
 | **Max Substeps** | The maximum number of sub-steps a full step is permitted to be broken into. |
 
