#UE4doc #docs - Added Engine Version Field. Added Skill Level Field. Added Description for the description field.  Updated FBX Import image to match latest UI.

diff --git a/Engine/Documentation/Source/Engine/Physics/Vehicles/DoubleWishboneVehicle/DoubleWishboneVehicle.INT.udn b/Engine/Documentation/Source/Engine/Physics/Vehicles/DoubleWishboneVehicle/DoubleWishboneVehicle.INT.udn
--- a/Engine/Documentation/Source/Engine/Physics/Vehicles/DoubleWishboneVehicle/DoubleWishboneVehicle.INT.udn
+++ b/Engine/Documentation/Source/Engine/Physics/Vehicles/DoubleWishboneVehicle/DoubleWishboneVehicle.INT.udn
@@ -1,8 +1,10 @@
 Availability:Public
 Title:How to Build a Double Wishbone Suspension Vehicle
 Crumbs:%ROOT%, Engine, Engine/Physics
-Description:
+Description: Description of how to setup Double Wish Bone suspension vehicle with Unreal Engine 4.
 Navigation:topic
+SkillLevel: Advanced
+Version: 4.9
 
 [TOC (start:2)]
 
