#UE4doc #docs - Updated Image for Physics Asset icon. Added Engine Version field. Added Skill Level field.

diff --git a/Engine/Documentation/Source/Engine/Physics/Vehicles/VehcileContentCreation/Vehicle_Content.INT.udn b/Engine/Documentation/Source/Engine/Physics/Vehicles/VehcileContentCreation/Vehicle_Content.INT.udn
--- a/Engine/Documentation/Source/Engine/Physics/Vehicles/VehcileContentCreation/Vehicle_Content.INT.udn
+++ b/Engine/Documentation/Source/Engine/Physics/Vehicles/VehcileContentCreation/Vehicle_Content.INT.udn
@@ -3,6 +3,8 @@ Title:Vehicle Art Setup
 Crumbs:%ROOT%
 Description:How to setup the art assets for a vehicle in Unreal Engine 4.
 Navigation:topic
+SkillLevel: Advanced
+Version: 4.9
 
 [VAR:Topic]
 [OBJECT:Topic]
@@ -67,7 +69,7 @@ can get away with some fancy skinning, but it will require more thought on the U
 
 ## Export
 
-Vehicles are simply exported as [Skeletal Meshes](Engine/Content/FBX/SkeletalMeshes), with no special considerations.
+Vehicles are simply exported as [**Skeletal Meshes**](Engine/Content/FBX/SkeletalMeshes), with no special considerations.
 
 
 ## Import
@@ -76,7 +78,7 @@ Standard FBX import for a Skeletal Mesh. You will want to have the importer crea
 
 ## Physics Asset
 
-![](PAicon.png)
+![](PAicon.png)(w:150)
 
 The Physics Asset is incredibly important for vehicles, and cannot be overlooked or skipped. When you first have Unreal Engine 4 create a Physics Asset for you, it will 
 probably end up looking much like this:
