#UE4doc #docs - Updated Engine Version Field to 4.9

diff --git a/Engine/Documentation/Source/Engine/QuickStart/EditorQuickStart.INT.udn b/Engine/Documentation/Source/Engine/QuickStart/EditorQuickStart.INT.udn
--- a/Engine/Documentation/Source/Engine/QuickStart/EditorQuickStart.INT.udn
+++ b/Engine/Documentation/Source/Engine/QuickStart/EditorQuickStart.INT.udn
@@ -5,7 +5,7 @@ Description: Getting up and running with the basics of the Unreal Editor.
 prereq:GettingStarted/Installation
 prereq:Engine/Basics/Projects/Browser
 SkillLevel: Beginner
-Version: 4.8
+Version: 4.9
 
 [VAR:Topic]
 [OBJECT:Topic]
