#UE4doc #docs - Updated Engine version to 4.9

diff --git a/Engine/Documentation/Source/Engine/QuickStart/1/LDQS_1.INT.udn b/Engine/Documentation/Source/Engine/QuickStart/1/LDQS_1.INT.udn
--- a/Engine/Documentation/Source/Engine/QuickStart/1/LDQS_1.INT.udn
+++ b/Engine/Documentation/Source/Engine/QuickStart/1/LDQS_1.INT.udn
@@ -3,7 +3,7 @@ Title:1 - Required Setup
 Crumbs:%ROOT%, GettingStarted, Engine/QuickStart
 Description:In this step we will use the Project Browser to create a New Project. 
 SkillLevel:Beginner
-Version:4.8
+Version:4.9
 Related:Engine/Basics/Projects
 Related:Engine/Basics/Projects/Browser
 Related:GettingStarted/Installation
