#UE4doc #docs - Updated Engine Version to 4.9

diff --git a/Engine/Documentation/Source/Engine/QuickStart/2/LDQS_2.INT.udn b/Engine/Documentation/Source/Engine/QuickStart/2/LDQS_2.INT.udn
--- a/Engine/Documentation/Source/Engine/QuickStart/2/LDQS_2.INT.udn
+++ b/Engine/Documentation/Source/Engine/QuickStart/2/LDQS_2.INT.udn
@@ -3,7 +3,7 @@ Title:2 - Navigating the Viewport
 Crumbs:%ROOT%, GettingStarted, Engine/QuickStart
 Description:Learn the different methods you can use to navigate the different viewports inside Unreal Engine 4. 
 SkillLevel:Beginner
-Version:4.8
+Version:4.9
 Related:Engine/UI/LevelEditor/Viewports
 Related:Engine/Levels
 Related:Engine/UI/LevelEditor/Viewports/ViewportControls
@@ -26,12 +26,12 @@ Related:Engine/UI/LevelEditor/Viewports/ViewportControls
 
 %Steps%
 
-With the project open and ready to go, the first thing you may notice is the [Viewport](Engine/UI/LevelEditor/Viewports) in the center of the Unreal Editor. 
+With the project open and ready to go, the first thing you may notice is the [**Viewport**](Engine/UI/LevelEditor/Viewports) in the center of the Unreal Editor. 
 
 ![](TableAndChairs.png)(w:720)
 
-Inside the Viewport is where you will do most of your level construction. The template project that we selected in the previous step includes a small sample [Level](Engine/Levels) and some assets for us to get us started with. 
-Using this little area as a point of reference, take a moment to get used to the [Viewport Camera Controls](Engine/UI/LevelEditor/Viewports/ViewportControls) by using the most common methods of navigating the Viewport in Unreal Engine 4. 
+Inside the Viewport is where you will do most of your level construction. The template project that we selected in the previous step includes a small sample [**Level**](Engine/Levels) and some assets for us to get us started with. 
+Using this little area as a point of reference, take a moment to get used to the [**Viewport Camera Controls**](Engine/UI/LevelEditor/Viewports/ViewportControls) by using the most common methods of navigating the Viewport in Unreal Engine 4. 
 
 
 ## Standard Controls
