#UE4doc #docs:
Updated engine version field
Replaced a couple images due to not matching the current UI of the editor.

diff --git a/Engine/Documentation/Source/Engine/Blueprints/BP_HowTo/EventDispatcher/EventDispatcher.INT.udn b/Engine/Documentation/Source/Engine/Blueprints/BP_HowTo/EventDispatcher/EventDispatcher.INT.udn
--- a/Engine/Documentation/Source/Engine/Blueprints/BP_HowTo/EventDispatcher/EventDispatcher.INT.udn
+++ b/Engine/Documentation/Source/Engine/Blueprints/BP_HowTo/EventDispatcher/EventDispatcher.INT.udn
@@ -5,7 +5,7 @@ Description:A sample of how an event called in one Blueprint can execute events
 Related: Engine/Blueprints/UserGuide/EventDispatcher
 Related: Engine/Blueprints/UserGuide/BlueprintCommsUsage
 Related: Engine/Blueprints/UserGuide
-version:4.8
+version:4.9
 skilllevel:Advanced
 
 When an [Event Dispatcher](Engine/Blueprints/UserGuide/EventDispatcher) is called, any Blueprints that implement and have Events bound to the Event Dispatcher will also execute when the Event Dispatcher is called. In other words, whenever the Event Dispatcher is called, any Events that are bound to it are also called allowing you to execute multiple Events at once from a single source. 
