#UE4doc #docs:
Updated engine version field
Replaced an image due to out of date UI
Changed a term to match changes to terminology

diff --git a/Engine/Documentation/Source/Engine/Blueprints/BP_HowTo/MakingFunctions/MakingFunctions.INT.udn b/Engine/Documentation/Source/Engine/Blueprints/BP_HowTo/MakingFunctions/MakingFunctions.INT.udn
--- a/Engine/Documentation/Source/Engine/Blueprints/BP_HowTo/MakingFunctions/MakingFunctions.INT.udn
+++ b/Engine/Documentation/Source/Engine/Blueprints/BP_HowTo/MakingFunctions/MakingFunctions.INT.udn
@@ -1,10 +1,10 @@
 Availability: Public
-Title:Create Functions
+Title:Creating Functions
 Crumbs: %ROOT%, Engine, Engine\Blueprints, Engine\Blueprints\BP_HowTo
 Description:Create a Function which displays text when a key is pressed. 
 Related: Engine/Blueprints/UserGuide/Functions
 Related: Engine/Blueprints/UserGuide/Functions/FunctionCalls
-version:4.8
+version:4.9
 skilllevel:Intermediate
 
 [INCLUDE:Engine/Blueprints/UserGuide/Functions#intro]
@@ -55,7 +55,7 @@ The steps below will guide you through creating a Function which will print text
 
 	![](RenameFunction.png)
 
-	Give the Function a name such as **Print Text**. 
+	Give the Function a name such as "Print Text". 
 
 1.	In the Function's graph, drag off the **Print Text** pin and search for and add a **Print String** node. 
 
