#UE4doc #docs:
Updated engine version field

diff --git a/Engine/Documentation/Source/Engine/Blueprints/BP_HowTo/MakingMacros/MakingMacros.INT.udn b/Engine/Documentation/Source/Engine/Blueprints/BP_HowTo/MakingMacros/MakingMacros.INT.udn
--- a/Engine/Documentation/Source/Engine/Blueprints/BP_HowTo/MakingMacros/MakingMacros.INT.udn
+++ b/Engine/Documentation/Source/Engine/Blueprints/BP_HowTo/MakingMacros/MakingMacros.INT.udn
@@ -5,7 +5,7 @@ Description:A Macro is used to check if a player has enough energy to jump.
 Related: Engine/Blueprints/UserGuide/Macros
 Related: Engine/Blueprints/UserGuide/Types/MacroLibrary
 Related: Engine/Blueprints
-version:4.8
+version:4.9
 skilllevel:Advanced
 
 **Macros** are essentially the same as collapsed graphs of nodes. They have an entry point and exit point designated by tunnel nodes. Each tunnel can have any number of execution or data pins which are visible on the macro node when used in other Blueprints and graphs.
@@ -15,7 +15,7 @@ skilllevel:Advanced
 In this example, we will create a **Macro** that will be used to check if a player character has enough energy to jump or not. If they have enough energy to jump, we will remove some energy from the player and print their current value to the screen then allow them to jump. If they do not have enough energy, we will print to the screen that they are out of energy and not allow them to jump. 
 
 [REGION:note]
-_For this example, we are using the **Blueprint Third Person Project** with **Starter Content** enabled._
+For this example, we are using the **Blueprint Third Person Project** with **Starter Content** enabled.
 [/REGION]
 
 1.  Inside the **Content/ThirdPersonBP/Blueprints** folder, open the **ThirdPersonCharacter** Blueprint. 
