#UE4doc #docs:
Updated engine version field

diff --git a/Engine/Documentation/Source/Engine/Blueprints/BP_HowTo/PlacingNodes/PlacingNodes.INT.udn b/Engine/Documentation/Source/Engine/Blueprints/BP_HowTo/PlacingNodes/PlacingNodes.INT.udn
--- a/Engine/Documentation/Source/Engine/Blueprints/BP_HowTo/PlacingNodes/PlacingNodes.INT.udn
+++ b/Engine/Documentation/Source/Engine/Blueprints/BP_HowTo/PlacingNodes/PlacingNodes.INT.udn
@@ -5,7 +5,7 @@ Description:Illustrates the methods in which you can add nodes to Blueprints.
 Related: Engine/Blueprints/BP_HowTo/ConnectingNodes
 Related: Engine/Blueprints/BP_HowTo/BasicUsage
 Related: Engine/Blueprints/UserGuide
-version:4.8
+version:4.9
 skilllevel:Beginner
 parent:Engine/Blueprints/BP_HowTo
 type:howto
