#UE4doc #docs:
Resubmitting as not all of the changes were saved

diff --git a/Engine/Documentation/Source/Engine/Blueprints/BP_HowTo/MakingFunctions/MakingFunctions.INT.udn b/Engine/Documentation/Source/Engine/Blueprints/BP_HowTo/MakingFunctions/MakingFunctions.INT.udn
--- a/Engine/Documentation/Source/Engine/Blueprints/BP_HowTo/MakingFunctions/MakingFunctions.INT.udn
+++ b/Engine/Documentation/Source/Engine/Blueprints/BP_HowTo/MakingFunctions/MakingFunctions.INT.udn
@@ -11,7 +11,7 @@ skilllevel:Intermediate
 
 The steps below will guide you through creating a Function which will print text to the screen when a button is pressed. 
 
-1.  Inside the **Content Browser**, click the **New** Button then select **Blueprint**. 
+1.  Inside the **Content Browser**, click the **New** Button then select **Blueprint Class**. 
 
 	![](NewBlueprint.png)
 
