#DocFix

Fixed a typ with 'if' needing to be changed to 'of'

diff --git a/Engine/Documentation/Source/Engine/Performance/Scalability/ScalabilityReference/ScalabilityReference.INT.udn b/Engine/Documentation/Source/Engine/Performance/Scalability/ScalabilityReference/ScalabilityReference.INT.udn
--- a/Engine/Documentation/Source/Engine/Performance/Scalability/ScalabilityReference/ScalabilityReference.INT.udn
+++ b/Engine/Documentation/Source/Engine/Performance/Scalability/ScalabilityReference/ScalabilityReference.INT.udn
@@ -432,6 +432,6 @@ Most material editor outputs only affect the pixel shader. World position offset
 
 ![](SC_Skeletal.png)(w:920)
 
-Skeletal meshes can have static Level of Details models. With the console variable r.SkeletalMeshLODBias, it is possible to globally bias the LOD level. Depending on the quality of your first or second LOD, it might be a good idea to have this option in the scalability settings. Here you can see the scalability setting on top if the distance dependent LOD.
+Skeletal meshes can have static Level of Details models. With the console variable r.SkeletalMeshLODBias, it is possible to globally bias the LOD level. Depending on the quality of your first or second LOD, it might be a good idea to have this option in the scalability settings. Here you can see the scalability setting on top of the distance dependent LOD.
 
 
