#UE4doc #docs:
Updated engine version field
Replaced image due to change in menu format
Replaced bolded text with link to section of the page it was referencing

diff --git a/Engine/Documentation/Source/Engine/Blueprints/BP_HowTo/WorkingWithArrays/WorkingWithArrays.INT.udn b/Engine/Documentation/Source/Engine/Blueprints/BP_HowTo/WorkingWithArrays/WorkingWithArrays.INT.udn
--- a/Engine/Documentation/Source/Engine/Blueprints/BP_HowTo/WorkingWithArrays/WorkingWithArrays.INT.udn
+++ b/Engine/Documentation/Source/Engine/Blueprints/BP_HowTo/WorkingWithArrays/WorkingWithArrays.INT.udn
@@ -4,14 +4,14 @@ Crumbs: %ROOT%, Engine, Engine\Blueprints, Engine\Blueprints\BP_HowTo
 Description:Showcases different methods in which you can use Arrays for a collection of variables. 
 Related: Engine/Blueprints/UserGuide/Arrays
 Related: Engine/Blueprints/UserGuide/Arrays/ArrayNodes
-version:4.8
+version:4.9
 skilllevel:Intermediate
 
 [TOC(start:2)]
 
-If you are not familiar with programming terms, the best way to think of an **Array** is that it is a collection of variables of the same type, contained within a single entity. For example, if you were to create a **Float** Array, that Array would only be able to hold Float variables in it. Once you have an Array created, you can access individual variables within the Array and modify them or modify all the items in the Array at once (which greatly improves repetitive tasks by doing it all at once versus individually). 
+If you are not familiar with programming terms, the best way to think of an **Array** is that it is a collection of variables of the same type, contained within a single entity. For example, if you were to create a Float Array, that Array would only be able to hold Float variables in it. Once you have an Array created, you can access individual variables within the Array and modify them or modify all the items in the Array at once (which greatly improves repetitive tasks by doing it all at once versus individually). 
 
-On this page, you will learn how to create Array variables inside a Blueprint as well as using the **Make Array** node to build an Array. Once you understand how to create an Array, you can check out the **Example Array Usage** section for some sample use cases for working with Arrays. 
+On this page, you will learn how to create Array variables inside a Blueprint as well as using the **Make Array** node to build an Array. Once you understand how to create an Array, you can check out the [Example Array Usage](#examplearrayusage) section for some sample use cases for working with Arrays. 
 
 [REGION:note]
 For more detailed information and usage, see the [](Engine/Blueprints/UserGuide/Arrays) and [](Engine/Blueprints/UserGuide/Arrays/ArrayNodes) documentation. 
@@ -86,7 +86,7 @@ The steps below will show you how to use the **Make Array** node.
 
 	The value in the box on the **Get** node is the **Index** (Array Element) to access (Get) from the **Make Array** node. 
 	
-	This could also be an **Integer** Variable as depicted below as the variable called **Get Camera**. 
+	This could also be an **Integer** variable as depicted below as the variable called **Get Camera**. 
 
 	![](HT11.png)
 
@@ -95,11 +95,13 @@ The steps below will show you how to use the **Make Array** node.
 	![](HT12.png)
 
 
-## Example Array Usage
+##Example Array Usage
 
 In this example, we will use Arrays to spawn a random Blueprint effect in our level at a random location. 
 
-_For this example, we are using the **Blueprint 3rd Person Template** project with **Starter Content** enabled._
+[REGION:note]
+For this example, we are using the **Blueprint Third Person Template** project with **Starter Content** enabled.
+[/REGION]
 
 1.	Inside your project, from the **Modes** menu on the **Basic** tab, drag 5 **Target Points** into your level. 
 
