Writes Panini projection post process effect's documentation

diff --git a/Engine/Documentation/Source/Engine/Rendering/PostProcessEffects/PaniniProjection/PaniniProjection.INT.udn b/Engine/Documentation/Source/Engine/Rendering/PostProcessEffects/PaniniProjection/PaniniProjection.INT.udn
new file mode 100644
--- /dev/null
+++ b/Engine/Documentation/Source/Engine/Rendering/PostProcessEffects/PaniniProjection/PaniniProjection.INT.udn
@@ -0,0 +1,66 @@
+Availability:Public
+Title:Panini projection
+Crumbs: %ROOT%, Engine, Engine/Rendering/PostProcessEffects
+Description:3D projection that fixes the perspective projection's geometric distortion on sides of the view having a wide FOV.
+
+
+## Configuration
+
+By default, Unreal Engine 4 is using a perspective projection. However this projection model create distortions
+of the geometry when working with wide field of view. It is very noticeable for instance here with sphere and a
+FOV = 120.
+
+![](qagame_fov120.png)
+
+Therefore Unreal Engine provides a optional post processing pass that attempt to fix this issue using a Panini
+Projection. The idea of this post processing is to displaced the rendered pixel at the location where they should
+using the panini projection instead of the perspective projection. To setup, the only thing you need to do is 
+setting up r.Upscale.Panini.D > 0. This post process effect is done in the Upscaling pass. It means this pass is
+then going to be used if r.ScreenPercentage != 100 or r.Upscale.Panini.D > 0. For further understanding of this
+parameter, you can refer directly to the original research paper referenced at the bottom of this page.
+
+![](qagame_fov120_panini.png)
+
+r.Upscale.Panini.D > 0 directly emphasizes the effect. However it not the only parameter. r.Upscale.Panini.S 
+lerp with the paper's hardcompression. And if still some sphere still looks elipsed in the corners, then this
+parameter can also be tweaked as well. However, for r.Upscale.Panini.S < 0, start shows unrendered black
+looking pixels in the corners.
+
+![](qagame_fov120_vertical_compression.png)
+
+To have more understanding of how this effects works, r.Upscale.ScreenFit can be modified as well. One thing
+to have in mind is that some of the perspective projected pixels are at the end being not show at the screen
+near the top and bottom of screen.
+
+![](qagame_actual_effect.png)
+
+
+## Straight lines
+
+The panini projection bring severals guarentees on the lines. Firstly the vertical lines are guarented to stay
+straight for any r.Upscale.Panini.D and r.Upscale.Panini.S. Secondly, the lines going through the center of the
+screen are guarented to stay straight as well for any r.Upscale.Panini.D, but r.Upscale.Panini.S = 0. This
+mathematicall property is therefore ideal in First-Person-Shooter game because the panini projection of the
+weapon directed towards the center of the screen is going to keep its straight lines going towards the center.
+
+![](shootergame_fov120.png)
+![](shootergame_fov120_panini.png)
+
+
+## Center Bluryness
+
+One issue of this post processing effect is that a blur phenomena at the center of the screen kicks in as 
+r.Upscale.Panini.D increase. The panini projection has been implemented in the upscale pass especially to
+take aventage of the sharpening filter to fight against this phenomena. However it can quickly become not
+enough to fix this issue. Therefore you may need to increase r.ScreenPercentage to help fixing that issue
+but that is going to have a cost on all your rendering performance because increasing the number of pixels
+to draw. You may consider choosing a better trade of by then choosing a lower value of r.Upscale.Panini.D
+on hardware that can't afford an higher screen percentage, and still have this effect by only rellying on
+the upscale pass's sharpening filter.
+
+![](shootergame_fov120_panini_screen_percentage.png)
+
+## References
+
+* [Original article's website](http://tksharpless.net/vedutismo/Pannini/).
+* [Original article's direct link](http://tksharpless.net/vedutismo/Pannini/panini.pdf).
