#UE4doc #docs:
Added engine version field
Removed Overview header
Removed "On this page..." field as there was only one entry

diff --git a/Engine/Documentation/Source/Engine/Blueprints/Editor/BlueprintEditorUserGuide.INT.udn b/Engine/Documentation/Source/Engine/Blueprints/Editor/BlueprintEditorUserGuide.INT.udn
--- a/Engine/Documentation/Source/Engine/Blueprints/Editor/BlueprintEditorUserGuide.INT.udn
+++ b/Engine/Documentation/Source/Engine/Blueprints/Editor/BlueprintEditorUserGuide.INT.udn
@@ -1,7 +1,8 @@
 Availability:Public
 Title:Blueprint Editor Reference
 Crumbs: %ROOT%, Engine, Engine/Blueprints
-Description:The Blueprint Editor Reference page outlines the Blueprint Editor's Interface elements and its basic usage instructions. 
+Description:The Blueprint Editor Reference page outlines the Blueprint Editor's Interface elements and its basic usage instructions.
+version: 4.9
 
 [VAR:Topic]
 [OBJECT:Topic]
@@ -27,20 +28,18 @@ Description:The Blueprint Editor Reference page outlines the Blueprint Editor's
 ![](blueprint_editor_example.png)
 [/REGION]
 
-[TOC (start:2)] 
 
-## Overview
 
-Blueprint visual scripting is a system that is used _throughout_ UE4 in many different ways. Blueprints can drive level-based events, control internally scripted behaviors for in-game Actors, and even be used to control complex animations across highly realistic game character systems. For each of these Blueprint applications, the location where you edit that Blueprint script, as well as the tools available to you, will change slightly depending on your needs. This means that there are actually multiple places and different ways within UE4 that the Blueprint Editor can appear. However, no matter the differences, the Blueprint Editor always performs the same key task: it enables you to create and edit powerful visual scripts to drive various aspects of your game.
+Blueprint visual scripting is a system that is used throughout UE4 in many different ways. Blueprints can drive level-based events, control internally scripted behaviors for in-game Actors, and even be used to control complex animations across highly realistic game character systems. For each of these Blueprint applications, the location where you edit that Blueprint script, as well as the tools available to you, will change slightly depending on your needs. This means that there are actually multiple places and different ways within UE4 that the Blueprint Editor can appear. However, no matter the differences, the Blueprint Editor always performs the same key task: it enables you to create and edit powerful visual scripts to drive various aspects of your game.
 
 
-The Blueprint Editor is, at its heart, simply a node-based graph editor. It is your primary tool for creating and editing visual scripting node networks, commonly referred to merely as _Blueprints_. The Blueprint Editor employs a _context sensitive_ design that helps you access functionality for the objects you need specifically when you need them, while also providing flexibility for those times when you need to do something a bit unconventional.
+The Blueprint Editor is, at its heart, simply a node-based graph editor. It is your primary tool for creating and editing visual scripting node networks, commonly referred to merely as Blueprints. The Blueprint Editor employs a context sensitive design that helps you access functionality for the objects you need specifically when you need them, while also providing flexibility for those times when you need to do something a bit unconventional.
 
 Here are some critical things to know about the Blueprint Editor:
 
 * It includes several tools and panels tools to help you create your own variables, functions, arrays, and more.
 * It has a variety of debugging and analysis tools built into it to help you quickly debug and improve data flow in your networks.
-* Within Unreal Engine 4, the Blueprint Editor comes in a variety of different specialized forms _depending on the types of Blueprint networks you are editing_.
+* Within Unreal Engine 4, the Blueprint Editor comes in a variety of different specialized forms depending on the types of Blueprint networks you are editing.
 
 
 Before you get too far into the Blueprint Editor, you should probably make sure you have at least a fairly solid understanding of Blueprints themselves. If you need more information, see [Getting Started with Blueprints](Engine/Blueprints/GettingStarted), and [Blueprints Overview](Engine/Blueprints/Overview).
