#UE4doc #docs:
Added engine version field
Removed "On this page..." field as there were no entries

diff --git a/Engine/Documentation/Source/Engine/Blueprints/Editor/Modes/Defaults/DefaultsTab/DefaultsTab.INT.udn b/Engine/Documentation/Source/Engine/Blueprints/Editor/Modes/Defaults/DefaultsTab/DefaultsTab.INT.udn
--- a/Engine/Documentation/Source/Engine/Blueprints/Editor/Modes/Defaults/DefaultsTab/DefaultsTab.INT.udn
+++ b/Engine/Documentation/Source/Engine/Blueprints/Editor/Modes/Defaults/DefaultsTab/DefaultsTab.INT.udn
@@ -3,7 +3,7 @@ Title: Blueprint Editor Defaults Tab
 Crumbs:%ROOT%, Engine/Blueprints, Engine/Blueprints/Editor, Engine/Blueprints/Editor/Modes
 Description:An overview of the Blueprints Defaults tab
 
-[TOC (start:2)]
+
 
 The **Class Defaults** tab contains information about the Blueprint's default settings and properties, and about any variables the Blueprint contains. You can modify these settings to suit your needs.
 
