#UE4doc #docs:
Resubmitting as not all changes were saved

diff --git a/Engine/Documentation/Source/Engine/Blueprints/Editor/Modes/Defaults/DefaultsTab/DefaultsTab.INT.udn b/Engine/Documentation/Source/Engine/Blueprints/Editor/Modes/Defaults/DefaultsTab/DefaultsTab.INT.udn
--- a/Engine/Documentation/Source/Engine/Blueprints/Editor/Modes/Defaults/DefaultsTab/DefaultsTab.INT.udn
+++ b/Engine/Documentation/Source/Engine/Blueprints/Editor/Modes/Defaults/DefaultsTab/DefaultsTab.INT.udn
@@ -2,7 +2,7 @@ Availability:Public
 Title: Blueprint Editor Defaults Tab
 Crumbs:%ROOT%, Engine/Blueprints, Engine/Blueprints/Editor, Engine/Blueprints/Editor/Modes
 Description:An overview of the Blueprints Defaults tab
-
+version: 4.9
 
 
 The **Class Defaults** tab contains information about the Blueprint's default settings and properties, and about any variables the Blueprint contains. You can modify these settings to suit your needs.
