Talks about the Unreal Tournament's panini projection material function.

diff --git a/Engine/Documentation/Source/Engine/Rendering/PostProcessEffects/PaniniProjection/PaniniProjection.INT.udn b/Engine/Documentation/Source/Engine/Rendering/PostProcessEffects/PaniniProjection/PaniniProjection.INT.udn
--- a/Engine/Documentation/Source/Engine/Rendering/PostProcessEffects/PaniniProjection/PaniniProjection.INT.udn
+++ b/Engine/Documentation/Source/Engine/Rendering/PostProcessEffects/PaniniProjection/PaniniProjection.INT.udn
@@ -60,6 +60,12 @@ the upscale pass's sharpening filter.
 
 ![](shootergame_fov120_panini_screen_percentage.png)
 
+Another way to use the panini projection is through a Material function outputing a world position
+offset to be pluged-in into the material's world position offset input pin. This actually what Unreal
+Tournament uses instead of rendering the weapon at a different FOV to fix the perspective projection
+distortions. You may want to have a look at it in the UT's github repository.
+
+
 ## References
 
 * [Original article's website](http://tksharpless.net/vedutismo/Pannini/).
