update new content into Chinese

diff --git a/Engine/Documentation/Source/Videos/VideoTutorials.INT.udn b/Engine/Documentation/Source/Videos/VideoTutorials.INT.udn
--- a/Engine/Documentation/Source/Videos/VideoTutorials.INT.udn
+++ b/Engine/Documentation/Source/Videos/VideoTutorials.INT.udn
@@ -1,5 +1,5 @@
-Title:Video Tutorials
-Description:Video tutorials demonstrating concepts and techniques in Unreal Engine 4.
+Title:视频教学
+Description:展示虚幻引擎4中理念和技巧的视频教程
 Crumbs:%ROOT%
 Availability:Public
 
