#DocFix

Fixing errors in multiple pages pointed out in the forums

diff --git a/Engine/Documentation/Source/Engine/Rendering/Materials/Functions/Editing/EditingFunctions.INT.udn b/Engine/Documentation/Source/Engine/Rendering/Materials/Functions/Editing/EditingFunctions.INT.udn
--- a/Engine/Documentation/Source/Engine/Rendering/Materials/Functions/Editing/EditingFunctions.INT.udn
+++ b/Engine/Documentation/Source/Engine/Rendering/Materials/Functions/Editing/EditingFunctions.INT.udn
@@ -18,7 +18,7 @@ In either the **Content Browser** or the **Material Editor's** Graph tab, you ca
 
 However, it is important to note that any changes made and saved to a material function will be present in all instances of that material function moving forward. For example, if you made a change to the internal network of a Radial Gradient material function, all existing instances of that function would receive the update, as well as all new instances moving forward. 
 
-For this reason, unless you are _certain_ that your change needs to propagate throughout all other instances of the function, it may be wise to make a copy of an existing function in the **Content Browser** (**right-click** and choose _Make a Copy_ from the context menu) than editing the original material function.
+For this reason, unless you are _certain_ that your change needs to propagate throughout all other instances of the function, it may be wise to make a copy of an existing function in the **Content Browser** (**right-click** and choose _Make a Copy_ from the context menu) rather than editing the original material function.
 
 Once you make changes to a function you have to click the ![button_apply.png](button_apply.png) button to propagate the changes to the function asset and any materials using the function. Once done, be **sure to save your asset in the Content Browser!**
 
diff --git a/Engine/Documentation/Source/Engine/Rendering/Materials/Functions/Overview/FunctionsOverview.INT.udn b/Engine/Documentation/Source/Engine/Rendering/Materials/Functions/Overview/FunctionsOverview.INT.udn
--- a/Engine/Documentation/Source/Engine/Rendering/Materials/Functions/Overview/FunctionsOverview.INT.udn
+++ b/Engine/Documentation/Source/Engine/Rendering/Materials/Functions/Overview/FunctionsOverview.INT.udn
@@ -61,12 +61,6 @@ Below are the material expression nodes that relate to material functions, along
 
 Since material functions are encapsulated networks of nodes, it is up to the user to make sure that data can flow into and out of them. This is handled by way of *FunctionInput* and *FunctionOutput* nodes. Understanding these nodes is critical for using material functions.
 
-[EXCERPT:FunctionParameters]
-[REGION:warning]
-Material Parameters, such as Scalar and Vector Parameters, cannot exist within Functions. This is to avoid any potential crosstalk that might occur when multiple instances of a Function are used in a single Material. To send data into a Function from a Parameter, you must use a FunctionInput node within the Function, and connect your Parameter in from outside, in the Material itself.
-[/REGION] 
-[/EXCERPT:FunctionParameters]
-
 Within the function itself, FunctionInput and FunctionOutput nodes appear like so:
 
 ![InputOutput_Within.png](InputOutput_Within.png)
@@ -78,7 +72,7 @@ From the outside, however, when the function is used within a material, those no
 
 ### FunctionInput Nodes
 
-As stated, MaterialInput nodes serve as the gateway through which data will enter a material function. A given function may have any number of these input nodes, each corresponding with an input pin that will appear on the function itself.
+As stated, Material Input nodes serve as the gateway through which data will enter a material function. A given function may have any number of these input nodes, each corresponding with an input pin that will appear on the function itself.
 
 ![InputNode.png](InputNode.png)
 
@@ -165,7 +159,7 @@ Most of the time you will want to preview a function output, so those will be pr
 
 ![PreviewingOutput.png](PreviewingOutput.png)
 
-Function input nodes have some options for specifying preview values, since they do not know the actual values they will be used with in a material.  Each input has a built in PreviewValue can be used to show a constant for float input types.  Function inputs also have a 'Preview' connector which allows you to override that built in value with any value that matches the input's type.  In this example, a texture sample has been used to provide the preview for a float 3 input.
+Function input nodes have some options for specifying preview values, since they do not know the actual values that will be used with in a material.  Each input has a built in PreviewValue that can be used to show a constant for float input types.  Function inputs also have a 'Preview' connector which allows you to override that built in value with any value that matches the input's type.  In this example, a texture sample has been used to provide the preview for a float 3 input.
 
 ![PreviewTexture.png](PreviewTexture.png)
 
diff --git a/Engine/Documentation/Source/Engine/Rendering/Materials/LayeredMaterials/LayeredMaterials.INT.udn b/Engine/Documentation/Source/Engine/Rendering/Materials/LayeredMaterials/LayeredMaterials.INT.udn
--- a/Engine/Documentation/Source/Engine/Rendering/Materials/LayeredMaterials/LayeredMaterials.INT.udn
+++ b/Engine/Documentation/Source/Engine/Rendering/Materials/LayeredMaterials/LayeredMaterials.INT.udn
@@ -13,13 +13,13 @@ Description:An introductory document on the Layered Materials technique in UE4.
 
 Layered Materials can be thought of as "Materials within Materials." They provide a means to create a single Material that has a series of sub-Materials (or _layers_) that can be placed across the surface of your objects using per-pixel operations, such as masks. They are perfect for handling complex blending between unique surface types. In the image of the rockets above, the rocket on the far right is using separate Material Layers, including chrome, aluminum, and copper, and is blending between each material on a per-pixel basis. This effect is easily achieved with Layered Materials.
 
-As a feature, Layered Materials exist as an extension Material Functions. Functions, as you may recall, are self-contained node networks that perform a specific operation, such as a complex math equation. These functions can then be reused at will in any number of Materials. Because of the _Make Material Attributes_ and _Break Material Attributes_ nodes, functions can also be used to define a Material, entirely _within_ a Function. Then, by bringing those functions into a final Material, you have then a Layered Material.
+As a feature, Layered Materials exist as an extension Material Function. Functions, as you may recall, are self-contained node networks that perform a specific operation, such as a complex math equation. These functions can then be reused at will in any number of Materials. Because of the _Make Material Attributes_ and _Break Material Attributes_ nodes, functions can also be used to define a Material, entirely _within_ a Function. Then, by bringing those functions into a final Material, you can have a Layered Material.
 
 ![](ChromeLayer.png)(w:600)
 
 In the image above, we have created a simple chrome material, entirely within a Material Function, via a Make Material Attributes node. This function can now be used as a Material Layer, and can be blended with other Material Layers.
 
-It is important to realize that since Material Functions cannot be applied directly to a surface, you need bring your Material Layer functions into a new Material, which can then be applied to your objects. The great thing about this is that you can bring in as many Material Layers as you wish, blending them together however you see fit.
+It is important to realize that since Material Functions cannot be applied directly to a surface, you need to bring your Material Layer functions into a new Material, which can then be applied to your objects. The great thing about this is that you can bring in as many Material Layers as you wish, blending them together however you see fit.
 
 So, at a high level, the workflow goes like this:
 
@@ -97,7 +97,7 @@ Within the Material Functions list of the Material Editor is a list of various M
 
 ## Instancing Layered Materials
 
-Because Layered Materials are essentially Material Functions, parameterizing them for instancing requires a little bit of extra forethought. Parameter expressions, such as Scalar and Vector Parameters, _cannot exist within a function_. To use them, you must create a Function Input expression as a part of the Material Layer. You can then connect a parameter into this input from the top-level Material. See the [Material Functions overview](Engine/Rendering/Materials\Functions) for more information.
+Because Layered Materials are essentially Material Functions, parameterizing them for instancing requires a little bit of extra forethought. To make better use of Scalar and Vector Parameters, you can create a Function Input expression as part of a Material Layer. You can then connect a parameter into this input from the top-level Material. See the [Material Functions overview](Engine/Rendering/Materials\Functions) for more information.
 
 The flow works something like this:
 
diff --git a/Engine/Documentation/Source/Engine/Rendering/Materials/ParameterCollections/ParemeterCollections.INT.udn b/Engine/Documentation/Source/Engine/Rendering/Materials/ParameterCollections/ParemeterCollections.INT.udn
--- a/Engine/Documentation/Source/Engine/Rendering/Materials/ParameterCollections/ParemeterCollections.INT.udn
+++ b/Engine/Documentation/Source/Engine/Rendering/Materials/ParameterCollections/ParemeterCollections.INT.udn
@@ -30,7 +30,7 @@ In the example above, global parameters are used to control the following aspect
 * The sun angle and color is also tracked and used to correctly colorize the water crests and limit the effect to only times when light shines directly through the wave crests (when the player is looking toward the sun). 
 * The coloration of effect across the time of day and overall diffuse lighting color for the water are also being controlled. 
 
-## Creating & Editing Material Parameter Collections
+## Creating and Editing Material Parameter Collections
 
 MaterialParameterCollections are created through the [**Content Browser**](Engine/Content/Browser) via the **New** menu 
 or the **Right-Click** context menu and can be edited and saved like any other assets.
@@ -41,7 +41,7 @@ or the **Right-Click** context menu and can be edited and saved like any other a
 
 	![](MPC_Create_New_Collection.png)
 
-1. The Material Parameter Collection will be added and saved to the **Content Browser** just like any other Actor. 
+1. The Material Parameter Collection will be added and saved to the **Content Browser** just like any other Asset. 
 	
 	![](MPC_Saved_In_CB.png)
 
