#UE4doc #docs:
Substituted a word for another

diff --git a/Engine/Documentation/Source/Engine/AI/BehaviorTrees/QuickStart/1/1.INT.udn b/Engine/Documentation/Source/Engine/AI/BehaviorTrees/QuickStart/1/1.INT.udn
--- a/Engine/Documentation/Source/Engine/AI/BehaviorTrees/QuickStart/1/1.INT.udn
+++ b/Engine/Documentation/Source/Engine/AI/BehaviorTrees/QuickStart/1/1.INT.udn
@@ -23,7 +23,7 @@ Description:Step 1. Essential Setup
 
 %Steps%
 
-This Quick Start makes some assumptions about what is available in the project, as such, if you aren't comfortable making substations or problem solving on the fly, then you should setup the project as follows.
+This Quick Start makes some assumptions about what is available in the project, as such, if you aren't comfortable making substitutions or problem solving on the fly, then you should setup the project as follows.
 
 ## Steps
 
