#UE4doc #docs:
Added engine version field
Replaced image due to doc standards
Added link to Details panel page

diff --git a/Engine/Documentation/Source/Engine/Blueprints/Editor/UIBreakdowns/ClassBPUI/ClassBPUI.INT.udn b/Engine/Documentation/Source/Engine/Blueprints/Editor/UIBreakdowns/ClassBPUI/ClassBPUI.INT.udn
--- a/Engine/Documentation/Source/Engine/Blueprints/Editor/UIBreakdowns/ClassBPUI/ClassBPUI.INT.udn
+++ b/Engine/Documentation/Source/Engine/Blueprints/Editor/UIBreakdowns/ClassBPUI/ClassBPUI.INT.udn
@@ -2,6 +2,7 @@ Availability:Public
 Title:Blueprint Class UI
 Crumbs: %ROOT%, Engine/Blueprints, Engine/Blueprints/Editor
 Description:A breakdown of the UI elements of the Blueprint Editor when working on Blueprint Classes.
+version: 4.9
 
 [VAR:Topic]
 [OBJECT:Topic]
@@ -25,7 +26,7 @@ Description:A breakdown of the UI elements of the Blueprint Editor when working
 
 ![](ClassBlueprint.png)
 
-The **Blueprint Class** editor contains several tabs by default:
+The **Blueprint Class Editor** contains several tabs by default:
 
 
 |Default Visible UI Components| Available in the Window Menu|
@@ -41,6 +42,7 @@ The **Blueprint Class** editor contains several tabs by default:
 1. [Components](Engine/Blueprints/UserGuide/Components)
 1. [My Blueprint](Engine/Blueprints/Editor/UIComponents/MyBlueprint)
 1. [Graph Editor](Engine/Blueprints/Editor/UIComponents/GraphEditor)
+1. [Details panel](Engine/Blueprints/Editor/UIComponents/Details)
  
 [/EXCERPT:Visible]
 
