#UE4doc #docs:
DataOnlyUI:
Added engine version field
ClassBlueprint:
Edited datablueprint excerpt to remove italics from multiple terms

diff --git a/Engine/Documentation/Source/Engine/Blueprints/Editor/UIBreakdowns/DataOnlyUI/DataOnlyUI.INT.udn b/Engine/Documentation/Source/Engine/Blueprints/Editor/UIBreakdowns/DataOnlyUI/DataOnlyUI.INT.udn
--- a/Engine/Documentation/Source/Engine/Blueprints/Editor/UIBreakdowns/DataOnlyUI/DataOnlyUI.INT.udn
+++ b/Engine/Documentation/Source/Engine/Blueprints/Editor/UIBreakdowns/DataOnlyUI/DataOnlyUI.INT.udn
@@ -2,6 +2,7 @@ Availability:Public
 Title:Data-Only Blueprint UI
 Crumbs: %ROOT%, Engine/Blueprints, Engine/Blueprints/Editor
 Description:Overview of the Blueprint Class Editor interface.
+version: 4.9
 
 [VAR:Topic]
 [OBJECT:Topic]
diff --git a/Engine/Documentation/Source/Engine/Blueprints/UserGuide/Types/ClassBlueprint/ClassBlueprint.INT.udn b/Engine/Documentation/Source/Engine/Blueprints/UserGuide/Types/ClassBlueprint/ClassBlueprint.INT.udn
--- a/Engine/Documentation/Source/Engine/Blueprints/UserGuide/Types/ClassBlueprint/ClassBlueprint.INT.udn
+++ b/Engine/Documentation/Source/Engine/Blueprints/UserGuide/Types/ClassBlueprint/ClassBlueprint.INT.udn
@@ -129,11 +129,11 @@ Refer to the sections below for information on working with **Blueprint Classes*
 
 ## Data-Only Blueprint
 [EXCERPT:DataBlueprint]
-A **Data-Only Blueprint** is a _Blueprint Class_ that contains only the code (in the form of node graphs), variables, and components 
-inherited from its _parent_. These allow those inherited properties to be tweaked and modified, but no new elements can be added. 
+A **Data-Only Blueprint** is a Blueprint Class that contains only the code (in the form of node graphs), variables, and components 
+inherited from its parent. These allow those inherited properties to be tweaked and modified, but no new elements can be added. 
 These are essentially a replacement for archetypes and can be used to allow designers to tweak properties or set items with variations.
 		
-_Data-Only Blueprints_ are edited in a compact property editor, but can also be "converted" to full _Blueprints_ by simply adding code, 
+Data-Only Blueprint are edited in a compact property editor, but can also be "converted" to full Blueprints by simply adding code, 
 variables, or components using the full **Blueprint Editor**.
 [/EXCERPT:DataBlueprint]
 
