#UE4doc #docs:
Removed italics from a few terms in an excerpt

diff --git a/Engine/Documentation/Source/Engine/Blueprints/UserGuide/Types/Interface/InterfaceBlueprint.INT.udn b/Engine/Documentation/Source/Engine/Blueprints/UserGuide/Types/Interface/InterfaceBlueprint.INT.udn
--- a/Engine/Documentation/Source/Engine/Blueprints/UserGuide/Types/Interface/InterfaceBlueprint.INT.udn
+++ b/Engine/Documentation/Source/Engine/Blueprints/UserGuide/Types/Interface/InterfaceBlueprint.INT.udn
@@ -45,17 +45,17 @@ Description:Blueprints that declare functions to define an interface between Blu
 
 [EXCERPT:Intro]
 A **Blueprint Interface** is a collection of one or more functions - name only, no implementation - that can be added to 
-other _Blueprints_. Any _Blueprint_ that has the _Interface_ added is guaranteed to have those functions. The functions 
-of the _Interface_ can be given functionality in each of the _Blueprints_ that added it. This is essentially like the 
+other Blueprints. Any Blueprint that has the Interface added is guaranteed to have those functions. The functions 
+of the Interface can be given functionality in each of the Blueprints that added it. This is essentially like the 
 concept of an interface in general programming, which allows multiple different types of Objects to all share and be accessed 
-through a common interface. Put simply, _Blueprint Interfaces_ allow different _Blueprints_ to share with and send data to one another.
+through a common interface. Put simply, Blueprint Interfaces allow different Blueprints to share with and send data to one another.
 
-_Blueprint Interfaces_ can be made by content creators through the editor in a similar fashion to other _Blueprints_, but they 
+Blueprint Interfaces can be made by content creators through the editor in a similar fashion to other Blueprints, but they 
 come with certain limitations in that they cannot:
 
 * Add new variables
 * Edit graphs
-* Add _Components_
+* Add Components
 
 
 [/EXCERPT:Intro]
