#UE4doc #docs:
Added link to the documentation for Menu

diff --git a/Engine/Documentation/Source/Engine/Blueprints/Editor/UIBreakdowns/ClassBPUI/ClassBPUI.INT.udn b/Engine/Documentation/Source/Engine/Blueprints/Editor/UIBreakdowns/ClassBPUI/ClassBPUI.INT.udn
--- a/Engine/Documentation/Source/Engine/Blueprints/Editor/UIBreakdowns/ClassBPUI/ClassBPUI.INT.udn
+++ b/Engine/Documentation/Source/Engine/Blueprints/Editor/UIBreakdowns/ClassBPUI/ClassBPUI.INT.udn
@@ -37,7 +37,7 @@ The **Blueprint Class Editor** contains several tabs by default:
 
 <!----------
 [EXCERPT:Visible]
-1. Menu
+1. [Menu](Engine/Blueprints/Editor\UIComponents\Menu)
 1. [Toolbar](Engine/Blueprints/Editor/UIComponents/Toolbar/ToolbarOptions)
 1. [Components](Engine/Blueprints/UserGuide/Components)
 1. [My Blueprint](Engine/Blueprints/Editor/UIComponents/MyBlueprint)
