#UE4doc #docs:
Added engine version field

diff --git a/Engine/Documentation/Source/Engine/Blueprints/Editor/UIBreakdowns/InterfaceUI/InterfaceUI.INT.udn b/Engine/Documentation/Source/Engine/Blueprints/Editor/UIBreakdowns/InterfaceUI/InterfaceUI.INT.udn
--- a/Engine/Documentation/Source/Engine/Blueprints/Editor/UIBreakdowns/InterfaceUI/InterfaceUI.INT.udn
+++ b/Engine/Documentation/Source/Engine/Blueprints/Editor/UIBreakdowns/InterfaceUI/InterfaceUI.INT.udn
@@ -2,6 +2,7 @@ Availability:Public
 Title:Blueprint Interface UI
 Crumbs: %ROOT%, Engine, Engine/Blueprints, Engine/Blueprints/Editor
 Description:Overview of the Blueprint Class Editor interface.
+version: 4.9
 
 [VAR:Topic]
 [OBJECT:Topic]
