#UE4doc #docs:
Removed italics from a few terms in the Intro excerpt

diff --git a/Engine/Documentation/Source/Engine/Blueprints/UserGuide/Types/LevelBlueprint/LevelScript.INT.udn b/Engine/Documentation/Source/Engine/Blueprints/UserGuide/Types/LevelBlueprint/LevelScript.INT.udn
--- a/Engine/Documentation/Source/Engine/Blueprints/UserGuide/Types/LevelBlueprint/LevelScript.INT.udn
+++ b/Engine/Documentation/Source/Engine/Blueprints/UserGuide/Types/LevelBlueprint/LevelScript.INT.udn
@@ -47,7 +47,7 @@ Related: Engine/Blueprints
 
 [EXCERPT:Intro]
 A **Level Blueprint** is a specialized type of **Blueprint** that acts as a level-wide global event graph. 
-Each level in your project has its own _Level Blueprint_ created by default that can be edited within the Unreal Editor, however new _Level Blueprints_ 
+Each level in your project has its own Level Blueprint created by default that can be edited within the Unreal Editor, however new Level Blueprints 
 cannot be created through the editor interface. 
 
 Events pertaining to the level as a whole, or specific instances of Actors within the level, are 
