#UE4doc #docs:
Resubmitting as not all of the changes from the previous submit were saved

diff --git a/Engine/Documentation/Source/Engine/Blueprints/Editor/UIBreakdowns/LevelBPUI/LevelBPUI.INT.udn b/Engine/Documentation/Source/Engine/Blueprints/Editor/UIBreakdowns/LevelBPUI/LevelBPUI.INT.udn
--- a/Engine/Documentation/Source/Engine/Blueprints/Editor/UIBreakdowns/LevelBPUI/LevelBPUI.INT.udn
+++ b/Engine/Documentation/Source/Engine/Blueprints/Editor/UIBreakdowns/LevelBPUI/LevelBPUI.INT.udn
@@ -2,6 +2,7 @@ Availability:Public
 Title:Level Blueprint UI
 Crumbs: %ROOT%, Engine/Blueprints, Engine/Blueprints/Editor
 Description:Overview of the Level Blueprint Editor interface.
+version: 4.9
 
 [VAR:Topic]
 [OBJECT:Topic]
@@ -24,7 +25,6 @@ Description:Overview of the Level Blueprint Editor interface.
 [/VAR]
 
 
-[TOC (start:2)]
 
 [INCLUDE:Engine/Blueprints/UserGuide/Types/LevelBlueprint#Intro]
 
