#UE4doc #docs:
Resubmitting as not all changes were saved

diff --git a/Engine/Documentation/Source/Engine/Blueprints/UserGuide/Types/LevelBlueprint/LevelScript.INT.udn b/Engine/Documentation/Source/Engine/Blueprints/UserGuide/Types/LevelBlueprint/LevelScript.INT.udn
--- a/Engine/Documentation/Source/Engine/Blueprints/UserGuide/Types/LevelBlueprint/LevelScript.INT.udn
+++ b/Engine/Documentation/Source/Engine/Blueprints/UserGuide/Types/LevelBlueprint/LevelScript.INT.udn
@@ -55,7 +55,7 @@ used to fire off sequences of actions in the form of Function Calls or Flow Cont
 Those familiar with Unreal Engine 3 should be very familiar with this concept as this is very similar 
 to how Kismet worked in Unreal Engine 3. 
 
-_Level Blueprints_ also provide a control mechanism for level streaming and [Matinee](Engine/Matinee) as well 
+Level Blueprints also provide a control mechanism for level streaming and [Matinee](Engine/Matinee) as well 
 as for binding events to Actors placed within the level. 
 [/EXCERPT:Intro]
 
