#UE4doc #docs:
Removed italics from multiple terms in Intro excerpt

diff --git a/Engine/Documentation/Source/Engine/Blueprints/UserGuide/Types/MacroLibrary/MacroLibrary.INT.udn b/Engine/Documentation/Source/Engine/Blueprints/UserGuide/Types/MacroLibrary/MacroLibrary.INT.udn
--- a/Engine/Documentation/Source/Engine/Blueprints/UserGuide/Types/MacroLibrary/MacroLibrary.INT.udn
+++ b/Engine/Documentation/Source/Engine/Blueprints/UserGuide/Types/MacroLibrary/MacroLibrary.INT.udn
@@ -44,13 +44,13 @@ Description:Blueprints that declare and define macros for reuse in other Bluepri
 [TOC(start:2)]
 
 [EXCERPT:Intro]
-A **Blueprint Macro Library** is a container that holds a collection of _Macros_ or self-contained graphs that can 
-be placed as nodes in other _Blueprints_. These can be time-savers as they can store commonly used sequences of nodes 
+A **Blueprint Macro Library** is a container that holds a collection of **Macros** or self-contained graphs that can 
+be placed as nodes in other Blueprints. These can be time-savers as they can store commonly used sequences of nodes 
 complete with inputs and outputs for both execution and data transfer.
 
-_Macros_ are shared among all graphs that reference them, but they are auto-expanded into graphs as if they were a 
-collapsed node during compiling. This means that _Blueprint Macro Libraries_ do not need to be compiled. However, changes 
-to a _Macro_ are only reflected in graphs that reference that _Macro_ when the _Blueprint_ containing those graphs is 
+Macros are shared among all graphs that reference them, but they are auto-expanded into graphs as if they were a 
+collapsed node during compiling. This means that Blueprint Macro Libraries do not need to be compiled. However, changes 
+to a Macro are only reflected in graphs that reference that Macro when the Blueprint containing those graphs is 
 recompiled.
 [/EXCERPT:Intro]
 
