#UE4doc #docs:
Added engine version field
Removed "On this page" field as there was only one entry

diff --git a/Engine/Documentation/Source/Engine/Blueprints/Editor/UIBreakdowns/MacroLibraryUI/MacroLibraryUI.INT.udn b/Engine/Documentation/Source/Engine/Blueprints/Editor/UIBreakdowns/MacroLibraryUI/MacroLibraryUI.INT.udn
--- a/Engine/Documentation/Source/Engine/Blueprints/Editor/UIBreakdowns/MacroLibraryUI/MacroLibraryUI.INT.udn
+++ b/Engine/Documentation/Source/Engine/Blueprints/Editor/UIBreakdowns/MacroLibraryUI/MacroLibraryUI.INT.udn
@@ -2,6 +2,7 @@ Availability:Public
 Title:Macro Library UI
 Crumbs: %ROOT%, Engine/Blueprints, Engine/Blueprints/Editor
 Description:Overview of the Blueprint Class Editor interface.
+version: 4.9
 
 [VAR:Topic]
 [OBJECT:Topic]
@@ -23,7 +24,6 @@ Description:Overview of the Blueprint Class Editor interface.
 [/OBJECT]
 [/VAR]
 
-[TOC (start:2)]
 
 [INCLUDE:Engine/Blueprints/UserGuide/Types/MacroLibrary#intro]
 
