#UE4doc #docs:
Added engine version
Removed italics from a few terms

diff --git a/Engine/Documentation/Source/Engine/Blueprints/Editor/UIComponents/CompilerResults/CompilerResults.INT.udn b/Engine/Documentation/Source/Engine/Blueprints/Editor/UIComponents/CompilerResults/CompilerResults.INT.udn
--- a/Engine/Documentation/Source/Engine/Blueprints/Editor/UIComponents/CompilerResults/CompilerResults.INT.udn
+++ b/Engine/Documentation/Source/Engine/Blueprints/Editor/UIComponents/CompilerResults/CompilerResults.INT.udn
@@ -2,13 +2,14 @@ Availability:Public
 Title: Compiler Results
 Crumbs:%ROOT%, Engine/Blueprints, Engine/Blueprints/Editor, Engine/Blueprints/Editor/UIBreakdowns/LevelBPUI
 Description:Interface for the Compiler Results panel of the Blueprint Editor.
+version: 4.9
 
 
 [REGION:fullwidth]
 ![](CompilerWindow.png)
 [/REGION]
 
-The **Compiler Results** panel provides feedback on the compiling of your Blueprint scripts, telling you whether a Blueprint script compile was successful or if there were _errors_ or _warnings_. Whenever errors or warnings are found, the **Compiler Results** panel gives information on the error and provides hyperlink access, enabling you to jump the graph view directly to the source of the problem.
+The **Compiler Results** panel provides feedback on the compiling of your Blueprint scripts, telling you whether a Blueprint script compile was successful or if there were errors or warnings. Whenever errors or warnings are found, the **Compiler Results** panel gives information on the error and provides hyperlink access, enabling you to jump the graph view directly to the source of the problem.
 
 ## Interface
 
