#UE4doc #docs:
Added engine version field
Entered a description as the field was empty

diff --git a/Engine/Documentation/Source/Engine/Blueprints/Editor/UIComponents/Debug/Debug.INT.udn b/Engine/Documentation/Source/Engine/Blueprints/Editor/UIComponents/Debug/Debug.INT.udn
--- a/Engine/Documentation/Source/Engine/Blueprints/Editor/UIComponents/Debug/Debug.INT.udn
+++ b/Engine/Documentation/Source/Engine/Blueprints/Editor/UIComponents/Debug/Debug.INT.udn
@@ -1,7 +1,8 @@
 Availability:Public
 Title: Debug Panel
 Crumbs:%ROOT%, Engine/Blueprints, Engine/Blueprints/Editor, Engine/Blueprints/Editor/UIBreakdowns/LevelBPUI
-Description:
+Description: Panel that provides a list of all debug tools currently in use for the Blueprint.
+version: 4.9
 
 The **Debug** panel provides a list of all debug tools currently in use for the Blueprint, such as Breakpoints and Watch Values. During Play-In-Editor, it also provides access to the Execution Trace, showing you each node execution taking place within a given Blueprint.
 
