#UE4doc #docs:
Added engine version field

diff --git a/Engine/Documentation/Source/Engine/Blueprints/Editor/UIComponents/Details/KismetDetails.INT.udn b/Engine/Documentation/Source/Engine/Blueprints/Editor/UIComponents/Details/KismetDetails.INT.udn
--- a/Engine/Documentation/Source/Engine/Blueprints/Editor/UIComponents/Details/KismetDetails.INT.udn
+++ b/Engine/Documentation/Source/Engine/Blueprints/Editor/UIComponents/Details/KismetDetails.INT.udn
@@ -2,7 +2,7 @@ Availability:Public
 Title: Details Panel
 Crumbs:%ROOT%, Engine/Blueprints, Engine/Blueprints/Editor, Engine/Blueprints/Editor/UIBreakdowns/LevelBPUI
 Description:Overview of the Details panel within the Blueprint Editor.
-
+version: 4.9
 
 
 
