#UE4doc #docs:
Added engine version field
Added description as the field was empty
Removed italics from and bolded a term

diff --git a/Engine/Documentation/Source/Engine/Blueprints/Editor/UIComponents/FindResults/FindResults.INT.udn b/Engine/Documentation/Source/Engine/Blueprints/Editor/UIComponents/FindResults/FindResults.INT.udn
--- a/Engine/Documentation/Source/Engine/Blueprints/Editor/UIComponents/FindResults/FindResults.INT.udn
+++ b/Engine/Documentation/Source/Engine/Blueprints/Editor/UIComponents/FindResults/FindResults.INT.udn
@@ -1,8 +1,8 @@
 Availability:Public
 Title: Find Result Panel
 Crumbs: %ROOT%, Engine, Engine/Blueprints, Engine/Blueprints/Editor, Engine/Blueprints/Editor/UIBreakdowns/LevelBPUI
-Description:
-
+Description: Panel that is a search tool within the Blueprint Editor that allows you to quickly track down a variety of objects.
+version: 4.9
 
 
 The **Find Results** panel is a powerful search tool within the Blueprint Editor that allows you to quickly track down a variety of objects based on the following criteria:
@@ -18,7 +18,7 @@ As the **Find Results** panel tracks down search matches, it will display a list
 Much like searching in a web browser, the **Find Results** panel is available by pressing **Ctrl-F** while working in the Blueprint Editor. By default, the panel will appear along the bottom of the **Graph** Panel. If the **Compiler Results** panel is showing, then the **Find Results** panel will dock alongside it.
 
 [REGION:note]
-Unlike many search fields in Unreal Engine 4, the **Find Results** panel does not filter results actively while you type, due to the sheer volume of possible results. Once you press _Enter_, the list will populate.
+Unlike many search fields in Unreal Engine 4, the **Find Results** panel does not filter results actively while you type, due to the sheer volume of possible results. Once you press **Enter**, the list will populate.
 [/REGION]
 
 ## Interface
