#UE4doc #docs:
Added engine version field
Set publishing to "Public" from "NoPublish"
Added a few entries and removed one
Removed some italics
Added a description to the description field

diff --git a/Engine/Documentation/Source/Engine/Blueprints/Editor/UIComponents/Menu/BlueprintEditorMenu.INT.udn b/Engine/Documentation/Source/Engine/Blueprints/Editor/UIComponents/Menu/BlueprintEditorMenu.INT.udn
--- a/Engine/Documentation/Source/Engine/Blueprints/Editor/UIComponents/Menu/BlueprintEditorMenu.INT.udn
+++ b/Engine/Documentation/Source/Engine/Blueprints/Editor/UIComponents/Menu/BlueprintEditorMenu.INT.udn
@@ -1,7 +1,8 @@
-Availability:NoPublish
+Availability:Public
 Title: Menu
 Crumbs:%ROOT%, Engine/Blueprints, Engine/Blueprints/Editor
-Description:Menu
+Description: Menu that is used to do multiple actions, such as saving, loading and opening windows.
+version: 4.9
 
 
 
@@ -47,11 +48,13 @@ The **Menu** is displayed at the top of the Blueprint Editor by default.
 | [REGION:tablesection]**History**[/REGION] ||
 | **Undo** | Undo the last action. |
 | **Redo** | Redo the last undone action. |
+| **Undo History** | Displays the entire Undo History. |
 | [REGION:tablesection]**Search**[/REGION] ||
 %BlueprintSearch%
 | [REGION:tablesection]**Configuration**[/REGION] ||
 | **Editor Preferences** | Opens the settings for the Editor. |
 | **Project Settings** | Opens the settings for the current project. |
+| **Plugins** | Opens the Plugin Browser tab. |
 [/EXCERPT:EditMenu]
 
 
@@ -61,7 +64,8 @@ The **Menu** is displayed at the top of the Blueprint Editor by default.
 | Command | Description |
 | --- | --- |
 | **Find in Content Browser...** | Summons the **Content Browser** and navigates to this asset. |
-| **View References** | Launches the reference viewer to show what the current asset references and what references the current asset.|
+| **Reference Viewer...** | Launches the reference viewer to show what the current asset references and what references the current asset.|
+| **Size Map...** | Displays an interactive map showing the approximate size of this asset and everything it references.|
 [/VAR]
 
 %assetmenu%
@@ -77,7 +81,7 @@ The **Menu** is displayed at the top of the Blueprint Editor by default.
 | **Hide Unused Pins** | Hides all pins with no connections and no default values. |
 | **Hide Unconnected Pins** | Hides all pins with no connections.  This option will hide input pins for inputs that have been set directly on the node. |
 | [REGION:tablesection]**Zoom**[/REGION] ||
-| **Zoom to Window ** | Fits the current view to the entire graph. |
+| **Zoom to Graph Extents ** | Fits the current view to the entire graph. |
 | **Zoom to Selection** | Fits the current view to the selection.  |
 [/EXCERPT:ViewMenu]
 
@@ -96,7 +100,7 @@ The **Menu** is displayed at the top of the Blueprint Editor by default.
 
 ### Window 
 
-The **Window** menu in the **Blueprint Editor** has a specific subsection for displaying _Blueprint Editor_-specific tabs. Not all tabs will be present in the menu while the _Blueprint Editor_ is in the **Defaults** and **Components** modes.
+The **Window** menu in the **Blueprint Editor** has a specific subsection for displaying Blueprint Editor-specific tabs. Not all tabs will be present in the menu while the Blueprint Editor is in the **Defaults** and **Components** modes.
 
 
 | Command | Description |
@@ -108,6 +112,7 @@ The **Window** menu in the **Blueprint Editor** has a specific subsection for di
 | **My Blueprint** | Shows the [](Engine/Blueprints\Editor/UIComponents\MyBlueprint) pane if it is not currently visible.  |
 | **Compiler Results** | Shows the **Compiler Results** pane if it is not currently visible.  |
 | **Find Results** | Shows the **Find Results** pane if it is not currently visible.  | 
+| **Components** | Shows the **Components panel** if it is not currently visible. |
 | **Viewport** | Shows the **Preview Viewport** pane if it is not currently visible.  |
 | **Content Browser** | Opens a sub-menu with access to all four **Content Browsers**. |
 | **Developer Tools** | [INCLUDE:#devtools] |
@@ -134,7 +139,6 @@ The **Window** menu in the **Blueprint Editor** has a specific subsection for di
 | **Class Viewer** | Opens the Class Viewer. |
 | **Device Manager** | Opens the Device Manager. |
 | **Device Profiles** | Opens the Device Profiles Tab. |
-| **Modules** | Opens the Modules Panel. |
 | **Session Frontend** | Shows the Session Frontend. Please see the [](Engine/Deployment/UnrealFrontend) documentation for more information. |
 | **Widget Reflector** | Opens the Widget Reflector. This enables you to see the various Slate elements that make up the interfaces within the Editor interface. |
 [/REGION]
