#UE4doc #docs:
Added engine version field
Added a control to the Controls excerpt
Replaced image due to bad coloring

diff --git a/Engine/Documentation/Source/Engine/Blueprints/Editor/UIComponents/GraphEditor/GraphEditor.INT.udn b/Engine/Documentation/Source/Engine/Blueprints/Editor/UIComponents/GraphEditor/GraphEditor.INT.udn
--- a/Engine/Documentation/Source/Engine/Blueprints/Editor/UIComponents/GraphEditor/GraphEditor.INT.udn
+++ b/Engine/Documentation/Source/Engine/Blueprints/Editor/UIComponents/GraphEditor/GraphEditor.INT.udn
@@ -2,9 +2,7 @@ Availability:Public
 Title: Graph Editor Tab
 Crumbs:%ROOT%, Engine/Blueprints, Engine/Blueprints/Editor, Engine/Blueprints/Editor/UIBreakdowns/LevelBPUI
 Description:Editable, visual representation of a network of nodes within a Blueprint.
-
-
-[TOC(start:2 end:3)]
+version: 4.9
 
 
 [EXCERPT:Intro]
@@ -47,6 +45,7 @@ The **Graph Editor** tab can be navigated using the controls below:
 | **Click + Drag** from pin to pin | Wires the pins together. |
 | **Ctrl + Click + Drag** from pin to pin | Moves wires from origin pin to destination pin. |
 | **Click + Drag** from pin to empty space | Brings up context menu showing only relevant nodes. Wires the original pin to a compatible pin on the created node. |
+| **Alt + Click** on a pin | Remove all wires connected to the selected pin. |
 
 [/EXCERPT:Controls]
 
