#UE4doc #docs - Added Engine Version field.

diff --git a/Engine/Documentation/Source/Engine/Physics/WalkableSlope/WalkableSlope.INT.udn b/Engine/Documentation/Source/Engine/Physics/WalkableSlope/WalkableSlope.INT.udn
--- a/Engine/Documentation/Source/Engine/Physics/WalkableSlope/WalkableSlope.INT.udn
+++ b/Engine/Documentation/Source/Engine/Physics/WalkableSlope/WalkableSlope.INT.udn
@@ -2,6 +2,7 @@ Availability: Public
 Title:Walkable Slope
 Crumbs: %ROOT%, Engine, Engine/Physics
 Description:
+Version: 4.9
 
 [VAR:Topic]
 [OBJECT:Topic]
