#UE4doc #docs - Added Engine Version Field.

diff --git a/Engine/Documentation/Source/Engine/Rendering/Rendering.INT.udn b/Engine/Documentation/Source/Engine/Rendering/Rendering.INT.udn
--- a/Engine/Documentation/Source/Engine/Rendering/Rendering.INT.udn
+++ b/Engine/Documentation/Source/Engine/Rendering/Rendering.INT.udn
@@ -3,6 +3,7 @@ Title: Rendering and Graphics
 Crumbs: %ROOT%, Engine
 Description:Rendering subsystem including lighting and shadowing, materials and textures, particles and effects, and post processing.
 Related:Programming/Rendering
+Version: 4.9
 
 [VAR:Topic]
 [OBJECT:Topic]
