#UE4doc #docs - Adde Engine Version Field.

diff --git a/Engine/Documentation/Source/Engine/Rendering/LightingAndShadows/LightingAndShadows.INT.udn b/Engine/Documentation/Source/Engine/Rendering/LightingAndShadows/LightingAndShadows.INT.udn
--- a/Engine/Documentation/Source/Engine/Rendering/LightingAndShadows/LightingAndShadows.INT.udn
+++ b/Engine/Documentation/Source/Engine/Rendering/LightingAndShadows/LightingAndShadows.INT.udn
@@ -4,6 +4,7 @@ Crumbs: %ROOT%, Engine
 Description:Using Lighting and shadowing geometry, using Glogal Illumination, and setting up reflections.
 Navigation:topic
 Related:Engine/Rendering/ParticleSystems/ParticleLights
+Version: 4.9
 
 [VAR:Topic]
 [OBJECT:Topic]
