#UE4doc #docs - Added Engine Version field.

diff --git a/Engine/Documentation/Source/Engine/Rendering/LightingAndShadows/AmbientCubemap/AmbientCubemap.INT.udn b/Engine/Documentation/Source/Engine/Rendering/LightingAndShadows/AmbientCubemap/AmbientCubemap.INT.udn
--- a/Engine/Documentation/Source/Engine/Rendering/LightingAndShadows/AmbientCubemap/AmbientCubemap.INT.udn
+++ b/Engine/Documentation/Source/Engine/Rendering/LightingAndShadows/AmbientCubemap/AmbientCubemap.INT.udn
@@ -2,6 +2,7 @@ Availability:Public
 Title:Ambient Cubemaps
 Crumbs: %ROOT%, Engine, Engine/Rendering/LightingAndShadows
 Description:Cubemap texture applied to scene lighting for high quality lighting environments.
+Version: 4.9
 
 [VAR:Topic]
 [OBJECT:Topic]
@@ -30,7 +31,7 @@ The blurry versions are computed in a way that they can be used for specular lig
 This effect is independent of the position a material is lit from. Viewer position, material roughness (for specular effects), and the material surface normal are all taken into account. 
 
 [REGION:note]
-The Ambient Cubemap has very limited local shadowing (SSAO only).  Use a [Sky Light](Engine/Rendering/LightingAndShadows/LightTypes/SkyLight) instead when representing the sky's lighting.
+The Ambient Cubemap has very limited local shadowing (SSAO only). Use a [Sky Light](Engine/Rendering/LightingAndShadows/LightTypes/SkyLight) instead when representing the sky's lighting.
 [/REGION]
 [/EXCERPT:AmbientCubemap_Desc]
 
