#UE4doc #docs - Added Engine Version Field. Added Comparison Slider for Lower vs Higher lightmap resolutions.

diff --git a/Engine/Documentation/Source/Engine/Rendering/LightingAndShadows/LightMobility/StaticLights/StaticLights.INT.udn b/Engine/Documentation/Source/Engine/Rendering/LightingAndShadows/LightMobility/StaticLights/StaticLights.INT.udn
--- a/Engine/Documentation/Source/Engine/Rendering/LightingAndShadows/LightMobility/StaticLights/StaticLights.INT.udn
+++ b/Engine/Documentation/Source/Engine/Rendering/LightingAndShadows/LightMobility/StaticLights/StaticLights.INT.udn
@@ -2,6 +2,7 @@ Availability: Public
 Title:Static Lights
 Crumbs: %ROOT%, Engine, Engine/Rendering/LightingAndShadows,Engine/Rendering/LightingAndShadows/LightMobility
 Description:A completely static light that has no overhead during a running game.
+Version: 4.9
 
 [VAR:Topic]
 [OBJECT:Topic]
@@ -49,7 +50,6 @@ Description:A completely static light that has no overhead during a running game
 [TOC(start:2)]
 
 
-
 [INCLUDE:Resources/ContentExamples\Lighting\2_1#StaticLighting01]
 
 The primary use case for static lights is for low power devices on mobile platforms.
@@ -64,8 +64,22 @@ As mentioned above, **Light Source Radius** for Static Lights have the additiona
 
 **Lightmap Resolution** allows you to control the detail of the baked lighting that static lights produce.
 
-[INCLUDE:Resources/ContentExamples\Lighting\2_1#StaticLighting02]
+On Static Mesh Components, lightmap resolution is set either on the Static Mesh asset, or by checking the Override Lightmap Res checkbox and setting the value. Larger values mean higher resolution, but also longer build times and more memory. 
+
+On Brush surfaces, lightmap resolution is set via the Lightmap Resolution property. This is really setting a density, so lower values produce higher resolution. 
 
+[OBJECT:ComparisonSlider]
+	[PARAM:before]
+	![Lightmap Resolution 4.0](LightmapRes_4.png)(w:600)
+	[/PARAM]
+	[PARAM:after]
+	![Lightmap Resolution 16.0](LightmapRes_16.png)(w:600)
+	[/PARAM]
+[/OBJECT]
+
+<!-- 
+[INCLUDE:Resources/ContentExamples\Lighting\2_1#StaticLighting02]
+-->
 <!-- Long Description Excerpt
 [EXCERPT:Static001]
 Static light cast completely static light and shadows. What this means is their effect gets baked down in **Lightmass** to light maps and into the indirect lighting cache. **Static lights have no direct affect on movable/dynamic actors nor do the create any specular highlights on any objects**, but since they will bake into the indirect lighting cache, they will have some _minor_ indirect effect on movable/dynamic objects.
