#UE4doc #docs - Added Engine Version Field. Added Comparison Slider for Translucent Fog images.

diff --git a/Engine/Documentation/Source/Engine/Rendering/LightingAndShadows/LightMobility/StationaryLights/StationaryLights.INT.udn b/Engine/Documentation/Source/Engine/Rendering/LightingAndShadows/LightMobility/StationaryLights/StationaryLights.INT.udn
--- a/Engine/Documentation/Source/Engine/Rendering/LightingAndShadows/LightMobility/StationaryLights/StationaryLights.INT.udn
+++ b/Engine/Documentation/Source/Engine/Rendering/LightingAndShadows/LightMobility/StationaryLights/StationaryLights.INT.udn
@@ -2,6 +2,7 @@ Availability:Public
 Title:Stationary lights
 Crumbs: %ROOT%, Engine, Engine/Rendering/LightingAndShadows,Engine/Rendering/LightingAndShadows/LightMobility
 Description:Light type that can change its color and brightness at runtime but cannot move, rotate or change influence size.
+Version: 4.9
 
 [VAR:Topic]
 [OBJECT:Topic]
@@ -77,7 +78,7 @@ stationary lights have the ability to have static shadowing on static objects, b
 
 #### On Opaque
 
-[Lightmass](Engine/Rendering/LightingAndShadows/Lightmass) generates _distance field 
+[**Lightmass**](Engine/Rendering/LightingAndShadows/Lightmass) generates _distance field 
 shadow maps_ for stationary lights on static objects during the lighting rebuild. Distance field shadow maps 
 provide very accurate shadow transitions even at low resolutions, and with very little runtime cost. Like 
 lightmaps, distance field shadow maps require uniquely unwrapped UVs on all _StaticMeshes_ using static 
@@ -114,11 +115,15 @@ StationaryLightOverlap viewmode showing one too many overlapping lights. Notice
 
 Translucency also receives shadowing very cheaply with Stationary lights - Lightmass precomputes a shadow depth map from static geometry which is applied to translucency at runtime.  This form of shadowing is fairly coarse and only captures shadowing on the scale of meters.
 
-[![](1OriginalCropped.png)(w:700)](1OriginalCropped.png)[![](1ShadowedCropped.png)(w:700)](1ShadowedCropped.png)
+[OBJECT:ComparisonSlider]
+	[PARAM:before]
+	![Unshadowed Translucency](1OriginalCropped.png)(w:700)
+	[/PARAM]
+	[PARAM:after]
+	![Translucency recieving static shadowing form a Diretional Light](1ShadowedCropped.png)(w:700)
+	[/PARAM]
+[/OBJECT]
 
-[REGION:caption]
-Top: Unshadowed translucency Bottom: Translucency receiving static shadowing from a directional light.
-[/REGION]
 
 ### Dynamic Shadowing
 
@@ -164,4 +169,4 @@ be modified at runtime by changing brightness and color like direct lighting can
 light has **Visible** unchecked, its indirect lighting will be put in the lightmap during the lighting build. 
 **IndirectLightingIntensity** on the light can be used to scale or disable the indirect lighting from a given light at lighting build time.  
 
- However there's a post process volume setting called **IndirectLightingIntensity** which lets you scale the contribution of the lightmap for all lights, which can be changed at runtime from a blueprint. 
+However there's a post process volume setting called **IndirectLightingIntensity** which lets you scale the contribution of the lightmap for all lights, which can be changed at runtime from a blueprint. 
