#UE4doc #docs - Added Engine Version Field. Removed redundant images in GPU Cost section.

diff --git a/Engine/Documentation/Source/Engine/Rendering/LightingAndShadows/LightShafts/LightShafts.INT.udn b/Engine/Documentation/Source/Engine/Rendering/LightingAndShadows/LightShafts/LightShafts.INT.udn
--- a/Engine/Documentation/Source/Engine/Rendering/LightingAndShadows/LightShafts/LightShafts.INT.udn
+++ b/Engine/Documentation/Source/Engine/Rendering/LightingAndShadows/LightShafts/LightShafts.INT.udn
@@ -2,6 +2,7 @@ Availability:Public
 Title:Light Shafts
 Crumbs: %ROOT%, Engine, Engine/Rendering/LightingAndShadows
 Description:Simulating light scattered in the atmosphere.
+Version: 4.9
 
 [VAR:Topic]
 [OBJECT:Topic]
@@ -92,6 +93,8 @@ The bloom method captures the scene's color (including translucency and fog in-s
   
 These effects (occlusion and bloom) cost separately, so it is best to only enable the one you want. At 1080p on a 680GTX, occlusion is .5ms and bloom on a single light is .68ms. 
 
+<!-- Redundant Images
 [REGION:lightbox]
 [![](2LightShafts.png)(h:250)](2LightShafts.png) [![](5LightShafts.png)(h:250)](5LightShafts.png) 
 [/REGION]
+-->
\ No newline at end of file
