#UE4doc #docs - Added Engine Version Field.

diff --git a/Engine/Documentation/Source/Engine/Rendering/LightingAndShadows/LightTypes/LightTypes.INT.udn b/Engine/Documentation/Source/Engine/Rendering/LightingAndShadows/LightTypes/LightTypes.INT.udn
--- a/Engine/Documentation/Source/Engine/Rendering/LightingAndShadows/LightTypes/LightTypes.INT.udn
+++ b/Engine/Documentation/Source/Engine/Rendering/LightingAndShadows/LightTypes/LightTypes.INT.udn
@@ -3,6 +3,7 @@ Title:Types of Lights
 Crumbs: %ROOT%, Engine, Engine/Rendering/LightingAndShadows
 Description:The types of lights available in Unreal Engine 4.
 DescriptionLong:The types of lights available in Unreal Engine 4: Point, Spot, Directional, and Sky.
+Version: 4.9
 
 [VAR:Topic]
 [OBJECT:Topic]
