#UE4doc #docs - Added Engine Version Field. Added Comparison Slider for Directional Light shadow frustum enabled.

diff --git a/Engine/Documentation/Source/Engine/Rendering/LightingAndShadows/LightTypes/Directional/Directional.INT.udn b/Engine/Documentation/Source/Engine/Rendering/LightingAndShadows/LightTypes/Directional/Directional.INT.udn
--- a/Engine/Documentation/Source/Engine/Rendering/LightingAndShadows/LightTypes/Directional/Directional.INT.udn
+++ b/Engine/Documentation/Source/Engine/Rendering/LightingAndShadows/LightTypes/Directional/Directional.INT.udn
@@ -2,6 +2,7 @@ Availability: Public
 Title:Directional Lights
 Crumbs: %ROOT%, Engine, Engine/Rendering/LightingAndShadows, Engine/Rendering/LightingAndShadows/LightTypes
 Description:The basics of Directional Lights.
+Version: 4.9
 
 [VAR:Topic]
 [OBJECT:Topic]
@@ -95,9 +96,20 @@ Below you will see a representation of sunlight being cast through an opening in
 [/REGION]
 
 
+[OBJECT:ComparisonSlider]
+	[PARAM:before]
+	![Directional Light](directional_001.png)(w:450)
+	[/PARAM]
+	[PARAM:after]
+	![Directional Light Shadow Frustum](directional_002.png)(w:450)
+	[/PARAM]
+[/OBJECT]
+
+<!--
 [REGION:lightbox]
 [![](directional_001.png)(w:450)](directional_001.png) [![](directional_002.png)(w:450)](directional_002.png)
 [/REGION]
+-->
 
 The left image shows just the light, the right image has its **Shadow Frustum** enabled in order to see that the Directional Light is casting parallel light rays. 
 
