#UE4doc #docs - Added Engine Version Field. Added comparison slider to point light radius.

diff --git a/Engine/Documentation/Source/Engine/Rendering/LightingAndShadows/LightTypes/Point/Point.INT.udn b/Engine/Documentation/Source/Engine/Rendering/LightingAndShadows/LightTypes/Point/Point.INT.udn
--- a/Engine/Documentation/Source/Engine/Rendering/LightingAndShadows/LightTypes/Point/Point.INT.udn
+++ b/Engine/Documentation/Source/Engine/Rendering/LightingAndShadows/LightTypes/Point/Point.INT.udn
@@ -2,6 +2,7 @@ Availability: Public
 Title:Point Lights
 Crumbs: %ROOT%, Engine, Engine/Rendering/LightingAndShadows, Engine/Rendering/LightingAndShadows/LightTypes
 Description:The basics of Point Lights.
+Version: 4.9
 
 [VAR:Topic]
 [OBJECT:Topic]
@@ -84,7 +85,14 @@ Description:The basics of Point Lights.
 
 Two examples of a Point Light placed inside a level are shown below.
 
-![](pointLight_004.png)(w:450) ![](pointLight_005.png)(w:450 convert:false)
+[OBJECT:ComparisonSlider]
+	[PARAM:before]
+	![Point Light Without Radius Showing](pointLight_004.png)(w:450)
+	[/PARAM]
+	[PARAM:after]
+	![Point Light With Radius Showing.](pointLight_005.png)(w:450)
+	[/PARAM]
+[/OBJECT]
 
 The image on the left is a Point Light without its radius showing, and the same light with it enabled is shown in the right image, giving a good impression of where the light will affect the world.
 
@@ -105,16 +113,16 @@ Property | Description |
 | **Attenuation Radius** | Bounds the light's visible influence.|
 | **Source Radius** | Radius of light source shape. |
 | **Source Length** | Length of light source shape. |
-| **Affects World** | Disables the light completely. Cannot be set at run time. To disable a light's effect during runtime, change its **Visibility** property. |
+| **Affects World** | Disables the light completely. Cannot be set at run time. To disable a light's effect during runtime, change its Visibility property. |
 | **Casts Shadows** | If the light casts shadows. |
 | **Indirect Lighting Intensity** | Scales the indirect lighting contribution from the light. |
-| **Use Inverse Squared Falloff** | Whether to use physically based inverse squared distance falloff, where **AttenuationRadius** is only clamping the light's contribution. |
-| **Light Falloff Exponent** | Controls the radial falloff of light when **UseInverseSquaredFalloff** is disabled. |
+| **Use Inverse Squared Falloff** | Whether to use physically based inverse squared distance falloff, where AttenuationRadius is only clamping the light's contribution. |
+| **Light Falloff Exponent** | Controls the radial falloff of light when UseInverseSquaredFalloff is disabled. |
 | **Min Roughness** | Min roughness effective for this light, used for softening specular highlights. |
 | **Shadow Bias** | Controls how accurate the shadows from this light are. |
 | **Shadow Filter Sharpen** | How much to sharpen shadow filtering for this light. |
 | **Cast Translucent Shadows** | Whether this light is allowed to cast dynamic shadows through translucent objects. |
-| **Affect Dynamic Indirect Lighting** | Whether the light should be injected into the **Light Propagation Volume**.|
+| **Affect Dynamic Indirect Lighting** | Whether the light should be injected into the Light Propagation Volume.|
 | **Cast Static Shadows** | Whether this light casts static shadows. |
 | **Cast Dynamic Shadows** | Whether this light casts dynamic shadows. |
 | **Affect Translucent Lighting** | Whether the light affects translucency or not. |
@@ -123,7 +131,7 @@ Property | Description |
 Property | Description |
 | -------- | ----------- |
 | **IES Texture** | The IES "Texture" used for the light profile. IES files are ASCII though Unreal represents them as textures, they are not image files. |
-| **Use IES Brightness** | If _false_, it will use the **Brightness** of the light to determine how much light to produce. If _true_, it will use the IES files brightness in Lumens (usually much larger than default values on lights in Unreal). |
+| **Use IES Brightness** | If _false_, it will use the Brightness of the light to determine how much light to produce. If _true_, it will use the IES files brightness in Lumens (usually much larger than default values on lights in Unreal). |
 | **IES Brightness Scale** | Scale for IES brightness contribution, as they can seriously blow out a scene. |
 
 ### Lightmass
@@ -132,12 +140,11 @@ Property | Description |
 | **Indirect Lighting Saturation** | A value of 0 will completely desaturate this light in Lightmass, 1 will be unchanged. |
 | **Shadow Exponent** | Controls the falloff of shadow penumbras. |
 
-
 ### Light Function
 Property | Description |
 | -------- | ----------- |
 | **Light Function Material** | The light function material to be applied to this light. |
 | **Light Function Scale** | Scales the light function projection. |
-| **Light Function Fade Distance** | The distance at which the light function should be completely faded to the value in **Disabled Brightness**. |
-| **Disabled Brightness** | Brightness factor applied to the light when the light function is specified but disabled, say from the property above: **Light Function Fade Distance**. |
+| **Light Function Fade Distance** | The distance at which the light function should be completely faded to the value in Disabled Brightness. |
+| **Disabled Brightness** | Brightness factor applied to the light when the light function is specified but disabled, say from the property above: Light Function Fade Distance. |
 
