#UE4doc #docs - Added Engine Version Field. Replaced two images for better comparison between skylight and no skylight.

diff --git a/Engine/Documentation/Source/Engine/Rendering/LightingAndShadows/LightTypes/SkyLight/SkyLight.INT.udn b/Engine/Documentation/Source/Engine/Rendering/LightingAndShadows/LightTypes/SkyLight/SkyLight.INT.udn
--- a/Engine/Documentation/Source/Engine/Rendering/LightingAndShadows/LightTypes/SkyLight/SkyLight.INT.udn
+++ b/Engine/Documentation/Source/Engine/Rendering/LightingAndShadows/LightTypes/SkyLight/SkyLight.INT.udn
@@ -2,6 +2,7 @@ Availability: Public
 Title:Sky Lights
 Crumbs: %ROOT%, Engine, Engine/Rendering/LightingAndShadows, Engine/Rendering/LightingAndShadows/LightTypes
 Description:The basics of Sky Lights.
+Version: 4.9
 
 [VAR:Topic]
 [OBJECT:Topic]
@@ -52,12 +53,12 @@ Description:The basics of Sky Lights.
 [EXCERPT:SL001]
 The **Sky Light** captures the distant parts of your level (everything further than `SkyDistanceThreshold`) and applies that to the scene as a light.  That means the sky's appearance and its lighting / reflections will match, even if your sky is coming from atmosphere, or layered clouds on top of a skybox, or distant mountains.  You can also manually specify a cubemap to use.
 
-The **Sky Light will** only recapture the scene when you rebuild lighting or manually tell it to using **Build** -> **Update Reflection Captures** or the **Recapture Scene** button on the _Sky Light Actor_.  **If you change the texture that the skybox is using, it will not automatically know to update.**
+The Sky Light will only recapture the scene when you rebuild lighting or manually tell it to using **Build** -> **Update Reflection Captures** or the **Recapture Scene** button on the _Sky Light Actor_.  **If you change the texture that the skybox is using, it will not automatically know to update.**
 
-A **Sky Light** should be used instead of the **Ambient Cubemap** to represent the sky's light, because **Sky Lights** support local shadowing, which prevents indoor areas from getting lit by the sky.
+A Sky Light should be used instead of the **Ambient Cubemap** to represent the sky's light, because Sky Lights support local shadowing, which prevents indoor areas from getting lit by the sky.
 [/EXCERPT:SL001]
 
-The **Sky Light** can be set to either **Static** or **Stationary** for its **Mobility**:
+The Sky Light can be set to either **Static** or **Stationary** for its **Mobility**:
 
 [OBJECT:topiclistnotitle]
 	[PARAM:icon]
@@ -74,8 +75,18 @@ The **Sky Light** can be set to either **Static** or **Stationary** for its **Mo
 ![](spacer.png)
 [/REGION]
 
-Pictured below are example images from a scene which utilizes a **Sky Light**.
+Pictured below are example images from a scene which utilizes a Sky Light.
 
+[OBJECT:ComparisonSlider]
+	[PARAM:before]
+	![Sky Light](Skylight.png)(w:800)
+	[/PARAM]
+	[PARAM:after]
+	![No Sky Light](NoSkylight.png)(w:800)
+	[/PARAM]
+[/OBJECT]
+
+<!--
 [REGION:lightbox]
 
 [![](skyLightEx_01.png)(w:884)](skyLightEx_01.png) [![](skyLightEx_02.png)(w:884)](skyLightEx_02.png)
@@ -83,31 +94,52 @@ Pictured below are example images from a scene which utilizes a **Sky Light**.
 [/REGION]
 
 _This scene was lit with only a Stationary Sky light, by specifying different HDR skybox cubemaps._
+-->
 
 ## Static Sky Light
 
-A **Sky Light** with **Static Mobility** will be baked completely into the lightmap for static objects and therefore costs nothing.  This is the only kind of sky light that is supported on mobile.  Edits to the light's properties will not be visible until lighting is built again.
+A Sky Light with **Static Mobility** will be baked completely into the lightmap for static objects and therefore costs nothing.  This is the only kind of sky light that is supported on mobile.  Edits to the light's properties will not be visible until lighting is built again.
 
-Note: Only components and lights with **Static** or **Stationary Mobility** will be captured and used for lighting with a **Static Sky Light**.  Furthermore, only the emissive contribution of materials can be captured with a **Static Sky Light**, in order to avoid a feedback loop.  Make sure your skybox has an unlit material.
+Note: Only components and lights with **Static** or **Stationary Mobility** will be captured and used for lighting with a Static Sky Light.  Furthermore, only the emissive contribution of materials can be captured with a **Static Sky Light**, in order to avoid a feedback loop.  Make sure your skybox has an unlit material.
 
 ## Stationary Skylight
 
-A **Sky Light** with **Stationary Mobility** gets baked shadowing from **Lightmass**.  Once you place a **Stationary Sky Light** in your level, you have to rebuild lighting once to get the baked shadowing.  You can then change the sky light as much as you want without rebuilding.  
+A Sky Light with **Stationary Mobility** gets baked shadowing from **Lightmass**.  Once you place a Stationary Sky Light in your level, you have to rebuild lighting once to get the baked shadowing.  You can then change the sky light as much as you want without rebuilding.  
 
-The **Sky Light** shadowing that **Lightmass** bakes stores directional occlusion information in what is called a **Bent Normal**.  This is the direction from the texel to the most un-occluded direction.  Areas that are mostly occluded use this direction for sky lighting instead of the surface normal, which improves quality quite a bit in cracks. 
+The Sky Light shadowing that **Lightmass** bakes stores directional occlusion information in what is called a **Bent Normal**.  This is the direction from the texel to the most un-occluded direction.  Areas that are mostly occluded use this direction for sky lighting instead of the surface normal, which improves quality quite a bit in cracks. 
 
+[OBJECT:ComparisonSlider]
+	[PARAM:before]
+	![Sky Light with AO only](skyLightEx_03.png)(w:500)
+	[/PARAM]
+	[PARAM:after]
+	![Sky Light with bent normal occlusion](skyLightEx_04.png)(w:500)
+	[/PARAM]
+[/OBJECT]
+
+<!--
 [REGION:lightbox]
 
 [![](skyLightEx_03.png)(w:462)](skyLightEx_03.png) [![](skyLightEx_04.png)(w:462)](skyLightEx_04.png)
 
 [/REGION]
-
+-->
 _The first image shows Sky Light with AO only. The second image shows Sky Light with bent normal occlusion.  Notice how surfaces in cracks 'agree' on where the lighting is coming from._
 
-Note: Only components and lights with **Static** or **Stationary Mobility** will be captured and used for lighting with a **Stationary Sky Light**.
+Note: Only components and lights with **Static** or **Stationary Mobility** will be captured and used for lighting with a Stationary Sky Light.
 
 Like all types of **Stationary Lights**, the color of the direct lighting can be changed at runtime through **Blueprint** or **Matinee**.  However the indirect lighting is baked into the lightmap and cannot be modified at runtime.  The amount of indirect lighting can be controlled by `IndirectLightingIntensity`.
 
+[OBJECT:ComparisonSlider]
+	[PARAM:before]
+	![Direct Lighting Only](0Original.jpg)(w:800)
+	[/PARAM]
+	[PARAM:after]
+	![Direct Lighting and Diffuse GI computed for a Stationary Sky Light](0SkyLightGI.jpg)(w:800)
+	[/PARAM]
+[/OBJECT]
+
+<!--
 [REGION:lightbox]
 
 [![](0Original.jpg)(w:884)](0Original.jpg) [![](0SkyLightGI.jpg)(w:884)](0SkyLightGI.jpg)
@@ -115,6 +147,7 @@ Like all types of **Stationary Lights**, the color of the direct lighting can be
 [/REGION]
 
 _Top: Direct lighting only, Bottom: Direct lighting and Diffuse GI computed by Lightmass for a Stationary Skylight._
+-->
 
 ## Movable Skylight
 
@@ -122,17 +155,17 @@ The **Movable Sky Light** does not use any form of precomputation.  It captures
 
 ### Distance Field Ambient Occlusion
 
-Shadowing for Movable Sky Lights is provided by a new feature called [Distance Field Ambient Occlusion](Engine/Rendering/LightingAndShadows/DistanceFieldAmbientOcclusion) which generates Ambient Occlusion from Signed Distance Field Volumes precomputed around each rigid object.  It supports dynamic scene changes - the rigid meshes can be moved or hidden and it will affect the occlusion.  This feature is not on by default and there are a few setup steps.
+Shadowing for Movable Sky Lights is provided by a new feature called [**Distance Field Ambient Occlusion**](Engine/Rendering/LightingAndShadows/DistanceFieldAmbientOcclusion) which generates Ambient Occlusion from Signed Distance Field Volumes precomputed around each rigid object.  It supports dynamic scene changes - the rigid meshes can be moved or hidden and it will affect the occlusion.  This feature is not on by default and there are a few setup steps.
 
 ## Limitations
 
-These features would be useful with **SkyLights** but are not yet implemented:
+These features would be useful with SkyLights but are not yet implemented:
 * Soft shadowing from movable objects (characters) onto the environment.
 * Dynamic recapturing of the scene for **Stationary** and **Movable** versions.  There is a blueprint function on the SkyLightComponent called RecaptureSky() which can be used, however that will cause a hitch.
 
 ## Sky Light Properties
 
-The properties of a **Sky Light** are broken up into 2 categories: **Light** and **Sky Light**. 
+The properties of a Sky Light are broken up into 2 categories: Light and Sky Light. 
 
 ### Light
 | Property | Description |
@@ -145,9 +178,9 @@ The properties of a **Sky Light** are broken up into 2 categories: **Light** and
 ### Sky Light
 | Property | Description |
 | -------- | ----------- |
-| **Source Type** | Whether to capture the distant scene and use it as the light source, or to use the specified cubemap.  When capturing the scene, anything further than `SkyDistanceThreshold` from the **Sky Light** position will be included. |
-| **Cubemap** | Specify the **Cubemap** to use for sky lighting if **Source Type** is set to `SLS_SpecifiedCubemap`.  |
-| **Sky Distance Threshold** | Distance from the **Sky Light** at which any geometry should be treated as part of the sky (also used by **Reflection Captures**). |
+| **Source Type** | Whether to capture the distant scene and use it as the light source, or to use the specified cubemap.  When capturing the scene, anything further than `SkyDistanceThreshold` from the Sky Light position will be included. |
+| **Cubemap** | Specify the Cubemap to use for sky lighting if Source Type is set to `SLS_SpecifiedCubemap`.  |
+| **Sky Distance Threshold** | Distance from the Sky Light at which any geometry should be treated as part of the sky (also used by Reflection Captures). |
 | **Lower Hemisphere is Black** | Whether all lighting from the lower hemisphere should be set to zero.  This is useful to prevent leaking from the lower hemisphere. |
-| **Recapture Scene** | When the **Sky Light Actor** is set to `SLS_CapturedScene`, this will recapture the image the Sky Light is using to light the scene. |
+| **Recapture Scene** | When the Sky Light Actor is set to `SLS_CapturedScene`, this will recapture the image the Sky Light is using to light the scene. |
 
