#UE4doc #docs - Added Engine Version Field. Added two images for use with Comparison Slider for the Spot Light cone.

diff --git a/Engine/Documentation/Source/Engine/Rendering/LightingAndShadows/LightTypes/Spot/Spot.INT.udn b/Engine/Documentation/Source/Engine/Rendering/LightingAndShadows/LightTypes/Spot/Spot.INT.udn
--- a/Engine/Documentation/Source/Engine/Rendering/LightingAndShadows/LightTypes/Spot/Spot.INT.udn
+++ b/Engine/Documentation/Source/Engine/Rendering/LightingAndShadows/LightTypes/Spot/Spot.INT.udn
@@ -2,6 +2,7 @@ Availability: Public
 Title:Spot Light
 Crumbs: %ROOT%, Engine, Engine/Rendering/LightingAndShadows, Engine/Rendering/LightingAndShadows/LightTypes
 Description:The basics of Spot Lights.
+Version: 4.9
 
 [VAR:Topic]
 [OBJECT:Topic]
@@ -87,7 +88,15 @@ Like other lights, a Spot Light can be set to one of three Mobility settings:
 
 Two examples of a **Spot Light** placed inside a level are shown below demonstrating how to determine the light radius and cone effectors:
 
-![](spotlight_example.png)
+[OBJECT:ComparisonSlider]
+	[PARAM:before]
+	![Spot Light without its Cone showing](SpotLightNoCone.png)(w:600)
+	[/PARAM]
+	[PARAM:after]
+	![Spot Light with Cone showing](SpotLightCone.png)(w:600)
+	[/PARAM]
+[/OBJECT]
+
 
 On the left a Spot Light without its cones showing light radius, on the right the same Spot Light with its cone effectors showing.
 
@@ -101,21 +110,21 @@ Property | Description |
 | -------- | ----------- |
 | **Intensity** | Total energy that the light emits. |
 | **Light Color** | The color that the light emits. |
-| **Inner Cone Angle** | Sets the inner cone angle of the **Spot Light**, in degrees. |
-| **Outer Cone Angle** | Sets the outer cone angle of the **Spot Light**, in degrees. |
+| **Inner Cone Angle** | Sets the inner cone angle of the Spot Light, in degrees. |
+| **Outer Cone Angle** | Sets the outer cone angle of the Spot Light, in degrees. |
 | **Attenuation Radius** | Bounds the light's visible influence.|
 | **Source Radius** | Radius of light source shape. |
 | **Source Length** | Length of light source shape. |
-| **Affects World** | Disables the light completely. Cannot be set at run time. To disable a light's effect during runtime, change its **Visibility** property. |
+| **Affects World** | Disables the light completely. Cannot be set at run time. To disable a light's effect during runtime, change its Visibility property. |
 | **Casts Shadows** | If the light casts shadows. |
 | **Indirect Lighting Intensity** | Scales the indirect lighting contribution from the light. |
-| **Use Inverse Squared Falloff** | Whether to use physically based inverse squared distance falloff, where **AttenuationRadius** is only clamping the light's contribution. |
-| **Light Falloff Exponent** | Controls the radial falloff of light when **UseInverseSquaredFalloff** is disabled. |
+| **Use Inverse Squared Falloff** | Whether to use physically based inverse squared distance falloff, where AttenuationRadius is only clamping the light's contribution. |
+| **Light Falloff Exponent** | Controls the radial falloff of light when UseInverseSquaredFalloff is disabled. |
 | **Min Roughness** | Min roughness effective for this light, used for softening specular highlights. |
 | **Shadow Bias** | Controls how accurate the shadows from this light are. |
 | **Shadow Filter Sharpen** | How much to sharpen shadow filtering for this light. |
 | **Cast Translucent Shadows** | Whether this light is allowed to cast dynamic shadows through translucent objects. |
-| **Affect Dynamic Indirect Lighting** | Whether the light should be injected into the **Light Propagation Volume**.|
+| **Affect Dynamic Indirect Lighting** | Whether the light should be injected into the Light Propagation Volume.|
 | **Cast Static Shadows** | Whether this light casts static shadows. |
 | **Cast Dynamic Shadows** | Whether this light casts dynamic shadows. |
 | **Affect Translucent Lighting** | Whether the light affects translucency or not. |
@@ -125,7 +134,7 @@ Property | Description |
 Property | Description |
 | -------- | ----------- |
 | **IES Texture** | The IES "Texture" used for the light profile. IES files are ASCII though Unreal represents them as textures, they are not image files. |
-| **Use IES Brightness** | If _false_, it will use the **Brightness** of the light to determine how much light to produce. If _true_, it will use the IES files brightness in Lumens (usually much larger than default values on lights in Unreal). |
+| **Use IES Brightness** | If _false_, it will use the Brightness of the light to determine how much light to produce. If _true_, it will use the IES files brightness in Lumens (usually much larger than default values on lights in Unreal). |
 | **IES Brightness Scale** | Scale for IES brightness contribution, as they can seriously blow out a scene. |
 
 ### Lightmass
@@ -140,5 +149,5 @@ Property | Description |
 | -------- | ----------- |
 | **Light Function Material** | The light function material to be applied to this light. |
 | **Light Function Scale** | Scales the light function projection. |
-| **Light Function Fade Distance** | The distance at which the light function should be completely faded to the value in **Disabled Brightness**. |
-| **Disabled Brightness** | Brightness factor applied to the light when the light function is specified but disabled, say from the property above: **Light Function Fade Distance**. |
+| **Light Function Fade Distance** | The distance at which the light function should be completely faded to the value in Disabled Brightness. |
+| **Disabled Brightness** | Brightness factor applied to the light when the light function is specified but disabled, say from the property above: Light Function Fade Distance. |
