#UE4doc #docs - Added Engine Version Field.

diff --git a/Engine/Documentation/Source/Engine/Rendering/LightingAndShadows/LitTranslucency/LitTranslucency.INT.udn b/Engine/Documentation/Source/Engine/Rendering/LightingAndShadows/LitTranslucency/LitTranslucency.INT.udn
--- a/Engine/Documentation/Source/Engine/Rendering/LightingAndShadows/LitTranslucency/LitTranslucency.INT.udn
+++ b/Engine/Documentation/Source/Engine/Rendering/LightingAndShadows/LitTranslucency/LitTranslucency.INT.udn
@@ -2,6 +2,7 @@ Availability:Public
 Title:Lit Translucency
 Crumbs: %ROOT%, Engine, Engine/Rendering/LightingAndShadows
 Description:Explanation of how translucent surfaces are lit and cast shadows including self-shadows.
+Version: 4.9
 
 [VAR:Topic]
 [OBJECT:Topic]
