#UE4doc #docs - Added Engine Version Field.

diff --git a/Engine/Documentation/Source/Engine/Rendering/LightingAndShadows/QuickStart/LightingQuickStart.INT.udn b/Engine/Documentation/Source/Engine/Rendering/LightingAndShadows/QuickStart/LightingQuickStart.INT.udn
--- a/Engine/Documentation/Source/Engine/Rendering/LightingAndShadows/QuickStart/LightingQuickStart.INT.udn
+++ b/Engine/Documentation/Source/Engine/Rendering/LightingAndShadows/QuickStart/LightingQuickStart.INT.udn
@@ -2,6 +2,8 @@ Availability: Public
 Title:Lighting Quick Start Guide
 Crumbs:%ROOT%, GettingStarted
 Description:In this guide you will create and light a small apartment using different types of lighting technqiues. 
+SkillLevel: Beginner
+Version: 4.9
 
 [VAR:Topic]
 [OBJECT:Topic]
