#UE4doc #docs - Added Skill Level Field. Added Engine Version Field.

diff --git a/Engine/Documentation/Source/Engine/Rendering/LightingAndShadows/QuickStart/1/LQS_1.INT.udn b/Engine/Documentation/Source/Engine/Rendering/LightingAndShadows/QuickStart/1/LQS_1.INT.udn
--- a/Engine/Documentation/Source/Engine/Rendering/LightingAndShadows/QuickStart/1/LQS_1.INT.udn
+++ b/Engine/Documentation/Source/Engine/Rendering/LightingAndShadows/QuickStart/1/LQS_1.INT.udn
@@ -2,6 +2,8 @@ Availability: Public
 Title:1. Required Setup
 Crumbs:%ROOT%, GettingStarted, Engine\Rendering\LightingAndShadows\QuickStart
 Description:Here we will create a few small rooms which we will light later in this guide. 
+SkillLevel: Beginner
+Version: 4.9
 
 
 [VAR:Steps]
@@ -29,11 +31,11 @@ Inside your **Blank** project (with **Starter Content**), follow the steps below
 
 1.	With your project open press **Ctrl+N** to create a new level, then select **Empty Level**. 
 
-1.	Inside the **Content Browser** under the **StarterContent/Architecture** folder, drag a **Floor_400x400** into the level. 
+1.	Inside the **Content Browser** under the **StarterContent > Architecture** folder, drag a **Floor_400x400** into the level. 
 
 	![](setup1.png)
 
-1.	Inside the **Outliner**, select the floor and press **Ctrl+W** 6 times to create the floors shown below. 
+1.	Inside the **World Outliner**, select the floor and press **Ctrl+W** 6 times to create the floors shown below. 
 
 	![](floors.png)
 
@@ -64,7 +66,7 @@ Inside your **Blank** project (with **Starter Content**), follow the steps below
 
 	![](setup11.png)
 
-1.	Inside the **Outliner**, select the wall and press **Ctrl+W** 11 times to create the walls shown below. 
+1.	Inside the **World Outliner**, select the wall and press **Ctrl+W** 11 times to create the walls shown below. 
 
 	![](wallsOutliner.png)
 
@@ -96,7 +98,7 @@ Inside your **Blank** project (with **Starter Content**), follow the steps below
 
 	![](setup25.png)
 
-1.	Inside the **Outliner**, select the pillar and press **Ctrl+W** 5 times to create the pillars shown below. 
+1.	Inside the **World Outliner**, select the pillar and press **Ctrl+W** 5 times to create the pillars shown below. 
 
 	![](pillarsAdded.png)
 
@@ -118,7 +120,7 @@ Inside your **Blank** project (with **Starter Content**), follow the steps below
 
 	![](setup32.png)
 	
-1.	In the **Outliner**, select the **Wall_400x407** wall piece. 
+1.	In the **World Outliner**, select the **Wall_400x407** wall piece. 
 
 	![](setup33.png)
 
