#UE4doc #docs - Added Skill Level field. Added Engine version field. Updated some images to match latest UI.

diff --git a/Engine/Documentation/Source/Engine/Rendering/LightingAndShadows/QuickStart/2/LQS_2.INT.udn b/Engine/Documentation/Source/Engine/Rendering/LightingAndShadows/QuickStart/2/LQS_2.INT.udn
--- a/Engine/Documentation/Source/Engine/Rendering/LightingAndShadows/QuickStart/2/LQS_2.INT.udn
+++ b/Engine/Documentation/Source/Engine/Rendering/LightingAndShadows/QuickStart/2/LQS_2.INT.udn
@@ -2,6 +2,8 @@ Availability: Public
 Title:2. Atmospheric and Directional Lighting
 Crumbs:%ROOT%, GettingStarted, Engine\Rendering\LightingAndShadows\QuickStart
 Description:In this step we will create the basic level lighting using Atmospheric Fog and a Directional Light. 
+Version: 4.9
+SkillLevel: Beginner
 
 
 [VAR:Steps]
