#UE4doc #docs - Added Skill Level Field. Added Engine Version Fields. Updated Image to match latest UI.

diff --git a/Engine/Documentation/Source/Engine/Rendering/LightingAndShadows/QuickStart/3/LQS_3.INT.udn b/Engine/Documentation/Source/Engine/Rendering/LightingAndShadows/QuickStart/3/LQS_3.INT.udn
--- a/Engine/Documentation/Source/Engine/Rendering/LightingAndShadows/QuickStart/3/LQS_3.INT.udn
+++ b/Engine/Documentation/Source/Engine/Rendering/LightingAndShadows/QuickStart/3/LQS_3.INT.udn
@@ -2,6 +2,8 @@ Availability: Public
 Title:3. Adding a Point Light
 Crumbs:%ROOT%, GettingStarted, Engine\Rendering\LightingAndShadows\QuickStart
 Description:In this step we will light up the small bathroom with a Point Light. 
+SkillLevel: Beginner
+Version: 4.9
 
 
 [VAR:Steps]
