#DocFix

Fixing grammar errors in documentation that have been pointed out in the forums

diff --git a/Engine/Documentation/Source/Engine/Rendering/Materials/LayeredMaterials/LayeredMaterials.INT.udn b/Engine/Documentation/Source/Engine/Rendering/Materials/LayeredMaterials/LayeredMaterials.INT.udn
--- a/Engine/Documentation/Source/Engine/Rendering/Materials/LayeredMaterials/LayeredMaterials.INT.udn
+++ b/Engine/Documentation/Source/Engine/Rendering/Materials/LayeredMaterials/LayeredMaterials.INT.udn
@@ -13,7 +13,7 @@ Description:An introductory document on the Layered Materials technique in UE4.
 
 Layered Materials can be thought of as "Materials within Materials." They provide a means to create a single Material that has a series of sub-Materials (or _layers_) that can be placed across the surface of your objects using per-pixel operations, such as masks. They are perfect for handling complex blending between unique surface types. In the image of the rockets above, the rocket on the far right is using separate Material Layers, including chrome, aluminum, and copper, and is blending between each material on a per-pixel basis. This effect is easily achieved with Layered Materials.
 
-As a feature, Layered Materials exist as an extension Material Function. Functions, as you may recall, are self-contained node networks that perform a specific operation, such as a complex math equation. These functions can then be reused at will in any number of Materials. Because of the _Make Material Attributes_ and _Break Material Attributes_ nodes, functions can also be used to define a Material, entirely _within_ a Function. Then, by bringing those functions into a final Material, you can have a Layered Material.
+As a feature, Layered Materials exist as an extension of Material Functions. Functions, as you may recall, are self-contained node networks that perform a specific operation, such as a complex math equation. These functions can then be reused at will in any number of Materials. Because of the _Make Material Attributes_ and _Break Material Attributes_ nodes, functions can also be used to define a Material, entirely _within_ a Function. Then, by bringing those functions into a final Material, you can have a Layered Material.
 
 ![](ChromeLayer.png)(w:600)
 
diff --git a/Engine/Documentation/Source/Engine/Rendering/PostProcessEffects/PostProcessEffects.INT.udn b/Engine/Documentation/Source/Engine/Rendering/PostProcessEffects/PostProcessEffects.INT.udn
--- a/Engine/Documentation/Source/Engine/Rendering/PostProcessEffects/PostProcessEffects.INT.udn
+++ b/Engine/Documentation/Source/Engine/Rendering/PostProcessEffects/PostProcessEffects.INT.udn
@@ -74,7 +74,7 @@ The properties found in the PostProcessVolume are explained below:
 
 ### Post Process Settings
 
-The post process settings are all of the properties of the various post process effects that can be controlled or overridden. These properties are declared in the `UScene` class. The **PostProcessVolume** contains a `FPostProcessSettings` struct as its `Settings` property that contains all of these properties and allows them to be overridden by the volume.
+The post process settings are all of the properties of the various post process effects that can be controlled or overridden. These properties are declared in the `UScene` class. The **PostProcessVolume** contains an `FPostProcessSettings` struct as its `Settings` property that contains all of these properties and allows them to be overridden by the volume.
 
 For descriptions of each of the available settings, see the page for the corresponding effect below.
 
