#UE4doc #docs:
Added engine version field
Replaced a couple images due to updated UI
Removed quotations from a term

diff --git a/Engine/Documentation/Source/Engine/Blueprints/Editor/UIComponents/MyBlueprint/MyBlueprint.INT.udn b/Engine/Documentation/Source/Engine/Blueprints/Editor/UIComponents/MyBlueprint/MyBlueprint.INT.udn
--- a/Engine/Documentation/Source/Engine/Blueprints/Editor/UIComponents/MyBlueprint/MyBlueprint.INT.udn
+++ b/Engine/Documentation/Source/Engine/Blueprints/Editor/UIComponents/MyBlueprint/MyBlueprint.INT.udn
@@ -2,6 +2,7 @@ Availability:Public
 Title: My Blueprint
 Crumbs:%ROOT%, Engine/Blueprints, Engine/Blueprints/Editor, Engine/Blueprints/Editor/UIBreakdowns/LevelBPUI
 Description:Hierarchical breakdown of the nodes within a Blueprint.
+version: 4.9
 
 [EXCERPT:Intro]
 The **My Blueprint** tab displays a tree list of the graphs, scripts, functions, macros, etc. contained within the Blueprint. This is essentially an outline of the Blueprint that lets you easily view existing elements of the Blueprint, as well as create new ones.
@@ -27,7 +28,7 @@ The **My Blueprint** tab has shortcut buttons (the ![](plus_Button.png) at the e
 | ![New Event Graph](Engine/Blueprints/Editor/UIComponents/MyBlueprint/myblueprint_eventgraph.png) | **New Event Graph** | Creates a new function, then immediately puts focus on the Name field of the **Details** tab so that it can be named. The new graph appears and is ready to have a node network defined within it. |
 | ![New Event Dispatcher](Engine/Blueprints/Editor/UIComponents/MyBlueprint/myblueprint_eventdispatcher.png) | **New Event Dispatcher** | Creates a new event dispatcher, then immediately puts focus on the Name field of the **Details** tab so that it can be named. | 
 
-These same buttons can also be accessed by **right-clicking** in the **My Blueprint** tab. The **right-click** menu also includes an option for creating a new enum asset.
+These same buttons can also be accessed by **Right-clicking** in the **My Blueprint** tab. The **Right-click** menu also includes an option for creating a new enum asset.
 
 ![](Engine/Blueprints/Editor/UIComponents/MyBlueprint/myblueprint_rightclick.png)
 
@@ -40,7 +41,7 @@ Finally, you can use the **Add New** button at the top of the **My Blueprint** P
 
 ## My Blueprint Sections
 
-The **My Blueprint** tab is divided up into 6 sections: The "**Add New**", Graphs, Functions, Macros, Variables, and Event Dispatchers sections.
+The **My Blueprint** tab is divided up into 6 sections: The **Add New**, Graphs, Functions, Macros, Variables, and Event Dispatchers sections.
 
 ![My Blueprint Pane](MyBlueprintPane.png)
 
