#UE4doc #docs:
Added engine version field
Changed font settings for a few terms
Replaced mentions of the Palette "tab" or "list" with "panel" for consistency

diff --git a/Engine/Documentation/Source/Engine/Blueprints/Editor/UIComponents/Palette/AllClassesPanel.INT.udn b/Engine/Documentation/Source/Engine/Blueprints/Editor/UIComponents/Palette/AllClassesPanel.INT.udn
--- a/Engine/Documentation/Source/Engine/Blueprints/Editor/UIComponents/Palette/AllClassesPanel.INT.udn
+++ b/Engine/Documentation/Source/Engine/Blueprints/Editor/UIComponents/Palette/AllClassesPanel.INT.udn
@@ -2,12 +2,13 @@ Availability:Public
 Title: Palette
 Crumbs:%ROOT%, Engine/Blueprints, Engine/Blueprints/Editor, Engine/Blueprints/Editor/UIBreakdowns/LevelBPUI
 Description:Non-context sensitive listing of all functions and variables for Palette that can be used in Blueprints.
+version: 4.9
 
 
 [TOC (start:2)]
 
 [EXCERPT:Intro]
-The **Palette** tab displays a list of all functions and variables that can be used in _Blueprints_. It is fully searchable, contains a **Favorites** panel that you can populate with the nodes you find yourself using the most, and you can filter it to search for nodes that are compatible with a specific class type. Nodes from the Palette can be added into your networks by simply dragging and dropping them from the list.
+The **Palette** panel displays a list of all functions and variables that can be used in Blueprints. It is fully searchable, contains a **Favorites** panel that you can populate with the nodes you find yourself using the most, and you can filter it to search for nodes that are compatible with a specific class type. Nodes from the **Palette** panel can be added into your networks by simply dragging and dropping them from the list.
 
 [/EXCERPT:Intro]
 
@@ -19,10 +20,10 @@ The **Palette** tab displays a list of all functions and variables that can be u
 ![](Palette.png)
 [/REGION]
 
-1. **Favorites list** - This shows your favorite and frequently used nodes. You can add nodes to this area by **right clicking** them within the Palette list and choosing **Add to Favorites**.
+1. **Favorites** - This shows your favorite and frequently used nodes. You can add nodes to this area by **Right clicking** them within the **Palette** panel and choosing **Add to Favorites**.
 1. **Palette list** - This is a list of all available functions, filterable by class and searchable.
-1. **Favorites search bar** - Actively filters the Favorites list while you type.
-1. **Class Filter** - This filters the Palette list to only show nodes which are compatible with a specific class. See [Class Filtering](#ClassFiltering) for more details.
+1. **Favorites search bar** - Actively filters the **Favorites** list while you type.
+1. **Class Filter** - This filters the **Palette** panel to only show nodes which are compatible with a specific class. See [Class Filtering](#ClassFiltering) for more details.
 1. **Palette search bar** - Actively filters the Palette list while you type.
 
 
@@ -31,7 +32,7 @@ The **Palette** tab displays a list of all functions and variables that can be u
 
 ## Class Filtering
 
-Sometimes it is useful to filter down the Palette to only show those nodes that are compatible with a given class type. This is done using the Class Filter dropdown. When clicked, it shows a list of all classes available in your project, similar to the [Class Viewer](Engine/UI/ClassViewer) panel. It also contains a search bar that you can use to actively filter down the list.
+Sometimes it is useful to filter down the **Palette** to only show those nodes that are compatible with a given class type. This is done using the **Class Filter** dropdown. When clicked, it shows a list of all classes available in your project, similar to the [Class Viewer](Engine/UI/ClassViewer) panel. It also contains a search bar that you can use to actively filter down the list.
 
 ![](ClassDropdown.png)
 
