#UE4 Docs: Applying feedback to the Android QS for PC and Mac.

diff --git a/Engine/Documentation/Source/Platforms/Android/Android.INT.udn b/Engine/Documentation/Source/Platforms/Android/Android.INT.udn
--- a/Engine/Documentation/Source/Platforms/Android/Android.INT.udn
+++ b/Engine/Documentation/Source/Platforms/Android/Android.INT.udn
@@ -19,6 +19,5 @@ Navigation:platform
 		* [](Platforms/Android/GettingStarted "%Platforms/Android/GettingStarted:description%")
 		* [](Platforms/Android/Reference "%Platforms/Android/Reference:description%")
 		* [](Platforms/Android/DeviceCompatibility "%Platforms/Android/DeviceCompatibility:description%")
-		* [](Platforms/Android/ReducingAPKSize "%Platforms/Android/ReducingAPKSize:description%")
 	[/PARAM]
 [/OBJECT]
diff --git a/Engine/Documentation/Source/Platforms/Android/GettingStarted/1/AndroidQS_1.INT.udn b/Engine/Documentation/Source/Platforms/Android/GettingStarted/1/AndroidQS_1.INT.udn
--- a/Engine/Documentation/Source/Platforms/Android/GettingStarted/1/AndroidQS_1.INT.udn
+++ b/Engine/Documentation/Source/Platforms/Android/GettingStarted/1/AndroidQS_1.INT.udn
@@ -23,6 +23,12 @@ Skilllevel:Intermediate
 
 %Steps%
 
+
+Before we can start developing for Android we first need to install Android Works. 
+Android Works will help to ensure that you have all the required programs and code that is needed for you to be able to deploy to Android devices.
+In the following section we will cover how you can get Android Works installed on your development PC.
+
+
 %Globals:OSSelect%
 
 [OBJECT:ToggleButtonContent]
@@ -55,7 +61,6 @@ Skilllevel:Intermediate
 [/OBJECT]
 
 
-
 <!--
 [EXCERPT:WindowsSetup]
 
@@ -63,7 +68,7 @@ Skilllevel:Intermediate
 If you already have Android Works installed, but your environment variables have been lost, you can skip all of this and simply point the editor to the install locations using the **Edit Menu** > **Project Settings** > **Android SDK** properties in the Editor.
 [/REGION]
 
-[region:tip]
+[region:warning]
 It is highly recommended that you install Android Works while the editor and launcher are both closed to ensure that there are not problems with the install.
 [/region]
 
@@ -79,11 +84,10 @@ It is highly recommended that you install Android Works while the editor and lau
 	Please note that the follow applies to both Windows and Mac based PCs that download the UE4 source code from Github.
 	If you are using the binary build downloaded from the launcher then you can skip this part and go right into installing Android Works. 
 
-	1. Once UE4 has been downloaded from Github and extracted, run the **Setup.bat** file to download the required files as well as the Android Works installer.
+	1. Once UE4 has been downloaded from Github and extracted, run the **Setup.bat** file to download the required files as well as the **Android Works** installer.
 
 	1. After  **Setup.bat** has finished running locate and run the Android Works installer following the instructions below.
 
-	1. When Android Works has finished installing, restart your PC so the environment variables are set and then run the **Setup.bat** file one more time to download the required Android Development files.
 	[/region]
 
 1. **Double-click** on the **AndroidWorks-1R1-windows.exe** file to initiate the install process.
@@ -133,7 +137,7 @@ It is highly recommended that you install Android Works while the editor and lau
 	
 	![](AndroidWorks_10.png)
 	[region:note]
-	If you have more programs to install you can press the **No** button but make sure to re-start your development PC before trying to use Android Work tools with UE4.
+	If you downloaded UE4 from Github, when Android Works has finished installing, restart your PC or Mac so the environment variables are set and then run the **Setup.bat** file one more time to download the required Android Development files.
 	[/region]
 [/EXCERPT:WindowsSetup]
 
@@ -143,7 +147,7 @@ It is highly recommended that you install Android Works while the editor and lau
 If you already have Android Works installed, but your environment variables have been lost, you can skip all of this and simply point the editor to the install locations using the **Edit Menu** > **Project Settings** > **Android SDK** properties in the Editor.
 [/REGION]
 
-[region:tip]
+[region:warning]
 It is highly recommended that you install Android Works while the editor and launcher are both closed to ensure that there are not problems with the install.
 [/region]
 Before you can install Android Works to your Mac you will need to make sure that both Java 6 and Java 7 are installed. 
@@ -153,10 +157,10 @@ Use the links provide below to download and install the correct versions of Java
 * **Java 7** can be downloaded from [here](http://www.oracle.com/technetwork/java/javase/downloads/jre7-downloads-1880261.html).
 
 1. To install Android Works navigate to the **Engine** > **Extras** > **Android** folder. 
-| Install Type  | Location  |
-| --- | --- |
-| Binary | [EngineInstallLocation]\Unreal Engine[EngineVersion]\Engine\Extras\Android |
-| GitHub | [EngineSyncLocation\UE4\Engine\Extras\Android |
+	| Install Type  | Location  |
+	| --- | --- |
+	| Binary | [EngineInstallLocation]\Unreal Engine[EngineVersion]\Engine\Extras\Android |
+	| GitHub | [EngineSyncLocation\UE4\Engine\Extras\Android |
 
 	[region:note]
 	If you downloaded UE4 from **Github** you will need to perform the following steps before you can run the Android Works installer. 
@@ -167,10 +171,8 @@ Use the links provide below to download and install the correct versions of Java
 
 	1. After  **Setup.bat** has finished running locate and run the Android Works installer following the instructions below.
 
-	1. When Android Works has finished installing, restart your Mac so the environment variables are set and then run the **Setup.bat** file one more time to download the required Android Development files.
 	[/region]
 
-
 1. **Double-click** on the **AndroidWorks-1R1-osx.dmg** to load the Android Works installer.
 	
 	![](AndroidWorks_Mac_00.png)
@@ -217,6 +219,9 @@ Use the links provide below to download and install the correct versions of Java
 
 	![](AndroidWorks_Mac_09.png)
 
+	[region:note]
+	If you downloaded UE4 from Github, when Android Works has finished installing, restart your PC or Mac so the environment variables are set and then run the **Setup.bat** file one more time to download the required Android Development files.
+	[/region]
 
 [/EXCERPT:MacOSSetup]
 -->
diff --git a/Engine/Documentation/Source/Platforms/Android/GettingStarted/2/AndroidQS_2.INT.udn b/Engine/Documentation/Source/Platforms/Android/GettingStarted/2/AndroidQS_2.INT.udn
--- a/Engine/Documentation/Source/Platforms/Android/GettingStarted/2/AndroidQS_2.INT.udn
+++ b/Engine/Documentation/Source/Platforms/Android/GettingStarted/2/AndroidQS_2.INT.udn
@@ -5,8 +5,7 @@ Description:Set up your Android device for development.
 Version:4.9
 Skilllevel:Intermediate
 
-
-
+Now that Android Works has been installed it is now time to setup your Android device so that it can be used in development. In the following section we will go over all you need to know to setup your Android device so that it can be used to run your UE4 project.
 [VAR:Steps]
 [OBJECT:Navigation]
 	[PARAM:previous]
@@ -60,65 +59,136 @@ Skilllevel:Intermediate
 <!--
 [EXCERPT:WindowsSetup]
 
-1. Connect your device.
+1. Connect your Android device to your development PC via USB.
+
 1.	Device driver software may automatically install, but if not, visit [Android's OEM USB Drivers](http://developer.android.com/tools/extras/oem-usb.html) page to find driver installation links and additional information.  
-1.	Open the **Settings** application on your device.
+
+1.	Open the **Settings** application on your Android device.
+
 1.	Look for the **Developer Options** menu. If it is not available, you need to enable Developer Options on your device. For more about developer options, visit [Android's Using Hardware Devices](http://developer.android.com/tools/device.html) page.
-	1.	Scroll down and select **About Phone**. Depending on your Android device, you may have an **About Device**, **About Tablet**, or **About Shield** option instead.
-		[REGION: note]
-		Note: This may be in a **More** section on newer Android versions.
-		[/REGION]
-	1.	Tap the **Build Number** 7 times. After the first few presses, you will see a message stating You are now # steps away from becoming a developer.
-	1.	You will get a success message once **Developer Mode** has been activated.
-	1.	Navigate back to the **Settings** menu and select **Developer Options** which is now available.
-1.	In the **Developer Options** menu, enable USB debugging. 
-1. Select **OK** when prompted.
-1.	To verify that your device is set up properly, open a command prompt by pressing **Windows+R** and typing `cmd.exe` in the **Run** window.
-1.	Type `adb devices` in the command window. You should see a device listed. Some troubleshooting points here are:
+
+1.	Scroll down and select **About Phone**. Depending on your Android device, you may have an **About Device**, **About Tablet**, or **About Shield** option instead.
+	[REGION: note]
+	Note: This may be in a **More** section on newer Android versions.
+	[/REGION]
+
+1.	Tap the **Build Number** **7** times to activate developer mode. 
+
+	![](Enable_Dev_Mode.png)
+
+1. After the first few presses, you will see a message stating You are now **# steps** away from becoming a developer.
+
+1.	Once developer mode has been activated you will get a success message displayed on the screen similar to the one below.
+
+	![](Dev_Mode_Active.png)
+
+1.	Navigate back to the **Settings** menu and select **Developer Options** which should now available.
+
+	![](Dev_Options_Enabled.png)
+
+1.	In the **Developer Options** menu, enable **USB debugging** by tapping on it with your finger. 
+
+	![](Enable_USB_Debugging.png)
+
+1. When prompted press the **OK** button to enable USB debugging.
+
+	![](Enable_USB_Debugging_Warning.png)
+
+1. You can tell USB Debugging has been enabled as there will be a green checkmark in the checkmark box next to USB debugging.
+
+	![](USB_Debugging_Enabled.png)
+
+1. Disconnect your Android device from USB and then immediately plug it back in. 
+	After the device is recognized by the PC the following message will appear on your Android device asking you if you want to allow this PC to talk to your device.
+	Press the checkmark box that says **Always allow this computer** and the press the **OK** button.
+	
+	![](PC_RSA_Fingerprint.png)
+
+1.	To verify that your device is set up properly, open a command prompt by pressing the **Windows + R** keys and typing **cmd** in the **Run** window.
+
+	![](Windows_Run_Command.png)
+
+1.	Type **adb devices** in the command prompt and you should see all devices that are connected to your PC.
+
+	![](ADB_Devices.png)
+
+	[region:note]
+
+	If you do not see your device after you input the adb devices command, here are some things you try to enable your PC to see your Android device. 
 	* If the device is listed but shows unauthorized, that indicates that you need to get the device to ask if you want to authorize the PC.
 	* You may need to unplug and reconnect your device to make it show up.
 	* It is recommended you check the Always allow checkbox.
-	* In some cases, a correctly configured device may not show up if it is connected as a **Media Device (MTP)** instead of as a **Camera (PTP)**. Select the **Camera (PTP)** option if your device is not visible at all under `adb devices` at this point.
+	* In some cases, a correctly configured device may not show up if it is connected as a **Media Device (MTP)** instead of as a **Camera (PTP)**. Select the **Camera (PTP)** option if your device is not visible at all under ***adb devices** at this point.
+
+	[/region]
+
 
 [/EXCERPT:WindowsSetup]
 
 [EXCERPT:MacOSSetup]
 
-1. Connect your device to your Mac.
+1. Connect your Android device to your Mac via USB.
 
-1. Open the **Settings application** on your device.
+1. Open the **Settings** application on your Android device.
 
 1. Look for the **Developer Options** menu. If it is not available, you need to enable Developer Options on your device. 
-	For more about developer options, visit **[Android's Using Hardware Devices](http://developer.android.com/tools/device.html)** page.
-	To enable developer options on your device you will need to do the following.
+	For more about developer options, visit [Android's Using Hardware Devices](http://developer.android.com/tools/device.html) page.
 
-	* Scroll down and select About Phone. 
-	* Depending on your Android device, you may have an **About Device**, **About Tablet**, or **About Shield** option instead.
+1. Scroll down and select **About Phone**. Depending on your Android device, you may have an **About Device**, **About Tablet**, or **About Shield** option instead.
 
-	[region:note]
-	In newer versions of Android this may be in the More section.
-	[/region]
+	[REGION: note]
+	Note: This may be in a **More** section on newer Android versions.
+	[/REGION]
+
+1. Tap the **Build Number** **7** times to activate developer mode.
+
+	![](Enable_Dev_Mode.png)
 
-	* Tap the Build Number **7** times. 
-	* After the first few presses, you will see a message that says, **You are now # steps away from becoming a developer**. 
-	* Once developer mode has been activated you will get a success message displayed on the screen.
-	* Navigate back to the **Settings menu** and select **Developer Options** which is now available.
+1. After the first few presses, you will see a message that says, **You are now # steps away from becoming a developer**. 
+
+1. Once developer mode has been activated you will get a success message displayed on the screen similar to the one below.
+
+	![](Dev_Mode_Active.png)
+
+1. Navigate back to the **Settings menu** and select **Developer Options** which is now available.
+
+	![](Dev_Options_Enabled.png)
 
 1. In the **Developer Options** menu, enable **USB debugging**. 
 
-1. Select **OK** when prompted.
+	![](Enable_USB_Debugging.png)
+
+1. When prompted press the **OK** button to enable USB debugging.
+
+	![](Enable_USB_Debugging_Warning.png)
+
+1. You can tell USB Debugging has been enabled as there will be a green checkmark in the checkmark box next to USB debugging.
+
+	![](USB_Debugging_Enabled.png)
+
+1. Disconnect your Android device from USB and then immediately plug it back in. 
+	After the device is recognized by the PC the following message will appear on your Android device asking you if you want to allow this Mac to talk to your device.
+	Press the checkmark box that says **Always allow this computer** and the press the **OK** button.
+	
+	![](PC_RSA_Fingerprint.png)
 
 1. To verify that your device is set up properly, open the **Terminal app** from /Applications/Utilties.
 
-1. Type **adb devices** in the terminal window and you should see a device listed. If you do not see your device listed here are some troubleshooting points you can try.
+	![](MAC_Terminal_Activate.png)
 
-	* If the device is listed but shows unauthorized, that indicates that you need to get the device to ask if you want to authorize the Mac.
+1.	Type **ADB devices** in the command prompt and you should see all devices that are connected to your Mac.
 
-	* You may need to unplug and reconnect your device to make it show up.
+	![](MAC_Checking_ADB_Devices.png)
+
+[region:note]
 
-	* It is recommended you check the Always allow checkbox to always allow connections from the computer that it is pluged into.
+If you do not see your device after you input the adb devices command, here are some things you try to enable your PC to see your Android device. 
+* If the device is listed but shows unauthorized, that indicates that you need to get the device to ask if you want to authorize the PC.
+* You may need to unplug and reconnect your device to make it show up.
+* It is recommended you check the Always allow checkbox.
+* In some cases, a correctly configured device may not show up if it is connected as a **Media Device (MTP)** instead of as a **Camera (PTP)**. Select the **Camera (PTP)** option if your device is not visible at all under ***adb devices** at this point.
 
-	* In some cases, a correctly configured device may not show up if it is connected as a Media Device (MTP) instead of as a Camera (PTP). Select the Camera (PTP) option if your device is not visible at all under adb devices at this point.
+[/region]
 
 [/EXCERPT:MacOSSetup]
 
diff --git a/Engine/Documentation/Source/Platforms/Android/GettingStarted/3/AndroidQS_3.INT.udn b/Engine/Documentation/Source/Platforms/Android/GettingStarted/3/AndroidQS_3.INT.udn
--- a/Engine/Documentation/Source/Platforms/Android/GettingStarted/3/AndroidQS_3.INT.udn
+++ b/Engine/Documentation/Source/Platforms/Android/GettingStarted/3/AndroidQS_3.INT.udn
@@ -21,23 +21,29 @@ Skilllevel:Intermediate
 [/OBJECT]
 [/VAR]
 
-
+In the following section we will be creating a new project for use to test deployment to Android with.
+The project that we create will be a Blueprint based project that will use the Blueprint Third Person template.
 
 
 %Steps%
 
 
-1.	Create a new project. Here we are testing a Blueprint-based project. If you start with a C++ project and then compile it, the following steps will be the same.  
+1. From the **Unreal Projects Browser** create a new Blueprint based project with the following options. 
+* Use the Third Person Template
+* Set target hardware to Mobile / Tablet
+* Set the Graphics level to Scalable 3D or 2D
+* Do not include the Starter Content
+* Give the project a name of AndroidProject
 
 ![](new_android_project.png)
 
 
 [REGION:tip]
-Note that we are not including the starter content in this project.
+While we have selected a Blueprint based project for this example, any type of project, C++ or Blueprint, can be used to packaged and deployed to an Android device.
 [/REGION]
 
 
 
+%Steps%
 
 
-%Steps%
diff --git a/Engine/Documentation/Source/Platforms/Android/GettingStarted/4/AndroidQS_4.INT.udn b/Engine/Documentation/Source/Platforms/Android/GettingStarted/4/AndroidQS_4.INT.udn
--- a/Engine/Documentation/Source/Platforms/Android/GettingStarted/4/AndroidQS_4.INT.udn
+++ b/Engine/Documentation/Source/Platforms/Android/GettingStarted/4/AndroidQS_4.INT.udn
@@ -22,32 +22,62 @@ Skilllevel:Intermediate
 [/VAR]
 
 
-
+In the following section we will be taking a look at how you can deploy a single level from the project you are working in to your Android device for quick testing. 
 
 %Steps%
 
+[region:note]
+For this section to work you will need to make sure that your Android device is plugged into your development PC via USB. If you have not done this already please do this before going any further.
+[/region]
 
-1.	Click on the dropdown arrow next to **Launch** in the **Level Editor** toolbar. You will notice that there are multiple entries for your Android device. 
-These are because of different texture formats. Here is a list of the formats as well as the devices that use them:
+Before you can deploy your level to your Android device you need to first determining which type of Texture format your Android Device supports.
+This is because different types of Android devices require different types of Texture formats to work correctly. 
+The chart below shows you which Texture formats work with which Android based devices.
 
 [INCLUDE:Platforms/Android/Reference#androidformats]
 
-1.	First, you can test launching your game, which will allow you to test out your current level quickly.
-1.	Select the **Launch** option with the correct format for your device.  
+[REGION:tip]
+	If you are unsure which Texture format to use or which ones are supported by your device, you can find out by running the game and then doing the following.
+
+	1. Tap the screen with 4 fingers at once, a dialog will pop up.
+	1. This dialog is usually used to enter a console command (like `stat fps`), but it will also display the formats that your device supports. 
+	1. Once you have that list, you can use that to choose a more ideal type when deploying/packaging.
+
+	![](Available_Texture_Formats.png)
+
+[/REGION:tip]
+
+1.	To test your current level on a Android based device you first need to make sure you have the level you want to test open. 
+	From this example we will be using the **ThirdPersonExampleMap** level from the Blueprint based project that was created in the last step.
+	
+	![](Correct_Level_Open.png)
+
+1.	With the **ThirdPersonExampleMap** now open go to the **Main Toolbar** and next to the **Launch** icon click on the small triangle that displays more options.
+
+	![](Level_Deploy_Options.png)
+
+1. From the **Launch** menu under the **Devices** section select your Android device from the list by clicking on it.
+
 	![](launch_rightformat.png)
-	[REGION:tip]
-	If you do not know your format, choose either the unlabeled or ETC1 format.  
-	![](launch_noformat.png)
 
-	When you run the game, if you tap the screen with 4 fingers at once, a dialog will pop up. This dialog is usually used to enter a console command (like `stat fps`), but it will also display the formats that your device supports. Once you have that list, you can use that to choose a more ideal type when deploying/packaging.
-	[/REGION:tip]
-1. A message will appear in the bottom right corner during the launching process.
-	![](launch_toast.png)
-1.	After the cooking and launching process, the game will appear automatically on your device, and also be accessible from your Apps drawer. 
-1.	You can uninstall the game now, or the batch file you will use to install your game package can manage uninstalling for you later.
+1. While your level is being launched on your device the progress will be displayed in the bottom right-hand corner of the screen. 
+ When the project has been successfully deployed to your device you will see a message telling you so.
 
+	[OBJECT:EmbeddedVideo]
+		[PARAMLITERAL:width]
+		640
+		[/PARAMLITERAL]
+		[PARAMLITERAL:height]
+		360
+		[/PARAMLITERAL]
+		[PARAMLITERAL:videoid]
+		afagQeEvf4Q
+		[/PARAMLITERAL]
+	[/OBJECT]
 
+1.	After the cooking and launching process has been completed, the game will appear automatically on your device, and also be accessible from your Apps drawer. 
 
+	![](Deployed_Project_Android.png)
 
 
 
diff --git a/Engine/Documentation/Source/Platforms/Android/GettingStarted/5/AndroidQS_5.INT.udn b/Engine/Documentation/Source/Platforms/Android/GettingStarted/5/AndroidQS_5.INT.udn
--- a/Engine/Documentation/Source/Platforms/Android/GettingStarted/5/AndroidQS_5.INT.udn
+++ b/Engine/Documentation/Source/Platforms/Android/GettingStarted/5/AndroidQS_5.INT.udn
@@ -21,6 +21,9 @@ Skilllevel:Intermediate
 [/OBJECT]
 [/VAR]
 
+In order for you to be able to share your project with others or put it on a digital store to sell, the project needs to be packaged. 
+In the following section we will take a look at how you can package your finished project for deployment inside of UE4.
+
 %Steps%
 
 %Globals:OSSelect%
@@ -57,51 +60,90 @@ Skilllevel:Intermediate
 <!--
 [EXCERPT:WindowsSetup]
 
-1.	To push your full game to your Android device, you will need to package it. In the **File** menu, hover over **Package Project** to expand the full list of options. Here you will again see the list of texture formats, and must select the right one (or unlabeled). If you are unsure what to select, use **ETC1**.
+[region:note]
+We are currently packaging in **Development** because we are testing this game, but when your game is complete, you can package it in **Shipping** by hovering over the **Build Configurations** option in the **Package Project** menu.
+[/region]
+
+1.	In the **File** menu, go to **Package Project > Android** and select the Texture format you want to package. 
+	For this example, we will choose Android **ETC1** as it is supported by all devices.
 
 	[REGION:raw]  
 	![](package_menu.png)
 	[/REGION]
 
-1.	You will be prompted for a save location for your packaged game, which should default to your project folder.
+	[REGION:note]
+	Different devices support different Texture formats depending on the hardware they run on. See, [Android Texture Formats](https://docs.unrealengine.com/latest/INT/Platforms/Android/Reference/index.html) for a detailed breakdown of the various formats.
+	[/REGION]
+
+1.	You will then be prompted for a location to save your packaged game. For now create a new folder on the desk top and select that as the location to save your project to.
+
 	![](package_folder.png)
-1. A packaging message will appear in the bottom right corner of the editor while the game is being packaged.
+
+1. While your project is being packaged a packaging message will appear in the bottom right corner showing the progress of the packing.
+
 	![](package_toast.png)
+
 1.	Once packaging is successful, a Packaging Complete message will appear in the bottom right of the editor.
-1.	Navigate to the save location you just set. 
-1.	In this folder, there will be a folder named specifically for your texture type. In this example, the folder is named `Android_DXT`. Open this folder.
-1.	**Double-click** on the `Install_[ProjectName]_Development.bat` batch file.
-	![](batch_file.png)
-1.	Your game is now installed and you can open it from your Apps drawer.
+
+	![](Packageing_Completed.png)
+
+1.	Navigate to the folder that was created in step two and in this folder you should have a new folder named specifically for the texture type you selected when packaging. 
+	In this example, the folder is named **Android_ETC1**.
+
+	![](Android_ECT1_Folder.png)
+
+1.	Open up the Android_ETC1 folder and then  **Double-click** on the **Install_[ProjectName]_Development.bat** batch file to install the project on your Android device.
+
+	![](Install_Bat.png)
+	[region:note]
+	You must have your Android device plugged into your development PC via USB to deploy the project to your device. If you do not have your device plugged in running the .BAT file will do nothing.
+	[/region]
+
+1.	After the .BAT file has finished you should now see your project ready to be used on your Android device.
+
+	![](Game_On_Phone.png)
 
 [/EXCERPT:WindowsSetup]
 
 [EXCERPT:MacOSSetup]
-1. Now, to push your full game to your Android device, you will need to package it. In the **File** menu, hover over **Package Project** to expand the full list of options. Here you will again see the list of texture formats, and must select the right one (or unlabeled).
-
-	![](AQS_Mac_01.png)
 
+[region:note]
 We are currently packaging in **Development** because we are testing this game, but when your game is complete, you can package it in **Shipping** by hovering over the **Build Configurations** option in the **Package Project** menu.
+[/region]
+
 
-1. You will be prompted for a save location for your packaged game, which should default to your project folder.
+1. 	In the **File** menu, go to **Package Project > Android** and select the Texture format you want to package. 
+	For this example, we will choose Android **ETC1** as it is supported by all devices.
+
+	![](AQS_Mac_01.png)
+
+1. You will be prompted for a save location for your packaged game, which should default to your project folder. Create a new folder and call it **Android_Builds**.
 
 	![](AQS_Mac_02.png)
 
-1. A packaging message will appear in the bottom right corner of the editor while the game is being packaged.
+1. While your project is being packaged a packaging message will appear in the bottom right corner showing the progress of the packing.
+
+	![](package_toast.png)
+
+1.	Once packaging is successful, a Packaging Complete message will appear in the bottom right of the editor.
 
-	![](AQS_Mac_03.png)
+	![](Packageing_Completed.png)
 
-1. Once packaging is successful, a Packaging Complete message will appear in the bottom right of the editor.
+1.	Navigate to the folder that was created in step two and in this folder you should have a new folder named specifically for the texture type you selected when packaging. 
+	In this example, the folder is named **Android_ETC1**.
 
-1. Navigate to the save location you just set.
+	![](MAC_Install_Folder.png)
 
-1. In this folder, there will be a folder named specifically for your texture type. In this example, the folder is named Android_DXT. Open this folder.
+1.	Open up the Android_ETC1 folder and then  **Double-click** on the **Install_[ProjectName]_Development_armv7-es3.command** file to install the project on your Android device.
 
-1. **Double-click** on the Install_[ProjectName]_Development-arm*.command script file.
+	![](MAC_Install_BAT.png)
+	[region:note]
+	You must have your Android device plugged into your development PC via USB to deploy the project to your device. If you do not have your device plugged in running the .BAT file will do nothing.
+	[/region]
 
-	![](AQS_Mac_04.png)
+1. After the .command file has finished you should now see your project ready to be used on your Android device.
 
-1. Your game is now installed and you can open it from your Apps drawer.
+	![](Game_On_Phone.png)
 
 [/EXCERPT:MacOSSetup]
 -->
diff --git a/Engine/Documentation/Source/Platforms/Android/GettingStarted/6/AndroidQS_6.INT.udn b/Engine/Documentation/Source/Platforms/Android/GettingStarted/6/AndroidQS_6.INT.udn
--- a/Engine/Documentation/Source/Platforms/Android/GettingStarted/6/AndroidQS_6.INT.udn
+++ b/Engine/Documentation/Source/Platforms/Android/GettingStarted/6/AndroidQS_6.INT.udn
@@ -1,10 +1,12 @@
 Availability:Public
-Title:6. Additional Resources
+Title:6. On your Own
 Crumbs:%ROOT%, Platforms, Platforms/Android, Platforms/Android/GettingStarted
 Description:Read more about Android development with Unreal Engine.
 Version:4.9
 Skilllevel:Intermediate
 
+The following section will go over some things that you can do on your own to practice and refine your skills at deploying UE4 projects to Android devices.
+
 [VAR:Steps]
 [OBJECT:Navigation]
 	[PARAM:previous]
@@ -23,12 +25,26 @@ Skilllevel:Intermediate
 
 %Steps%
 
-For more information about device compatibility and performance:
+Now that we have gone over how to get your UE4 project packaged and ready to be tested on your Android device, 
+try and accomplish the following using what you have learned from this Quick Start guide. 
+
+* Create a new project using one of the default templates and then try and get that project running on your Android device.
+* Try packaging your project using different Android Texture formats and see which ones will work for your Android device.
+* Check the size of your packaged and then using the [Reducing APK Package size](Platforms/Android/ReducingAPKSize) document and try reducing the size of your packaged game.
+* Check out the [Gear VR Debugging](Platforms/GearVR/Debugging) document and then try and deploy your project to your Android device over a wireless network.
+
+
+The following selection of links can help you with everything from which Android devices will work with UE4 to things you should consider when creating content that is aimed at mobile devices. 
+If you have not gone through these links before it is highly recommended that you do as they contain a plethora of very valuable information.
+
+For more information about device compatibility and performance checkout the following links:
+
 * [](Platforms\Android\DeviceCompatibility)
 * [](Platforms\DeviceProfiles)
 * [](Platforms\Mobile\Performance)
 
-For more information about designing game content and levels for mobile devices:
+For more information about designing game content and levels for mobile devices checkout the following links:
+
 * [](Platforms\Mobile\Content)
 * [](Platforms\Mobile\Lighting)
 * [](Platforms\Mobile\Materials)
diff --git a/Engine/Documentation/Source/Platforms/Android/GettingStarted/AndroidQuickStart.INT.udn b/Engine/Documentation/Source/Platforms/Android/GettingStarted/AndroidQuickStart.INT.udn
--- a/Engine/Documentation/Source/Platforms/Android/GettingStarted/AndroidQuickStart.INT.udn
+++ b/Engine/Documentation/Source/Platforms/Android/GettingStarted/AndroidQuickStart.INT.udn
@@ -39,9 +39,19 @@ Skilllevel:Intermediate
 [/OBJECT]
 [/VAR]
 
-The Android Quick Start will walk you through setting up your computer and your Android device for deployment, and putting a template game on your device using both the launching and packaging features.
-This guide uses **Android Works** as the Android SDK, but if you would like to know how to set up another SDK for use with Unreal Engine, see the [](Platforms\Android\Reference) page.
+The Android Quick Start will walk you through setting up your computer and your Android device for deployment, as well as putting a template game on your device using both the launching and packaging features.
 
+## Goals
+
+The focus of the Android Quick start is to walk you through the basics of setting up your PC or Mac for Android Development with UE4.
+
+## Objectives
+
+* Install and setup the Android SDK.
+* Setting up your Android device for development.
+* Creating a project that is targeted towards Android mobile development.
+* How to deploy a level to your Android device
+* How to package your game for release.
 
 [OBJECT:TopicButtonList]
 	[PARAM:icon]
@@ -61,3 +71,9 @@ This guide uses **Android Works** as the Android SDK, but if you would like to k
 	* [](Platforms/Android/GettingStarted/6 "%Platforms/Android/GettingStarted/6:description%")
 	[/PARAM]
 [/OBJECT]
+
+[region:note]
+This guide uses **Android Works** as the Android SDK, but if you would like to know how to set up another SDK for use with Unreal Engine 4.0, please checkout the [](Platforms\Android\Reference) page.
+[/region]
+
+
