#UE4doc #docs:
Added description as the field was blank
Removed "On this page" section as there was only one entry

diff --git a/Engine/Documentation/Source/Engine/Blueprints/Editor/UIComponents/Toolbar/ToolbarOptions/ToolbarOptions.INT.udn b/Engine/Documentation/Source/Engine/Blueprints/Editor/UIComponents/Toolbar/ToolbarOptions/ToolbarOptions.INT.udn
--- a/Engine/Documentation/Source/Engine/Blueprints/Editor/UIComponents/Toolbar/ToolbarOptions/ToolbarOptions.INT.udn
+++ b/Engine/Documentation/Source/Engine/Blueprints/Editor/UIComponents/Toolbar/ToolbarOptions/ToolbarOptions.INT.udn
@@ -1,9 +1,8 @@
 Availability:Public
 Title: Toolbar Options
 Crumbs:%ROOT%, Engine/Blueprints, Engine/Blueprints/Editor, Engine/Blueprints/Editor/Modes, Engine/Blueprints/Editor/UIComponents/Toolbar
-Description:
+Description: A list of Toolbar buttons and their functions
 
-[TOC (start:2)]
 
 
 [INCLUDE:Engine/Blueprints/Editor/UIComponents/Toolbar#Intro]
