#UE4doc #docs:
Added engine version field
Removed Overview header and "On this page" section as it only had one entry after removing the Overview header
Added a description as the field was blank

diff --git a/Engine/Documentation/Source/Engine/Blueprints/Editor/UIComponents/Toolbar/KismetToolbar.INT.udn b/Engine/Documentation/Source/Engine/Blueprints/Editor/UIComponents/Toolbar/KismetToolbar.INT.udn
--- a/Engine/Documentation/Source/Engine/Blueprints/Editor/UIComponents/Toolbar/KismetToolbar.INT.udn
+++ b/Engine/Documentation/Source/Engine/Blueprints/Editor/UIComponents/Toolbar/KismetToolbar.INT.udn
@@ -1,14 +1,12 @@
 Availability:Public
 Title: Toolbar
 Crumbs:%ROOT%, Engine/Blueprints, Engine/Blueprints/Editor
-Description:
+Description: A bar at the top of the editor that contains buttons that control multiple functions.
+version: 4.9
 
-[TOC (start:2)]
-
-## Overview
 
 [EXCERPT:Intro]
-The **Toolbar** is displayed at the top-left of the Blueprint Editor by default. The Blueprint Editor Toolbar buttons provide easy access to common commands needed when editing Blueprints. The Toolbar provides different buttons depending on which mode is active and which _Blueprint_ type you are currently editing.
+The **Toolbar** is displayed at the top-left of the Blueprint Editor by default. The Blueprint Editor Toolbar buttons provide easy access to common commands needed when editing Blueprints. The Toolbar provides different buttons depending on which mode is active and which Blueprint type you are currently editing.
 [/EXCERPT:Intro]
 
 [REGION:fullwidth]
