#UE4doc #docs:
As the two pages were linking to each other but talking about the same subject, I set KismetToolbar to include the content from ToolbarOptions and set ToolbarOptions to NoPublish

diff --git a/Engine/Documentation/Source/Engine/Blueprints/Editor/UIComponents/Toolbar/KismetToolbar.INT.udn b/Engine/Documentation/Source/Engine/Blueprints/Editor/UIComponents/Toolbar/KismetToolbar.INT.udn
--- a/Engine/Documentation/Source/Engine/Blueprints/Editor/UIComponents/Toolbar/KismetToolbar.INT.udn
+++ b/Engine/Documentation/Source/Engine/Blueprints/Editor/UIComponents/Toolbar/KismetToolbar.INT.udn
@@ -17,8 +17,4 @@ The Toolbar contains two sections:
 * **Toolbar options** - Tools for working with your Blueprint.
 * **Mode buttons** - Buttons that you can use to switch which mode your Blueprint is in.
 
-## Toolbar Options
-
-The Blueprint Toolbar contains many different options.  Different options are available depending on the type of Blueprint and the mode the Blueprint is in.
-
-For more information about Toolbar options, see [](Engine/Blueprints/Editor/UIComponents/Toolbar/ToolbarOptions).
+[INCLUDE:Engine/Blueprints/Editor/UIComponents/Toolbar/ToolbarOptions#main]
diff --git a/Engine/Documentation/Source/Engine/Blueprints/Editor/UIComponents/Toolbar/ToolbarOptions/ToolbarOptions.INT.udn b/Engine/Documentation/Source/Engine/Blueprints/Editor/UIComponents/Toolbar/ToolbarOptions/ToolbarOptions.INT.udn
--- a/Engine/Documentation/Source/Engine/Blueprints/Editor/UIComponents/Toolbar/ToolbarOptions/ToolbarOptions.INT.udn
+++ b/Engine/Documentation/Source/Engine/Blueprints/Editor/UIComponents/Toolbar/ToolbarOptions/ToolbarOptions.INT.udn
@@ -1,4 +1,4 @@
-Availability:Public
+Availability:NoPublish
 Title: Toolbar Options
 Crumbs:%ROOT%, Engine/Blueprints, Engine/Blueprints/Editor, Engine/Blueprints/Editor/Modes, Engine/Blueprints/Editor/UIComponents/Toolbar
 Description: A list of Toolbar buttons and their functions
@@ -33,7 +33,6 @@ Description: A list of Toolbar buttons and their functions
 | ![Eject Button](Engine/Blueprints/Editor/UIComponents/Toolbar/bp_eject.png)							| [INCLUDE:Engine/Blueprints/Editor/UIComponents/Toolbar/Toolbaroptions#eject]				|
 | ![Step Button](Engine/Blueprints/Editor/UIComponents/Toolbar/bp_step.png)								| [INCLUDE:Engine/Blueprints/Editor/UIComponents/Toolbar/Toolbaroptions#step]					|
 | ![Debug Dropdown](Engine/Blueprints/Editor/UIComponents/Toolbar/bp_debug_dropdown.png)				| [INCLUDE:Engine/Blueprints/Editor/UIComponents/Toolbar/Toolbaroptions#debugobject]			| 
-
 <!--
 [EXCERPT:CompileSuccessful]
 Compilation was successful. Clicking the button compiles the Blueprint being edited. Output from the compiling process is displayed in the Blueprint Log of the Message Log. This button will be inactive during debugging.
