#UE4doc #docs:
Added engine version field
Added a description

diff --git a/Engine/Documentation/Source/Engine/Blueprints/Editor/UIComponents/Toolbar/ToolbarOptions/ToolbarOptions.INT.udn b/Engine/Documentation/Source/Engine/Blueprints/Editor/UIComponents/Toolbar/ToolbarOptions/ToolbarOptions.INT.udn
--- a/Engine/Documentation/Source/Engine/Blueprints/Editor/UIComponents/Toolbar/ToolbarOptions/ToolbarOptions.INT.udn
+++ b/Engine/Documentation/Source/Engine/Blueprints/Editor/UIComponents/Toolbar/ToolbarOptions/ToolbarOptions.INT.udn
@@ -2,7 +2,7 @@ Availability:NoPublish
 Title: Toolbar Options
 Crumbs:%ROOT%, Engine/Blueprints, Engine/Blueprints/Editor, Engine/Blueprints/Editor/Modes, Engine/Blueprints/Editor/UIComponents/Toolbar
 Description: A list of Toolbar buttons and their functions
-
+version: 4.9
 
 
 [INCLUDE:Engine/Blueprints/Editor/UIComponents/Toolbar#Intro]
