#UE4doc #docs:
Added engine version field
Changed font settings for a few terms
Removed "On this page" section as it only had one entry

diff --git a/Engine/Documentation/Source/Engine/Blueprints/Editor/UIComponents/Viewport/BlueprintViewport.INT.udn b/Engine/Documentation/Source/Engine/Blueprints/Editor/UIComponents/Viewport/BlueprintViewport.INT.udn
--- a/Engine/Documentation/Source/Engine/Blueprints/Editor/UIComponents/Viewport/BlueprintViewport.INT.udn
+++ b/Engine/Documentation/Source/Engine/Blueprints/Editor/UIComponents/Viewport/BlueprintViewport.INT.udn
@@ -2,18 +2,18 @@ Availability:Public
 Title: Blueprint Editor Viewport
 Crumbs:%ROOT%, Engine/Blueprints, Engine/Blueprints/Editor, Engine/Blueprints/Editor/UIBreakdowns/LevelBPUI
 Description:Overview of the viewport in the Components mode of Blueprint Classes
+version: 4.9
 
-[TOC (start:2)]
 
 In the **Viewport**, in the Blueprint Editor, you can view and manipulate your Blueprint's components. 
 
 ![](BP_Viewport.png)
 
 [REGION:note]
-You can also manipulate a selected component by adjusting its values manually in the **[Details](Engine/Blueprints/Editor/UIComponents/Details)** panel.
+You can also manipulate a selected component by adjusting its values manually in the [Details](Engine/Blueprints/Editor/UIComponents/Details) panel.
 [/REGION]
 
-Selecting a component in the **[Component](Engine/Blueprints/UserGuide/Components)** list selects that component in the **Viewport** too, and populates the **Details** panel with its property information. 
+Selecting a component in the [Component](Engine/Blueprints/UserGuide/Components) list selects that component in the **Viewport** too, and populates the **Details** panel with its property information. 
 
 You can use the transformation widgets in the **Viewport** to adjust the location, rotation, and scale of selected components. For more information about transformation widgets, see [Interactive Transforms](Engine/Actors/Transform#InteractiveTransformation).
 
