#UE4doc #docs:
Added engine version field
Removed Topic List field from a section as it didn't follow the rest of the page's consistency
Changed font settings for many terms

diff --git a/Engine/Documentation/Source/Engine/Blueprints/GettingStarted/BlueprintIntro.INT.udn b/Engine/Documentation/Source/Engine/Blueprints/GettingStarted/BlueprintIntro.INT.udn
--- a/Engine/Documentation/Source/Engine/Blueprints/GettingStarted/BlueprintIntro.INT.udn
+++ b/Engine/Documentation/Source/Engine/Blueprints/GettingStarted/BlueprintIntro.INT.udn
@@ -2,6 +2,7 @@ Availability:Public
 Title: Introduction to Blueprints
 Crumbs: %ROOT%, Engine, Engine/Blueprints
 Description:If you are just getting started with Blueprints, this provides a high-level overview of what they are and what they can do. 
+version: 4.9
 
 [VAR:Topic]
 [OBJECT:Topic]
@@ -32,18 +33,18 @@ and script level events; giving designers and gameplay programmers the tools to
 Unreal Editor without ever needing to write a line of code.
 [/EXCERPT:BlueprintIntro]
 
-For more detailed overview information about _Blueprints_, including types of _Blueprints_, _Blueprint_ anatomy, and a discussion of how _Blueprints_ relate to UE3's Kismet and UnrealScript, see [](Engine/Blueprints\Overview).
+For more detailed overview information about Blueprints, including types of Blueprints, Blueprint anatomy, and a discussion of how Blueprints relate to UE3's Kismet and UnrealScript, see [](Engine/Blueprints\Overview).
 
 
 ## How Do Blueprints Work?
-In their basic form, _Blueprints_ are visually scripted additions to your game. By connecting _Nodes_, _Events_, _Functions_, and _Variables_ with _Wires_, it is possible to create complex gameplay elements.
+In their basic form, Blueprints are visually scripted additions to your game. By connecting Nodes, Events, Functions, and Variables with Wires, it is possible to create complex gameplay elements.
 
-_Blueprints_ work by using graphs of _Nodes_ for various purposes - object construction, individual functions, and general gameplay events - that are specific to each instance of the _Blueprint_ in order to implement behavior and other functionality.
+Blueprints work by using graphs of Nodes for various purposes - object construction, individual functions, and general gameplay events - that are specific to each instance of the Blueprint in order to implement behavior and other functionality.
 
 
 ## Commonly Used Blueprint Types
 
-The most common _Blueprint_ types you will be working with are _Level Blueprints_ and _Blueprint Classes_.  
+The most common Blueprint types you will be working with are Level Blueprints and Blueprint Classes.  
 
 These are just two of the [Types of Blueprints](Engine/Blueprints/UserGuide/Types), which also include [Blueprint Macros](Engine/Blueprints/UserGuide/Types/MacroLibrary) and [Blueprint Interfaces](Engine/Blueprints/UserGuide/Types/Interface).
 
@@ -55,7 +56,7 @@ These are just two of the [Types of Blueprints](Engine/Blueprints/UserGuide/Type
 		[REGION:lightbox]
 		[![](Engine/Blueprints/GettingStarted/Level_Blueprint_Main.png)(w:640)](Engine/Blueprints/UserGuide/Types/LevelBlueprint) 
 		[/REGION]
-		The _Level Blueprint_ fills the same role that Kismet did in Unreal Engine 3, and has the same capabilities. Each level has its own _Level Blueprint_, and this can reference and manipulate Actors within the level, control cinematics using Matinee Actors, and manage things like level streaming, checkpoints, and other level-related systems. The _Level Blueprint_ can also interact with _Blueprint Classes_ (see the next section for examples of these) placed in the level, such as reading/setting any variables or triggering custom events they might contain.
+		The Level Blueprint fills the same role that Kismet did in Unreal Engine 3, and has the same capabilities. Each level has its own Level Blueprint, and this can reference and manipulate Actors within the level, control cinematics using Matinee Actors, and manage things like level streaming, checkpoints, and other level-related systems. The Level Blueprint can also interact with Blueprint Classes (see the next section for examples of these) placed in the level, such as reading/setting any variables or triggering custom events they might contain.
 	[/PARAM]
 [/OBJECT]
 
@@ -67,14 +68,14 @@ These are just two of the [Types of Blueprints](Engine/Blueprints/UserGuide/Type
 		[REGION:lightbox]
 		[![](Engine/Blueprints/GettingStarted/Blueprint_Main.png)(w:640)](Engine/Blueprints/UserGuide/Types/ClassBlueprint)
 		[/REGION]
-		_Blueprint Classes_ are ideal for making interactive assets such as doors, switches, collectible items, and destructible scenery.  In the image above, the button and the set of doors are each separate _Blueprints_ that contain the necessary script to respond to player overlap events, make them animate, play sound effects, and change their materials (the button lights up when pressed, for example).
+		Blueprint Classes are ideal for making interactive assets such as doors, switches, collectible items, and destructible scenery.  In the image above, the button and the set of doors are each separate Blueprints that contain the necessary script to respond to player overlap events, make them animate, play sound effects, and change their materials (the button lights up when pressed, for example).
 		
-		In this case, pressing the button activates an event inside the door _Blueprint_, causing it to open - but the doors could just as easily be activated by another type of _Blueprint_, or by a _Level Blueprint_ sequence. Because of the self-contained nature of _Blueprints_, they can be constructed in such a way that you can drop them into a level and they will simply work, with minimal setup required. This also means that editing a _Blueprint_ that is in use throughout a project will update every instance of it.
+		In this case, pressing the button activates an event inside the door Blueprint, causing it to open - but the doors could just as easily be activated by another type of Blueprint, or by a Level Blueprint sequence. Because of the self-contained nature of Blueprints, they can be constructed in such a way that you can drop them into a level and they will simply work, with minimal setup required. This also means that editing a Blueprint that is in use throughout a project will update every instance of it.
 	[/PARAM]
 [/OBJECT]
 
 ## What Else Can Blueprints Do?
-You have read about _Level Blueprints_ and _Blueprint Classes_, listed below are a handful of examples that can be accomplished with the _Blueprint_ system. 
+You have read about Level Blueprints and Blueprint Classes, listed below are a handful of examples that can be accomplished with the Blueprint system. 
 
 [OBJECT:TopicList]
 	[PARAM:title]
@@ -83,11 +84,11 @@ You have read about _Level Blueprints_ and _Blueprint Classes_, listed below are
 	[PARAM:description]
 		![](Construction_Script.png)
 
-		The _Construction Script_ is a type of graph within _Blueprint Classes_ that executes when that Actor is placed or updated in the editor, but not during gameplay. It is useful for creating easily customizable props that allow environment artists to work faster, such as a light fixture that automatically updates its material to match the color and brightness of its point light component, or a _Blueprint_ that randomly scatters foliage meshes over an area.
+		The Construction Script is a type of graph within Blueprint Classes that executes when that Actor is placed or updated in the editor, but not during gameplay. It is useful for creating easily customizable props that allow environment artists to work faster, such as a light fixture that automatically updates its material to match the color and brightness of its point light component, or a Blueprint that randomly scatters foliage meshes over an area.
 
-		In the Content Examples maps, the long rooms that contain each example (pictured above) are actually a single _Blueprint_ made up of many components. The _Blueprint's Construction Script_ creates and arranges the various Static Meshes and lights according to parameters exposed in the _Blueprint's_ **Details** panel. With each Content Example map we created, we were able to drop in the demo room _Blueprint_, set values for the length, height, and number of rooms that would be generated (and a few other options), and have a complete set of rooms ready in moments.
+		In the Content Examples maps, the long rooms that contain each example (pictured above) are actually a single Blueprint made up of many components. The Blueprint's Construction Script creates and arranges the various Static Meshes and lights according to parameters exposed in the Blueprint's **Details** panel. With each Content Example map we created, we were able to drop in the demo room Blueprint, set values for the length, height, and number of rooms that would be generated (and a few other options), and have a complete set of rooms ready in moments.
 
-		A _Blueprint_ like this can be time-consuming to create initially, but if you know you will use it often, the time saved when building a level and the ease of making changes can make it very worthwhile.
+		A Blueprint like this can be time-consuming to create initially, but if you know you will use it often, the time saved when building a level and the ease of making changes can make it very worthwhile.
 	[/PARAM]
 [/OBJECT]
 
@@ -98,9 +99,9 @@ You have read about _Level Blueprints_ and _Blueprint Classes_, listed below are
 	[PARAM:description]
 		![](Game_Characters.png)
 
-		_Pawns_ are also a type of _Blueprint Class_, and it is possible to put together every element you need for a playable character in the _Blueprint_ graph. You can manipulate camera behavior, set up input events for mouse, controller, and touch screens, and create an _Animation Blueprint_ asset for handling skeletal mesh animations.
+		Pawns are also a type of Blueprint Class, and it is possible to put together every element you need for a playable character in the Blueprint graph. You can manipulate camera behavior, set up input events for mouse, controller, and touch screens, and create an Animation Blueprint asset for handling skeletal mesh animations.
 
-		When you create a new _Character Blueprint_, it comes with a character component that has much of the behavior needed for moving around, jumping, swimming, and falling built-in, and all that is required is to add some input events in accordance with how you want your character to be controlled.
+		When you create a new Character Blueprint, it comes with a character component that has much of the behavior needed for moving around, jumping, swimming, and falling built-in, and all that is required is to add some input events in accordance with how you want your character to be controlled.
 	[/PARAM]
 [/OBJECT]
 
@@ -111,42 +112,33 @@ You have read about _Level Blueprints_ and _Blueprint Classes_, listed below are
 	[PARAM:description]
 		![](Create_HUDs.png)
 
-		_Blueprint_ script can be used to create a game's HUD as well, which is similar to _Blueprint Classes_ in that it can contain event sequences and variables, but is assigned to your project's _GameMode_ asset instead of being added directly to a level.
+		Blueprint script can be used to create a game's HUD as well, which is similar to Blueprint Classes in that it can contain event sequences and variables, but is assigned to your project's GameMode asset instead of being added directly to a level.
 
-		You can set up a HUD to read variables from other _Blueprints_ and use them to display a health bar, update a score value, display objective markers, and so on. It is also possible to use the HUD to add hit-boxes for elements like buttons that can be clicked on or, in the case of mobile games, can respond to touch input.
+		You can set up a HUD to read variables from other Blueprints and use them to display a health bar, update a score value, display objective markers, and so on. It is also possible to use the HUD to add hit-boxes for elements like buttons that can be clicked on or, in the case of mobile games, can respond to touch input.
 
 		All of the examples pictured above exist in the sample content available with UE4, so if you want to get a closer look, they can be found in the **Content Examples**, **First Person Shooter** and **Swing Ninja** projects.
 	[/PARAM]
 [/OBJECT]
 
-[OBJECT:TopicList]
-	[PARAM:icon]
-	![](%ROOT%/blueprint_icon.png)(convert:false)
-	[/PARAM]
-	[PARAM:title]
-		Blueprint Samples and Tutorials
-	[/PARAM]
-	[PARAM:description]
-		Below are additional resources you can use to learn more about the _Blueprint_ system. 
-	[/PARAM]
-	[PARAM:links]
-		* [](Resources/Showcases/BlueprintExamples)
-		* [Blueprint Lift Tutorial](http://wiki.unrealengine.com/Blueprint_Lift_Tutorial)
-		* [Blueprint Power-Up Tutorial](http://wiki.unrealengine.com/Blueprint_Power_Up_Tutorial)
-	[/PARAM]
-[/OBJECT]
+##Blueprint Samples and Tutorials
+
+Below are additional resources you can use to learn more about the Blueprint system. 
+
+* [](Resources/Showcases/BlueprintExamples)
+* [Blueprint Lift Tutorial](http://wiki.unrealengine.com/Blueprint_Lift_Tutorial)
+* [Blueprint Power-Up Tutorial](http://wiki.unrealengine.com/Blueprint_Power_Up_Tutorial)
 
 
 ## Blueprint Editors and Graphs
 
-Whether you are building a _Level Blueprint_ or a _Blueprint Class_, you will be using [Blueprint Elements](Engine/Blueprints\UserGuide) assembled in a [Blueprint Editor](Engine/Blueprints/Editor).  
+Whether you are building a Level Blueprint or a Blueprint Class, you will be using [Blueprint Elements](Engine/Blueprints\UserGuide) assembled in a [Blueprint Editor](Engine/Blueprints/Editor).  
 
 [REGION:lightbox]
 	[![](Engine/Blueprints/Editor/blueprint_editor_example.png)(w:480)](Engine/Blueprints/Editor/blueprint_editor_example.png)
 [/REGION]
 
 
-Different types of **Blueprint Editor** are available depending on the type of _Blueprint_ you are working with. The core feature of most **Blueprint Editors** is the **Graph** mode, with its central **Graph** tab 
-for laying out the network of your _Blueprint_.  
+Different types of **Blueprint Editor** are available depending on the type of Blueprint you are working with. The core feature of most **Blueprint Editors** is the **Graph** mode, with its central **Graph** tab 
+for laying out the network of your Blueprint.  
 
 
