Removing the commented out sections that were for the 4.9 release.
Added a blurb about reducing the shader cache size

diff --git a/Engine/Documentation/Source/Platforms/Android/ReducingAPKSize/ReducingAPKSizes.INT.udn b/Engine/Documentation/Source/Platforms/Android/ReducingAPKSize/ReducingAPKSizes.INT.udn
--- a/Engine/Documentation/Source/Platforms/Android/ReducingAPKSize/ReducingAPKSizes.INT.udn
+++ b/Engine/Documentation/Source/Platforms/Android/ReducingAPKSize/ReducingAPKSizes.INT.udn
@@ -8,7 +8,7 @@ Version: 4.8
 [TOC(start:2 end:3)]
 
 No matter what type of Android device your UE4 project is targeting reducing the size of the Android application package or **APK** to be under the 50MB limit for the Google Play store can be challenging. 
-In the following tutorial we will cover steps you can take to help reduce your projects final APK package size to be as small as possible using nothing but the tools provided to you in the UE4 Editor. 
+In the following guide we will cover steps you can take to help reduce your projects final APK package size to be as small as possible using nothing but the tools provided to you in the UE4 Editor. 
 
 ## ETC1 Texture Considerations
 When packing your project for Android using the ETC1 Texture format you need to be aware that ETC1 will not compress Textures that have embedded Alpha information in them and instead treat them as fully uncompressed Textures. 
@@ -26,10 +26,6 @@ This way you ensure that the only content in your project is content that you ne
 The easiest and quickest way to reduce your APK package size is to tell UE4 to compress the APK packages during the packing process. 
 To turn on package compression you need to do the following in the UE4 Editor.
 
-[region:note]
-Please note that the Create compressed cooked packages option is **enabled** by default and only available in Unreal Engine versions **4.8** and later.
-[/region]
-
 1. First open up the projects settings by going to the Main Toolbar and selecting the **Edit** option then selecting **Project Settings**.
 
 	![](T_Project_Settings_Menu.png)
@@ -132,7 +128,6 @@ To disable a plug-in inside of UE4 you need to do the following.
 	[/PARAMLITERAL]
 	[/OBJECT]
 
-<!--
 ## Package Black List
 
 With the release of Unreal Engine 4.9 and later you can now place a text file in your projects **Build/Platform/(Target Platform Folder)** directory that will tell the cooker to exclude partial or complete file paths from being packaged into your project.
@@ -149,9 +144,28 @@ Here is an example from Tappy Chicken where the cooker is told to not include th
 	![](T_Black_List.png)
 
 The first three entries should be safe for any project, but there may be cases were the **AssetRegistry.bin** is needed at runtime depending on your project.
-It is highly recommend that you thoroughly test to ensure that the content you are removing will not negatively impact your game.
-An easy and simple way to go about this is to open up all your projects levels to ensure that they correctly load.
+It is highly recommend that you thoroughly test to ensure that the content you are removing will not negatively impact your project.
+An easy and simple way to go about this is to open up all your projects levels to ensure that they correctly load and have no errors or warnings.
 If you run into further problems and are not sure where they are coming from check the Logcat for errors.
--->
+
+## Max Dynamic Point Lights
+
+With the release of UE 4.9 a new option has been introduced that can help reduce the number of shader that are generated for dynamic lights
+ making this feature is particularly useful for mobile games that do not require dynamic lighting. 
+The property, **Max Dynamic Point Lights** can be found under the **Mobile** section of the **Project Settings**.
+
+![](Max_Dyn_Point_Lights.png)
+
+To completely disable this feature set the **Max Dynamic Point Lights** from **4** to **0** and then re-start the UE4 Editor when prompted.
+
+![](Max_Dyn_Point_Lights_To_0.png)
+
+Depending on the size of your project and the amount of Materials used the size saving you can achieve can range from just a few MB into the multiple MB range.
+Using the default Third Person template for an example, we can see what the size savings are when Max Dynamic Point Lights are enabled verses when Max Dynamic Point Lights and disabled completely. 
+
+| State| Packaged Game Size | Savings|
+| ---------| ------| -----|
+|ON|54.3 MB|0|
+|OFF| 53.2 MB| 1.1 MB|
 
 
