#UE4doc #docs:
Added engine version field
Replaced a few images due to outdated UI
Changed font settings for multiple terms

diff --git a/Engine/Documentation/Source/Engine/Blueprints/HowTo/BPHT_1/BPHT_1.INT.udn b/Engine/Documentation/Source/Engine/Blueprints/HowTo/BPHT_1/BPHT_1.INT.udn
--- a/Engine/Documentation/Source/Engine/Blueprints/HowTo/BPHT_1/BPHT_1.INT.udn
+++ b/Engine/Documentation/Source/Engine/Blueprints/HowTo/BPHT_1/BPHT_1.INT.udn
@@ -2,6 +2,7 @@ Availability: Public
 Title:1. Required Asset Setup
 Crumbs: %ROOT%, Engine, Engine/Blueprints, Engine/Blueprints/HowTo
 Description:In this step we will set up the assets we will use in the project. 
+version: 4.9
 
 
 [VAR:Steps]
@@ -144,11 +145,11 @@ The final piece of setup requires us to create a basic HUD to communicate inform
 
 	![](MovedWidget.png)
 
-1.	Click on the left-most text displayed in the box, then in the **Details** panel under **Text** enter **Pulls:**. 
+1.	Click on the left-most text displayed in the box, then in the **Details** panel under **Text** enter "Pulls:". 
 
 	![](EnterText1.png)
 
-1.	For the other text, enter the **10** for its **Text** field. 
+1.	For the other text, enter the "10" for its **Text** field. 
 
 	![](EnterText2.png)
 
@@ -160,7 +161,7 @@ The final piece of setup requires us to create a basic HUD to communicate inform
 
 	![](AddedBoxes.png)
 
-1.	Update the other text fields to **Shots** (set to **3**) and **Score** (set to **0**) and click the **Compile** button.  
+1.	Update the other text fields to "Shots" (set to "3") and "Score" (set to "0") and click the **Compile** button.  
 
 	![](UpdateText.png)
 
