#UE4doc #docs:
Added engine version field
Replaced a couple images due to outdated UI and test that didn't match the documentation.
Changed font settings for a few terms.

diff --git a/Engine/Documentation/Source/Engine/Blueprints/HowTo/BPHT_2/BPHT_2.INT.udn b/Engine/Documentation/Source/Engine/Blueprints/HowTo/BPHT_2/BPHT_2.INT.udn
--- a/Engine/Documentation/Source/Engine/Blueprints/HowTo/BPHT_2/BPHT_2.INT.udn
+++ b/Engine/Documentation/Source/Engine/Blueprints/HowTo/BPHT_2/BPHT_2.INT.udn
@@ -2,6 +2,7 @@ Availability: Public
 Title:2. Required Character Setup
 Crumbs: %ROOT%, Engine, Engine/Blueprints, Engine/Blueprints/HowTo
 Description:In this step we are going to make sure our playable character is setup correctly for our game. 
+version: 4.9
 
 
 [VAR:Steps]
@@ -58,7 +59,7 @@ Let us get to work on our Character Blueprint.
 
 	![](AddVariable.png)
 
-1.	Select the new variable, then in the **Details** panel, rename it to **isActive** and make sure it is the **Boolean** type. 
+1.	Select the new variable, then in the **Details** panel, rename it to **IsActive** and make sure it is the **Boolean** type. 
 
 	![](NewVariable1.png)
 
@@ -97,27 +98,27 @@ Here we are going to make some minor changes to how many times the player can fi
 
 	This node will execute the script following it a set number of times (allowing us to control the number of times the player can shoot). 
 
-1.	Set the **N** to **3**, then while holding **Alt** drag in the **Shots** variable from the **MyBlueprint** window and reconnect the nodes. 
+1.	Set the **N** to **3**, then while holding **Alt** drag in the **Shots** variable from the **My Blueprint** window and reconnect the nodes. 
 
 	![](SetDoN.png)
 
 	Following the **Shots** node, connect the **Set** back to the **Montage Play** node. 
 
-1.	While holding **Control**, drag in the **Shots** variable again and connect it to a **Int - Int** node (set to **1**) and connect as shown. 
+1.	While holding **Control**, drag in the **Shots** variable again and connect it to a **Integer - Integer** node (set to **1**) and connect as shown. 
 
 	![](ConnectShots.png)
 
-	Here we are setting shots to equal itself - 1 each time **Fire** is pressed. 
+	Here we are setting shots to decrement by 1 each time **Fire** is pressed. 
 
 1.	**Right-click** in the graph next to **Shots** and add a **Custom Event** called **ResetShots**. 
 
 	![](CreateCustomEvent.png)
 
-1.	**Alt** drag in **Shots** and set it to **3**, then connect it following the **Custom Event** and into the **Reset** on the **DoN** node. 
+1.	**Alt** drag in **Shots** and set it to **3**, then connect it following the **Custom Event** and into the **Reset** on the **Do N** node. 
 
 	![](ResetEventSet.png)
 
-	When this custom event is called, it will set shots back to 3 and reset the DoN node so that it can be executed again. 
+	When this custom event is called, it will set shots back to 3 and reset the Do N node so that it can be executed again. 
 
 
 
