#UE4doc #docs:
Added engine version field
Replaced multiple pictures due to outdated UI
Reworded a few things to fit with the editor's updated UI
Changed font settings for a few terms
Added a note about what a node was called in previous versions

diff --git a/Engine/Documentation/Source/Engine/Blueprints/HowTo/BPHT_3/BPHT_3.INT.udn b/Engine/Documentation/Source/Engine/Blueprints/HowTo/BPHT_3/BPHT_3.INT.udn
--- a/Engine/Documentation/Source/Engine/Blueprints/HowTo/BPHT_3/BPHT_3.INT.udn
+++ b/Engine/Documentation/Source/Engine/Blueprints/HowTo/BPHT_3/BPHT_3.INT.udn
@@ -2,6 +2,7 @@ Availability: Public
 Title:3. Spawning the Shootable Rock
 Crumbs: %ROOT%, Engine, Engine/Blueprints, Engine/Blueprints/HowTo
 Description:In this step we will allow the player to press a button to spawn our shootable rock at one of our spawn locations. 
+version: 4.9
 
 
 [VAR:Steps]
@@ -33,21 +34,24 @@ First, let us allow the player to spawn the shootable rock by key press.
 
 1.	Hold **B** and **Left-click** in the graph to create a **Branch** node. 
 
-1.	**Alt** drag in the **isActive** variable and check the checkbox then connect as shown below. 
+1.	**Alt** drag in the **IsActive** variable and check the checkbox then connect as shown below. 
 
 	![](BranchAdded.png)
 
-1.	Hold **Control** and drag in the **isActive** and **Pulls** variables. 
+1.	Hold **Control** and drag in the **IsActive** and **Pulls** variables. 
 
-1.	Connect the **isActive** to a **!=** node (checked) and the **Pulls** to a **>** node (set to 0). 
+1.	Connect the **IsActive** to a **NOT Boolean** node (Not true) and the **Pulls** to a **>** node (set to 0). 
 
 	![](TwoNodesAdded.png)
+	[REGION:note]
+	In older versions of the engine, the **NOT Boolean** node is referred to as the **!=** node
+	[/REGION]
 
 1.	Connect both outputs to an **And** node and feed it into the **Condition** of the **Branch**. 
 
 	![](AndAdded.png)
 
-	Here, we set up the logic to determine if the player can spawn a rock by checking if a variable called **isActive** is not true and if the player's **Pulls** count is greater than 0. If both those statements are true, we will proceed down the **True** side of the **Branch** node and spawn the rock (well, that part of the script still needs to be set up which we will do below). 
+	Here, we set up the logic to determine if the player can spawn a rock by checking if a variable called **IsActive** is not true and if the player's **Pulls** count is greater than 0. If both those statements are true, we will proceed down the **True** side of the **Branch** node and spawn the rock (well, that part of the script still needs to be set up which we will do below). 
 
 1.	**Alt** drag in the **Pulls** variable and connect it following the **Set IsActive** node.
 
@@ -65,7 +69,7 @@ First, let us allow the player to spawn the shootable rock by key press.
 
 	Our logic for when the player can spawn a rock is set, now we need to communicate with the level to get a spawn location. We will do this by using an **Event Dispatcher** that is called when **F** is pressed. The **Level Blueprint** will listen for the Event Dispatcher and when it is called, will fire its own event that will get the location of one of our spawn points and pass that back to our Character Blueprint. 
 
-1.	Click the **Add Event Dispatchers** button in the **MyBlueprint** window. 
+1.	Click the **+ Event Dispatcher** button in the **My Blueprint** window. 
 
 	![](AddEventDispatcher.png)
 
@@ -93,7 +97,7 @@ First, let us allow the player to spawn the shootable rock by key press.
 
 	Our Character Blueprint is now set up to spawn a shootable rock and receive a **Transform** location to spawn the rock at. 
 
-In the Next Step, we will use the **Level Blueprint** to get the location of one of our spawn points and pass that information to our Character Blueprint when our Character Blueprint's Event Dispatcher is called (so we are going from Character Blueprint to Level Blueprint and back to Character Blueprint, passing information between them). 
+In the next step, we will use the **Level Blueprint** to get the location of one of our spawn points and pass that information to our Character Blueprint when our Character Blueprint's Event Dispatcher is called (so we are going from Character Blueprint to Level Blueprint and back to Character Blueprint, passing information between them). 
 
 
 
