#UE4doc #docs - Added Skill Level field. Added Engine Version Field. Replaced imagetable with comparison slider.

diff --git a/Engine/Documentation/Source/Engine/Rendering/LightingAndShadows/QuickStart/4/LQS_4.INT.udn b/Engine/Documentation/Source/Engine/Rendering/LightingAndShadows/QuickStart/4/LQS_4.INT.udn
--- a/Engine/Documentation/Source/Engine/Rendering/LightingAndShadows/QuickStart/4/LQS_4.INT.udn
+++ b/Engine/Documentation/Source/Engine/Rendering/LightingAndShadows/QuickStart/4/LQS_4.INT.udn
@@ -2,6 +2,8 @@ Availability: Public
 Title:4. Adding a Spot Light
 Crumbs:%ROOT%, GettingStarted, Engine\Rendering\LightingAndShadows\QuickStart
 Description:Here we will create a stylized looking light by using a Spot Light. 
+SkillLevel: Beginner
+Version: 4.9
 
 
 [VAR:Steps]
@@ -78,14 +80,18 @@ You can also use **Spot Lights** to provide additional ambient light.
 
 	Doing this provides a contrast to our level lighting and creates a more naturally lit looking room. 
 
-[REGION:imagetable]
-| Ambient Spot Light Comparison ||
-| --- | --- |
-| ![](withSpotLights.png)(w:455) | ![](withoutSpotLights.png)(w:455) |
-| **With ambient Spot Light** | **Without ambient Spot Light** |
-[/REGION]
 
-We could have used a [Sky Light](Engine/Rendering/LightingAndShadows/LightTypes/SkyLight) instead of using a **Spot Light** for our extra ambient lighting. On a large scale area this might make more sense, however for our little apartment and to control the lighting more effectively, we used **Spot Lights** to create our ambient lighting. 
+[OBJECT:ComparisonSlider]
+	[PARAM:before]
+	![With ambient Spot Light](withSpotLights.png)(w:700)
+	[/PARAM]
+	[PARAM:after]
+	![Without ambient Spot Light](withoutSpotLights.png)(w:700)
+	[/PARAM]
+[/OBJECT]
+
+
+We could have used a [**Sky Light**](Engine/Rendering/LightingAndShadows/LightTypes/SkyLight) instead of using a **Spot Light** for our extra ambient lighting. On a large scale area this might make more sense, however for our little apartment and to control the lighting more effectively, we used **Spot Lights** to create our ambient lighting. 
 
 In the **Next Step** we will improve the quality of our lights and add reflective lighting to our apartment. 	
 
