#UE4doc #docs - Added Engine Version Field. Added Skill Level Field. Updated two images to match latest UI.

diff --git a/Engine/Documentation/Source/Engine/Rendering/LightingAndShadows/QuickStart/5/LQS_5.INT.udn b/Engine/Documentation/Source/Engine/Rendering/LightingAndShadows/QuickStart/5/LQS_5.INT.udn
--- a/Engine/Documentation/Source/Engine/Rendering/LightingAndShadows/QuickStart/5/LQS_5.INT.udn
+++ b/Engine/Documentation/Source/Engine/Rendering/LightingAndShadows/QuickStart/5/LQS_5.INT.udn
@@ -2,7 +2,8 @@ Availability: Public
 Title:5. Light Quality and Reflections
 Crumbs:%ROOT%, GettingStarted, Engine\Rendering\LightingAndShadows\QuickStart
 Description:In this step we'll take a look at optimizing and improving light quality. 
-
+SkillLevel: Beginner
+Version: 4.9
 
 [VAR:Steps]
 [OBJECT:Navigation]
