#UE4doc #docs - Added Engine Version Field. Added Skill Level Field.

diff --git a/Engine/Documentation/Source/Engine/Rendering/LightingAndShadows/QuickStart/6/LQS_6.INT.udn b/Engine/Documentation/Source/Engine/Rendering/LightingAndShadows/QuickStart/6/LQS_6.INT.udn
--- a/Engine/Documentation/Source/Engine/Rendering/LightingAndShadows/QuickStart/6/LQS_6.INT.udn
+++ b/Engine/Documentation/Source/Engine/Rendering/LightingAndShadows/QuickStart/6/LQS_6.INT.udn
@@ -2,7 +2,8 @@ Availability: Public
 Title:6. Using a Light Profile
 Crumbs:%ROOT%, GettingStarted, Engine\Rendering\LightingAndShadows\QuickStart
 Description:In this step we'll create a slightly more advanced light which uses a Light Profile. 
-
+SkillLevel: Beginner
+Version: 4.9
 
 [VAR:Steps]
 [OBJECT:Navigation]
