#UE4doc #docs - Added Engine Version Field. Added Skill Level Field.

diff --git a/Engine/Documentation/Source/Engine/Rendering/LightingAndShadows/QuickStart/7/LQS_7.INT.udn b/Engine/Documentation/Source/Engine/Rendering/LightingAndShadows/QuickStart/7/LQS_7.INT.udn
--- a/Engine/Documentation/Source/Engine/Rendering/LightingAndShadows/QuickStart/7/LQS_7.INT.udn
+++ b/Engine/Documentation/Source/Engine/Rendering/LightingAndShadows/QuickStart/7/LQS_7.INT.udn
@@ -2,7 +2,8 @@ Availability: Public
 Title:7. On Your Own!
 Crumbs:%ROOT%, GettingStarted, Engine\Rendering\LightingAndShadows\QuickStart
 Description:In this final step, you'll see some examples you can take to make this your own project as well as links to additional topics on lighting and lighting techniques. 
-
+SkillLevel: Beginner
+Version: 4.9
 
 [VAR:Steps]
 [OBJECT:Navigation]
