#UE4doc #docs - Added Engine Version Field.  Bolded hyper links.

diff --git a/Engine/Documentation/Source/Engine/Rendering/LightingAndShadows/RayTracedDistanceFieldShadowing/RayTracedDistanceFieldShadowing.INT.udn b/Engine/Documentation/Source/Engine/Rendering/LightingAndShadows/RayTracedDistanceFieldShadowing/RayTracedDistanceFieldShadowing.INT.udn
--- a/Engine/Documentation/Source/Engine/Rendering/LightingAndShadows/RayTracedDistanceFieldShadowing/RayTracedDistanceFieldShadowing.INT.udn
+++ b/Engine/Documentation/Source/Engine/Rendering/LightingAndShadows/RayTracedDistanceFieldShadowing/RayTracedDistanceFieldShadowing.INT.udn
@@ -2,6 +2,7 @@ Availability: Public
 Title: Ray Traced Distance Field Soft Shadows
 Crumbs: %ROOT%, Engine, Engine/Rendering/LightingAndShadows
 Description:Dynamic area shadows
+Version: 4.9
 
 [VAR:Topic]
 [OBJECT:Topic]
@@ -50,7 +51,7 @@ Description:Dynamic area shadows
 
 
 
-Ray Traced Distance Field shadows leverage the properties of a distance field representation of the scene to compute efficient area shadowing from dynamic meshes.  This uses the same data as [Distance Field Ambient Occlusion](Engine/Rendering/LightingAndShadows/DistanceFieldAmbientOcclusion), and therefore has many of the same limitations.
+Ray Traced Distance Field shadows leverage the properties of a distance field representation of the scene to compute efficient area shadowing from dynamic meshes.  This uses the same data as [**Distance Field Ambient Occlusion**](Engine/Rendering/LightingAndShadows/DistanceFieldAmbientOcclusion), and therefore has many of the same limitations.
 
 To calculate shadowing, a ray is traced from the point being shaded through the scene's signed distance fields toward each light.  The closest distance to an occluder is used to approximate a cone trace for about the same cost as the ray trace.  This allows high quality area shadowing from spherical light source shapes.
 
@@ -75,7 +76,7 @@ The editor draws the source shape of a light with yellow lines.
 Ray Traced soft shadows from a point light with a Source Radius of 100.
 [/REGION]
 
-[_Video of a moving point light using Ray Traced Distance Field shadowing._](https://www.youtube.com/watch?v=-xai0iBffUc&list=UUBobmJyzsJ6Ll7UbfhI4iwQ)
+**[_Video of a moving point light using Ray Traced Distance Field shadowing._](https://www.youtube.com/watch?v=-xai0iBffUc&list=UUBobmJyzsJ6Ll7UbfhI4iwQ)**
 
 ## Directional lights
 
@@ -102,7 +103,7 @@ Ray Traced Distance Field shadows behave much more gracefully in the distance, d
 [/OBJECT]
 
 
-[_Video of changing time of day using Ray Traced Distance Field shadowing._](https://www.youtube.com/watch?v=4249b94KtyA&list=UUBobmJyzsJ6Ll7UbfhI4iwQ)
+**[_Video of changing time of day using Ray Traced Distance Field shadowing._](https://www.youtube.com/watch?v=4249b94KtyA&list=UUBobmJyzsJ6Ll7UbfhI4iwQ)**
 
 ## Quality
 
@@ -112,7 +113,7 @@ These shadows are computed at half resolution with a depth-aware upsample.  Temp
 
 ## Limitations
 
-The limitations of this technique are similar to [Distance Field Ambient Occlusion](Engine/Rendering/LightingAndShadows/DistanceFieldAmbientOcclusion)'s limitations:
+The limitations of this technique are similar to [**Distance Field Ambient Occlusion**](Engine/Rendering/LightingAndShadows/DistanceFieldAmbientOcclusion)'s limitations:
 
 * Only Static Mesh components, Instanced Static Mesh components and foliage (once "Affect Distance Field Lighting" has been enabled) generate distance fields and can cast shadows.  However, skeletal meshes can still composite their shadows with 'Cast Inset Shadow'.  Landscape is best handled with Far cascades enabled by "Far Shadow Cascade Count".
 * Only slight non-uniform scaling (squishing) of meshes is supported.  Artifacts in the penumbras will appear with higher amounts of mesh squishing.
