#UE4doc #docs -  Added Engine Version Field. Added Skill Level Field. Replaced image tables with Comparison Slider where needed.

diff --git a/Engine/Documentation/Source/Engine/Rendering/LightingAndShadows/ReflectionEnvironment/ReflectionEnvironment.INT.udn b/Engine/Documentation/Source/Engine/Rendering/LightingAndShadows/ReflectionEnvironment/ReflectionEnvironment.INT.udn
--- a/Engine/Documentation/Source/Engine/Rendering/LightingAndShadows/ReflectionEnvironment/ReflectionEnvironment.INT.udn
+++ b/Engine/Documentation/Source/Engine/Rendering/LightingAndShadows/ReflectionEnvironment/ReflectionEnvironment.INT.udn
@@ -2,6 +2,7 @@ Availability:Public
 Title: Reflection Environment
 Crumbs: %ROOT%, Engine, Engine/Rendering/LightingAndShadows
 Description:System for capturing and displaying local glossy reflections.
+Version: 4.9
 
 [VAR:Topic]
 [OBJECT:Topic]
@@ -68,11 +69,15 @@ analytical lights, but that only provides reflections in a few bright directions
 the sky through a Sky Light, but that does not provide local reflections, since the Sky Light cubemap is infinitely far away. Indirect specular allows all parts of the scene to reflect on all other parts, which 
 is especially important for materials like metal which have no diffuse response.
 
-[REGION:imagetable]
-| ![Diffuse Only](DiffuseOnly.png)(w:465) | ![Reflection Only](ReflectionOnly.png)(w:465) |
-| ------ | ------ |
-| Diffuse Only | Reflection Only |
-[/REGION]
+[OBJECT:ComparisonSlider]
+	[PARAM:before]
+	![Diffuse Only](DiffuseOnly.png)(w:700)
+	[/PARAM]
+	[PARAM:after]
+	![Reflection Only](ReflectionOnly.png)(w:700)
+	[/PARAM]
+[/OBJECT]
+
 
 **Full lighting**
 
@@ -106,11 +111,16 @@ Reflection Environment as they will put direct lighting in the lightmap.**  Note
 also means that the map must have meaningful indirect diffuse lighting, and that lighting must already be built 
 to see results.
 
-[REGION:imagetable]
-| ![Rough with no Shadowing](2RoughWithNoShadowing.png)(w:465) | ![Rough with Shadowing](2RoughWithShadowing.png)(w:465) |
-| ------ | ------ |
-| Reflections on a rough surface with no shadowing | Shadowing from the lightmap and SSAO incorporated |
-[/REGION]
+[OBJECT:ComparisonSlider]
+	[PARAM:before]
+	![Reflections on a rough surface with no shadowing](2RoughWithNoShadowing.png)(w:550)
+	[/PARAM]
+	[PARAM:after]
+	![Rough with Shadowing](2RoughWithShadowing.png)(w:550)
+	[/PARAM]
+[/OBJECT]
+
+]
 
 ## Reflection capture shapes
 
