#UE4doc #docs -  Added Engine Version Field.

diff --git a/Engine/Documentation/Source/Engine/Rendering/LightingAndShadows/Shadows/Shadows.INT.udn b/Engine/Documentation/Source/Engine/Rendering/LightingAndShadows/Shadows/Shadows.INT.udn
--- a/Engine/Documentation/Source/Engine/Rendering/LightingAndShadows/Shadows/Shadows.INT.udn
+++ b/Engine/Documentation/Source/Engine/Rendering/LightingAndShadows/Shadows/Shadows.INT.udn
@@ -3,6 +3,7 @@ Title:Shadow Casting
 Crumbs: %ROOT%, Engine, Engine/Rendering/LightingAndShadows
 Description:High level overview of shadows.
 Related: Engine/Rendering/LightingAndShadows/RayTracedDistanceFieldShadowing
+Version: 4.9
 
 [VAR:Topic]
 [OBJECT:Topic]
@@ -56,7 +57,7 @@ Shadows make objects feel grounded in the world and give the viewer a sense of d
 ## Static Lights
 
 [EXCERPT:StaticLightShadows]
-Static lights cast completely static shadows **and light**, meaning they will have no _direct_ effect on dynamic objects (static lights are baked into the Indirect Lighting Cache so they have _some_ effect) like on the bronze men below.
+Static lights cast completely static shadows and light, meaning they will have no _direct_ effect on dynamic objects (static lights are baked into the Indirect Lighting Cache so they have _some_ effect) like on the bronze men below.
 [/EXCERPT:StaticLightShadows]
 
 ![](LI_StaticMeansStatic.png)(w:900)
