#UE4doc #docs - Added Engine Version Field.

diff --git a/Engine/Documentation/Source/Engine/Rendering/Materials/MaterialsAndTextures.INT.udn b/Engine/Documentation/Source/Engine/Rendering/Materials/MaterialsAndTextures.INT.udn
--- a/Engine/Documentation/Source/Engine/Rendering/Materials/MaterialsAndTextures.INT.udn
+++ b/Engine/Documentation/Source/Engine/Rendering/Materials/MaterialsAndTextures.INT.udn
@@ -3,7 +3,7 @@ Title:Materials
 Crumbs: %ROOT%, Engine
 Description:Controlling the appearance of surfaces in the world using shaders.
 Navigation:topic
-
+Version:4.9
 
 [VAR:Topic]
 [OBJECT:Topic]
