#UE4doc #docs - Added Engine Version Field. Removed Overview header.

diff --git a/Engine/Documentation/Source/Engine/Rendering/Materials/CustomizedUVs/CustomizedUVs.INT.udn b/Engine/Documentation/Source/Engine/Rendering/Materials/CustomizedUVs/CustomizedUVs.INT.udn
--- a/Engine/Documentation/Source/Engine/Rendering/Materials/CustomizedUVs/CustomizedUVs.INT.udn
+++ b/Engine/Documentation/Source/Engine/Rendering/Materials/CustomizedUVs/CustomizedUVs.INT.udn
@@ -2,6 +2,7 @@ Availability: Public
 Title:Customized UVs
 Crumbs: %ROOT%, Engine, Engine/Rendering/Materials
 Description:Feature that allows running calculations in the vertex shader to increase performance over running them per-pixel.
+Version: 4.9
 
 [VAR:Topic]
 [OBJECT:Topic]
@@ -44,7 +45,6 @@ Description:Feature that allows running calculations in the vertex shader to inc
 
 [TOC(start:2)]
 
-## Overview
 
 On GPU's, the vertex shader is run for every vertex, and the pixel shader is run for every pixel. Almost all the material nodes in Unreal Engine 4 are run for every pixel. While a **UV Coordinate** node can be either part of the vertex or pixel shader, the **CustomizedUVs** feature is only run in the vertex shader, and offers a performance increase over running the same calculations in the pixel shader. This offers an excellent way to speed up even just tiling a texture. While the system is not limited in the math you run on the UVs, the result will depend on the tessellation of your mesh.
 
