#UE4doc #docs - Added Engine Version Field. Added Skill Level field.  Added Comparison Slider for image table with comparison of lightmap resolutions.

diff --git a/Engine/Documentation/Source/Engine/Rendering/Materials/HowTo/ColoredTransluscentShadows/ColoredTransluscentShadows.INT.udn b/Engine/Documentation/Source/Engine/Rendering/Materials/HowTo/ColoredTransluscentShadows/ColoredTransluscentShadows.INT.udn
--- a/Engine/Documentation/Source/Engine/Rendering/Materials/HowTo/ColoredTransluscentShadows/ColoredTransluscentShadows.INT.udn
+++ b/Engine/Documentation/Source/Engine/Rendering/Materials/HowTo/ColoredTransluscentShadows/ColoredTransluscentShadows.INT.udn
@@ -6,12 +6,14 @@ Related: Engine/Rendering/Materials/MaterialProperties/BlendModes
 Related: Engine/Rendering/LightingandShadows/Lightmass
 Related: Engine/Rendering/Materials/Functions/Reference/Opacity
 Related: Engine/Rendering/LightingAndShadows/LitTranslucency
+Version: 4.9
+SkillLevel: Intermediate
 
 [TOC(start:2 end:3)]
 
 
 
-In this "How-to," we intend to show you how you can setup and use **Translucent Shadow Color** to create a material that can cast a colored shadow. 
+In this How-to, we intend to show you how you can setup and use **Translucent Shadow Color** to create a material that can cast a colored shadow. 
 There are many applications where this could be useful to inject some color into your scene, but a common use would be for something like a stained glass window.
  
  **Example:**
@@ -82,6 +84,7 @@ Notes
 * Translucency Lighting Model: TLM_Surface
 1. Plug the Texture Sample into **Emissive**.
 
+
 ## Colored Translucent Shadows in Action
 
 ### Alpha Channel for Mask
@@ -103,12 +106,15 @@ Shadow sharpness can be affected by several different factors, including the lig
 
 ### Lightmap Resolution
 
+[OBJECT:ComparisonSlider]
+	[PARAM:before]
+	![Lightmap Resolution 64](TCS_LmRes_64.png)(w:450)
+	[/PARAM]
+	[PARAM:after]
+	![Lightmap Resolution 1024](TCS_LmRes_1024.png)(w:450)
+	[/PARAM]
+[/OBJECT]
 
-[REGION:imagetable]
-| ![](TCS_LmRes_64.png)(w:350) | ![](TCS_lmRes_1024.png)(h:367) |
-| --- | --- |
-| Lightmap Resolution = 64 | Lightmap Resolution = 1024 |
-[/REGION]
 
 
 ##Limitations
