#UE4doc #docs - Added Engine Version Field. Updated multiple comparison images to user the Comparison Slider.

diff --git a/Engine/Documentation/Source/Engine/Rendering/PostProcessEffects/LensFlare/LensFlare.INT.udn b/Engine/Documentation/Source/Engine/Rendering/PostProcessEffects/LensFlare/LensFlare.INT.udn
--- a/Engine/Documentation/Source/Engine/Rendering/PostProcessEffects/LensFlare/LensFlare.INT.udn
+++ b/Engine/Documentation/Source/Engine/Rendering/PostProcessEffects/LensFlare/LensFlare.INT.udn
@@ -2,6 +2,7 @@ Availability:Public
 Title:Lens Flare
 Crumbs: %ROOT%, Engine, Engine/Rendering/PostProcessEffects
 Description:Simulating scattered light from bright objects due to imperfections in camera lenses.
+Version: 4.9
 
 [REGION:fullwidth]
 ![Lens Flare](lens_flare.png)(w:1182)
@@ -20,36 +21,51 @@ The **Lens Flare** effect is an image-based technique that simulates the scatter
 | **Lens Flare Tints 1/2/3/4/5/6/7/8** | Tints the color of each individual lens flare. |
 | **Intensity** | Brightness scale of the image cased lens flares (linear). |
 
+
+
 <!--
+
 [EXCERPT:Tint]
-[REGION:imagetable]
-| ![Lens Flare Tint - (1.0, 1.0, 1.0)](lens_flare_tint_off_small.png)(w:350) | ![Lens Flare Tint - (0.25, 1.0, 1.0)](lens_flare_tint_on_small.png)(w:350)
-| ------ | ------ |
-| (1.0, 1.0, 1.0) | (0.25, 1.0, 1.0) |
-[/REGION]
+[OBJECT:ComparisonSlider]
+	[PARAM:before]
+	![Lens Flare Tint (1.0, 1.0, 1.0)](lens_flare_tint_off_small.png)(w:500)
+	[/PARAM]
+	[PARAM:after]
+	![Lens Flare Tint (0.25, 1.0, 1.0)](lens_flare_tint_on_small.png)(w:500)
+	[/PARAM]
+[/OBJECT]
 [/EXCERPT:Tint]
 
 [EXCERPT:Threshold]
-[REGION:imagetable]
-| ![Lens Flare Threshold - 0.1](lens_flare_threshold_0_small.png)(w:350) | ![Lens Flare Threshold - 8](lens_flare_threshold_8_small.png)(w:350) |
-| ------ | ------ |
-| 0.1 | 8.0 |
-[/REGION]
+[OBJECT:ComparisonSlider]
+	[PARAM:before]
+	![Lens Flare Threshold - 0.1](lens_flare_threshold_0_small.png)(w:500)
+	[/PARAM]
+	[PARAM:after]
+	![Lens Flare Threshold - 8](lens_flare_threshold_8_small.png)(w:500)
+	[/PARAM]
+[/OBJECT]
 [/EXCERPT:Threshold]
 
 [EXCERPT:BokehSize]
-[REGION:imagetable]
-| ![Lens Flare Size - 1](lens_flare_bokehsize_1_small.png)(w:350) | ![Lens Flare Size - 8](lens_flare_bokehsize_8_small.png)(w:350) |
-| ------ | ------ |
-| 1.0 | 8.0 |
-[/REGION]
+[OBJECT:ComparisonSlider]
+	[PARAM:before]
+	![Lens Flare Size - 1](lens_flare_bokehsize_1_small.png)(w:500)
+	[/PARAM]
+	[PARAM:after]
+	![Lens Flare Size - 8](lens_flare_bokehsize_8_small.png)(w:500)
+	[/PARAM]
+[/OBJECT]
 [/EXCERPT:BokehSize]
 
 [EXCERPT:BokehShape]
-[REGION:imagetable]
-| ![Lens Flare Shape - Default](lens_flare_bokehshape_default_small.png)(w:350) | ![Lens Flare Shape - Square](lens_flare_bokehshape_square_small.png)(w:350) |
-| ------ | ------ |
-| Bokeh Shape - Default | Bokeh Shape - Square |
-[/REGION]
+[OBJECT:ComparisonSlider]
+	[PARAM:before]
+	![Lens Flare Shape - Default](lens_flare_bokehshape_default_small.png)(w:500)
+	[/PARAM]
+	[PARAM:after]
+	![Lens Flare Shape - Square](lens_flare_bokehshape_square_small.png)(w:500)
+	[/PARAM]
+[/OBJECT]
 [/EXCERPT:BokehShape]
 -->
