#UE4doc #docs - Added Engine Version field.

diff --git a/Engine/Documentation/Source/Engine/Rendering/PostProcessEffects/PostProcessEffects.INT.udn b/Engine/Documentation/Source/Engine/Rendering/PostProcessEffects/PostProcessEffects.INT.udn
--- a/Engine/Documentation/Source/Engine/Rendering/PostProcessEffects/PostProcessEffects.INT.udn
+++ b/Engine/Documentation/Source/Engine/Rendering/PostProcessEffects/PostProcessEffects.INT.udn
@@ -3,6 +3,7 @@ Title: Post Process Effects
 Crumbs: %ROOT%, Engine
 Description:Effects applied to the whole rendered scene prior to being rendered.
 Navigation:topic
+Version: 4.9
 
 [VAR:Topic]
 [OBJECT:Topic]
