#UE4doc #docs - Added Engine version field. Replaced Image table with Comparison Slider.

diff --git a/Engine/Documentation/Source/Engine/Rendering/PostProcessEffects/AntiAliasing/AntiAliasing.INT.udn b/Engine/Documentation/Source/Engine/Rendering/PostProcessEffects/AntiAliasing/AntiAliasing.INT.udn
--- a/Engine/Documentation/Source/Engine/Rendering/PostProcessEffects/AntiAliasing/AntiAliasing.INT.udn
+++ b/Engine/Documentation/Source/Engine/Rendering/PostProcessEffects/AntiAliasing/AntiAliasing.INT.udn
@@ -2,24 +2,29 @@ Availability:Public
 Title:Anti-Aliasing
 Crumbs: %ROOT%, Engine, Engine/Rendering/PostProcessEffects
 Description:Smoothing of jagged edges using the FXAA method.
+Version: 4.9
 
-Anti-Aliasing (AA) refers to the smoothing of aliased, or jagged, lines when displayed on a computer monitor.
 
-[REGION:imagetable]
-| ![Without Anti-Aliasing](aa_off.png)(w:455) | ![With Anti-Aliasing](aa_on.png)(w:455) |
-| :----: | :----: |
-| Without Anti-Aliasing | With Anti-Aliasing |
-[/REGION]
+**Anti-Aliasing** (AA) refers to the smoothing of aliased, or jagged, lines when displayed on a computer monitor.
+
+[OBJECT:ComparisonSlider]
+	[PARAM:before]
+	![Without Anti-Aliasing](aa_off.png)(w:800)
+	[/PARAM]
+	[PARAM:after]
+	![With Anti-Aliasing](aa_on.png)(w:800)
+	[/PARAM]
+[/OBJECT]
 
 Anti-aliasing in Unreal Engine 4 is performed in post processing using FXAA, which is an efficient GPU MLAA implementation. 
 This method solves most causes of aliasing artifacts, but cannot fully prevent temporal aliasing.
 
-Other techniques: UE4 currently supports MSAA only for a few editor 3d primitives (e.g. 3D Gizmo).
+Other techniques: UE4 currently supports MSAA only for a few editor 3d primitives (e.g. 3D Gizmo) and Temporal AA.
 
-The feature can be disabled by the show flag "Post-process Anti-aliasing."
+The feature can be disabled by the show flag **Post-process Anti-aliasing**.
 
 
-For more information, see: [FXAA Whitepaper](http://developer.download.nvidia.com/assets/gamedev/files/sdk/11/FXAA_WhitePaper.pdf) (external link)
+For more information, see: [**FXAA Whitepaper**](http://developer.download.nvidia.com/assets/gamedev/files/sdk/11/FXAA_WhitePaper.pdf) (external link)
 
 
 
