#UE4doc #docs - Added Engine Version Field. Updated Image to match style guide coloring for arrows.

diff --git a/Engine/Documentation/Source/Engine/Rendering/PostProcessEffects/AutomaticExposure/AutomaticExposure.INT.udn b/Engine/Documentation/Source/Engine/Rendering/PostProcessEffects/AutomaticExposure/AutomaticExposure.INT.udn
--- a/Engine/Documentation/Source/Engine/Rendering/PostProcessEffects/AutomaticExposure/AutomaticExposure.INT.udn
+++ b/Engine/Documentation/Source/Engine/Rendering/PostProcessEffects/AutomaticExposure/AutomaticExposure.INT.udn
@@ -2,6 +2,7 @@ Availability:Public
 Title:Eye Adaptation
 Crumbs: %ROOT%, Engine, Engine/Rendering/PostProcessEffects
 Description:Automatic adjustment of scene exposure to simulate eye adaptation from changes in brightness.
+Version: 4.9
 
 [EXCERPT:EyeAdaptation_Desc]
 **Eye Adaptation**, or automatic exposure, causes the exposure of the scene to automatically adjust to recreate the effect experienced as human eyes adjust when going from a bright environment into a dark environment or vice versa.
