#UE4doc #docs - Added Engine Version Field.

diff --git a/Engine/Documentation/Source/Engine/Rendering/PostProcessEffects/Bloom/Bloom.INT.udn b/Engine/Documentation/Source/Engine/Rendering/PostProcessEffects/Bloom/Bloom.INT.udn
--- a/Engine/Documentation/Source/Engine/Rendering/PostProcessEffects/Bloom/Bloom.INT.udn
+++ b/Engine/Documentation/Source/Engine/Rendering/PostProcessEffects/Bloom/Bloom.INT.udn
@@ -2,6 +2,7 @@ Availability:Public
 Title:Bloom
 Crumbs: %ROOT%, Engine, Engine/Rendering/PostProcessEffects
 Description:Haloing to produce glow effects for bright objects such as lights.
+Version: 4.9
 
 [EXCERPT:Bloom_Desc]
 **Bloom** is a real world light phenomena that can greatly add to the perceived realism of a rendered image at a moderate render performance cost. Bloom can be seen by the naked eye when looking at very bright objects that are on a much darker background. Even brighter objects also cause other effects (streaks, lens flares), but those are not covered by the classic bloom effect. Because our displays (e.g. TV, TFT, ...) usually do not support HDR (high dynamic range), we cannot really render very bright objects. Instead we simulate the effects that happen in the eye (retina subsurface scattering), when light hits the film (film subsurface scattering), or in front of the camera (milky glass filter). The effect might not always be physically correct but it can help to hint the relative brightness of objects or add realism to the LDR (low dynamic range) image that is shown on the screen.
