#UE4doc #docs - Added Engine Version field. Replaced two image tables with comparison slider.

diff --git a/Engine/Documentation/Source/Engine/Rendering/PostProcessEffects/ColorGrading/ColorGrading.INT.udn b/Engine/Documentation/Source/Engine/Rendering/PostProcessEffects/ColorGrading/ColorGrading.INT.udn
--- a/Engine/Documentation/Source/Engine/Rendering/PostProcessEffects/ColorGrading/ColorGrading.INT.udn
+++ b/Engine/Documentation/Source/Engine/Rendering/PostProcessEffects/ColorGrading/ColorGrading.INT.udn
@@ -2,6 +2,7 @@ Availability:Public
 Title:Color Grading
 Crumbs: %ROOT%, Engine, Engine/Rendering/PostProcessEffects
 Description:Tonemapping and color correction effects for adjusting scene colors.
+Version: 4.9
 
 Within the Unreal engine, the term **Color Grading** covers the Tone Mapping function (HDR to LDR transformation) and the further Color Correction (LDR color to screen color transformation).
 
@@ -30,19 +31,25 @@ The presented tone mapper formula already has two editable constants and while m
 
 <!--
 [EXCERPT:SceneColorTint]
-[REGION:imagetable]
-| ![Scene Color Tint - (1.0, 1.0, 1.0)](scene_tint_full_small.png)(w:365) | ![Scene Color Tint - (0.25, 0.25, 0.25)](scene_tint_quarter_small.png)(w:365)
-| ------ | ------ |
-| (1.0, 1.0, 1.0) | (0.25, 0.25, 0.25) |
-[/REGION]
+[OBJECT:ComparisonSlider]
+	[PARAM:before]
+	![Scene Color Tint - (1.0, 1.0, 1.0)](scene_tint_full_small.png)(w:600)
+	[/PARAM]
+	[PARAM:after]
+	![Scene Color Tint - (0.25, 0.25, 0.25)](scene_tint_quarter_small.png)(w:600)
+	[/PARAM]
+[/OBJECT]
 [/EXCERPT:SceneColorTint]
 
 [EXCERPT:TonemapperToeFactor]
-[REGION:imagetable]
-| ![Toe Factor - 0](toe_factor_0_small.png)(w:365) | ![Toe Factor - 1](toe_factor_1_small.png)(w:365)
-| ------ | ------ |
-| 0.0 | 1.0 |
-[/REGION]
+[OBJECT:ComparisonSlider]
+	[PARAM:before]
+	![Toe Factor - 0](toe_factor_0_small.png)(w:600)
+	[/PARAM]
+	[PARAM:after]
+	![Toe Factor - 1](toe_factor_1_small.png)(w:600)
+	[/PARAM]
+[/OBJECT]
 [/EXCERPT:TonemapperToeFactor]
 -->
 
