#UE4doc #docs - Added Engine Version Field. Replaced multiple image tables with comparison sliders.

diff --git a/Engine/Documentation/Source/Engine/Rendering/PostProcessEffects/DepthOfField/DepthOfField.INT.udn b/Engine/Documentation/Source/Engine/Rendering/PostProcessEffects/DepthOfField/DepthOfField.INT.udn
--- a/Engine/Documentation/Source/Engine/Rendering/PostProcessEffects/DepthOfField/DepthOfField.INT.udn
+++ b/Engine/Documentation/Source/Engine/Rendering/PostProcessEffects/DepthOfField/DepthOfField.INT.udn
@@ -2,13 +2,14 @@ Availability:Public
 Title:Depth of Field
 Crumbs: %ROOT%, Engine, Engine/Rendering/PostProcessEffects
 Description:Simulating focus by blurring the scene depending on depth.
+Version: 4.9
 
 [REGION:fullwidth]
 ![Depth of Field](dof.png)
 [/REGION]
 
 [EXCERPT:DepthOfField_Desc]
-Depth of Field (DoF) applies a blur to the scene based on distance in front of or behind a focal point. This is to simulate what happens in real-world cameras. The effect can be used to draw the viewer's attention and to
+**Depth of Field** (DoF) applies a blur to the scene based on distance in front of or behind a focal point. This is to simulate what happens in real-world cameras. The effect can be used to draw the viewer's attention and to
 make the rendering appear more like a photograph or like a movie.
 [/EXCERPT:DepthOfField_Desc]
 
@@ -60,12 +61,14 @@ Depth of Field is broken up into three layers - Near, Far, and the Focal Region
 
 
 These layers, including the transition regions, can be visualized using the **Depth of Field Layers** show flag ( under **Show > Visualize** ). The near layer is shown in green, the far layer in blue, and the focal region is black.
-
-[REGION:imagetable]
-| ![Scene](dof_layers.png)(w:455) | ![Layer Visualization](dof_layers_visualize.png)(w:455) |
-| :----: | :----: |
-| Scene with DoF | DoF Layers Visualized |
-[/REGION]
+[OBJECT:ComparisonSlider]
+	[PARAM:before]
+	![Scene](dof_layers.png)(w:700)
+	[/PARAM]
+	[PARAM:after]
+	![Layer Visualization](dof_layers_visualize.png)(w:700)
+	[/PARAM]
+[/OBJECT]
 
 We further improved this visualization to also include a lot of interesting information:
 
@@ -173,11 +176,15 @@ Note that the circle size might not always match to the actual effect as the imp
 
 Bokeh depth of field is rendered at quarter-resolution (half-res in each direction) by default for performance reasons. In most cases, this downsampling is barely noticeable and is perfectly acceptable; however, it can lead to artifacting and unwanted results in some situations.
 
-[REGION:imagetable]
-| ![Adaptive Reoslution - Off](dof_adaptive_0.png)(w:455) | ![Adaptive Reoslution - On](dof_adaptive_1.png)(w:455) |
-| :----: | :----: |
-| Downsampled | Adaptive Algorithm |
-[/REGION]
+[OBJECT:ComparisonSlider]
+	[PARAM:before]
+	![Adaptive Resolution - Off](dof_adaptive_0.png)(w:700)
+	[/PARAM]
+	[PARAM:after]
+	![Adaptive Reoslution - On](dof_adaptive_1.png)(w:700)
+	[/PARAM]
+[/OBJECT]
+
 
 You can see that the characters being blurred in the background appear blocky when using only the downsampled DoF technique. Also, there is some artifacting around the horns of the character in the foreground. The artifacting is eliminated by using the adaptive DoF, and the background characters have a much smoother appearance.
 
@@ -189,8 +196,12 @@ The adaptive DoF can be visualized using the **Adaptive Depth of Field** show fl
 
 In general, you want this to display all green. The more red shown by the visualization, the more expensive the scene is to render each frame.
 
-[REGION:imagetable]
-| ![Adaptive Reoslution - Optimized](dof_adaptive_good.png)(w:455) | ![Adaptive Reoslution - Expensive](dof_adaptive_bad.png)(w:455) |
-| :----: | :----: |
-| Optimized | Expensive |
-[/REGION]
+[OBJECT:ComparisonSlider]
+	[PARAM:before]
+	![Adaptive Reoslution - Optimized](dof_adaptive_good.png)(w:700)
+	[/PARAM]
+	[PARAM:after]
+	![Adaptive Reoslution - Expensive](dof_adaptive_bad.png)(w:700)
+	[/PARAM]
+[/OBJECT]
+
