#UE4doc #docs - Added Engine Version field. Replaced a couple of image tables with the comparison slider.

diff --git a/Engine/Documentation/Source/Engine/Rendering/PostProcessEffects/PostProcessMaterials/PostProcessMaterials.INT.udn b/Engine/Documentation/Source/Engine/Rendering/PostProcessEffects/PostProcessMaterials/PostProcessMaterials.INT.udn
--- a/Engine/Documentation/Source/Engine/Rendering/PostProcessEffects/PostProcessMaterials/PostProcessMaterials.INT.udn
+++ b/Engine/Documentation/Source/Engine/Rendering/PostProcessEffects/PostProcessMaterials/PostProcessMaterials.INT.udn
@@ -13,7 +13,7 @@ Description:How to author and blend custom Post Process passes with the material
 
 [EXCERPT:PostProcessMaterials_Desc]
 The engine already features complex post processing based on a graph of post processing nodes. The **Post Processing Materials**
-can be additionally inserted in some specific position. See [FAQ](#FAQ) **r.CompositionGraphDebug** to get a dump of the full graph.
+can be additionally inserted in some specific position. See [**FAQ**](#FAQ) **r.CompositionGraphDebug** to get a dump of the full graph.
 The graph is actually not only doing post processing but also some parts of the lighting. Over time we plan to make more
 parts be customizable through the material editor.
 [/EXCERPT:PostProcessMaterials_Desc]
@@ -40,7 +40,7 @@ new slots, select a material in the **Content Browser**, and press the left arro
 
 ## Making a Simple Post Process Material
 
-For an overview of creating a simple Post Process Material effect from scratch, see the [](Engine/Rendering/PostProcessEffects/PostProcessMaterials/PostProcessDemoMaterialOverview).
+For an overview of creating a simple Post Process Material effect from scratch, see the **[](Engine/Rendering/PostProcessEffects/PostProcessMaterials/PostProcessDemoMaterialOverview)**.
 
 ![](Engine/Rendering/PostProcessEffects/PostProcessMaterials/PostProcessDemoMaterialOverview/FinalPostEffect.png)(w:640)
 
@@ -74,11 +74,14 @@ Each volume we specify a different material instance of the same material. The c
 
 Depending on the camera position the setting of one volume are used:
 
-[REGION:imagetable]
-| ![](VignetteG.png) | ![](VignetteR.png) |
-| ------ | ------ |
-| Green vignette material instance | Red vignette material instance |
-[/REGION]
+[OBJECT:ComparisonSlider]
+	[PARAM:before]
+	![Green vignette material instance](VignetteG.png)(w:450)
+	[/PARAM]
+	[PARAM:after]
+	![Red vignette material instance](VignetteR.png)(w:450)
+	[/PARAM]
+[/OBJECT]
 
 
 With camera movement, you can perceive soft linear transition between the two effect setting.
@@ -89,11 +92,15 @@ The fuzzy borders are defined by the BlendRadius property specified in the volum
 
 With the correct setup all blending happens as expected.
 
-[REGION:imagetable]
-| ![](VolumeBlendBad.png) | ![](VolumeBlendGood.png) |
-| ------ | ------ |
-| Bad setup | Good setup |
-[/REGION]
+[OBJECT:ComparisonSlider]
+	[PARAM:before]
+	![Bad Setup](VolumeBlendBad.png)(w:450)
+	[/PARAM]
+	[PARAM:after]
+	![Good Setup](VolumeBlendGood.png)(w:450)
+	[/PARAM]
+[/OBJECT]
+
 
 The difference between the both setup is the default value specified on the material (scalar or vector) parameters.
 The good setup has values that make the pass appear to not have any effect (e.g. multiply with white or lerp with 0).
