#UE4doc #docs - Added Engine Version field. Replaced a couple of image tables with the comparison slider.

diff --git a/Engine/Documentation/Source/Engine/Rendering/PostProcessEffects/PostProcessMaterials/PostProcessMaterials.INT.udn b/Engine/Documentation/Source/Engine/Rendering/PostProcessEffects/PostProcessMaterials/PostProcessMaterials.INT.udn
--- a/Engine/Documentation/Source/Engine/Rendering/PostProcessEffects/PostProcessMaterials/PostProcessMaterials.INT.udn
+++ b/Engine/Documentation/Source/Engine/Rendering/PostProcessEffects/PostProcessMaterials/PostProcessMaterials.INT.udn
@@ -2,6 +2,7 @@ Availability:Public
 Title:Post Process Materials
 Crumbs: %ROOT%, Engine, Engine/Rendering/PostProcessEffects
 Description:How to author and blend custom Post Process passes with the material editor.
+Version: 4.9
 
 ![](Teaser3.png)
 ![](Teaser0.png)
