#UE4doc #docs:
Added engine version field
Changed font settings for a few terms
Replaced an image due to difference in UI
Added onto a step to explain how to select multiple things in the World Outliner

diff --git a/Engine/Documentation/Source/Engine/Blueprints/HowTo/BPHT_4/BPHT_4.INT.udn b/Engine/Documentation/Source/Engine/Blueprints/HowTo/BPHT_4/BPHT_4.INT.udn
--- a/Engine/Documentation/Source/Engine/Blueprints/HowTo/BPHT_4/BPHT_4.INT.udn
+++ b/Engine/Documentation/Source/Engine/Blueprints/HowTo/BPHT_4/BPHT_4.INT.udn
@@ -2,6 +2,7 @@ Availability: Public
 Title:4. Passing Information from Level to Character
 Crumbs: %ROOT%, Engine, Engine/Blueprints, Engine/Blueprints/HowTo
 Description:In this step we are using the Level Blueprint to get the location of our spawn points and kicking one over to our Character Blueprint. 
+version: 4.9
 
 
 [VAR:Steps]
@@ -71,7 +72,7 @@ Description:In this step we are using the Level Blueprint to get the location of
 
 	![](SetSpawnLocation.png)
 
-1. Return to the Main Editor window and in the **World Outliner**, search for and select each of the **Target Points**. 
+1. Return to the Main Editor window and in the **World Outliner**, search for and select each of the **Target Points** by holding **Ctrl** and clicking each one. 
 
 	![](SelectAllTargetPoints.png)
 
@@ -87,7 +88,7 @@ Description:In this step we are using the Level Blueprint to get the location of
 
 1.	Off the **Array** pin, add a **Get** node. 
 
-1.	Off the **Get**, add a **Random Integer in Range** node with **Min** set to **0** and **Max** set to **2**. 
+1.	Off the **Get**, add a **Random Integer in Range** node with **Min** set to "0" and **Max** set to "2". 
 
 	![](GetNodeAdded.png)
 
