#UE4doc #docs:
Added engine version field
Changed text to match the picture's variable names for consistency

diff --git a/Engine/Documentation/Source/Engine/Blueprints/HowTo/BPHT_5/BPHT_5.INT.udn b/Engine/Documentation/Source/Engine/Blueprints/HowTo/BPHT_5/BPHT_5.INT.udn
--- a/Engine/Documentation/Source/Engine/Blueprints/HowTo/BPHT_5/BPHT_5.INT.udn
+++ b/Engine/Documentation/Source/Engine/Blueprints/HowTo/BPHT_5/BPHT_5.INT.udn
@@ -1,7 +1,8 @@
 Availability: Public
 Title:5. Destroying the Rock 
 Crumbs: %ROOT%, Engine, Engine/Blueprints, Engine/Blueprints/HowTo
-Description:In this section we will use a Bind Event which is bound to OnHit to destroy our rock when it is hit. 
+Description:In this section we will use a Bind Event which is bound to OnHit to destroy our rock when it is hit.
+version: 4.9
 
 
 [VAR:Steps]
@@ -66,7 +67,7 @@ Let's blow stuff up!
 
 	Drag off the **Return Value** of the **Spawn Actor** node and feed it into the **Target** of the **Destroy Actor** node. 
 
-1. Following the **Destroy Actor** node, add the **Set isActive** variable node and leave it unchecked. 
+1. Following the **Destroy Actor** node, add the **Set IsActive** variable node and leave it unchecked. 
 
 	![](SetIsActive.png)
 
