#UE4doc #docs:
Added engine version field
Replaced an image due to it not being completely correct.
Added information to a couple steps to expand on the information given
Fixed spacing for a term

diff --git a/Engine/Documentation/Source/Engine/Blueprints/HowTo/BPHT_6/BPHT_6.INT.udn b/Engine/Documentation/Source/Engine/Blueprints/HowTo/BPHT_6/BPHT_6.INT.udn
--- a/Engine/Documentation/Source/Engine/Blueprints/HowTo/BPHT_6/BPHT_6.INT.udn
+++ b/Engine/Documentation/Source/Engine/Blueprints/HowTo/BPHT_6/BPHT_6.INT.udn
@@ -2,6 +2,7 @@ Availability: Public
 Title:6. Communicating with the HUD
 Crumbs: %ROOT%, Engine, Engine/Blueprints, Engine/Blueprints/HowTo
 Description:In this step we will update our HUD based on values stored in our Character Blueprint. 
+version: 4.9
 
 
 [VAR:Steps]
@@ -46,9 +47,9 @@ Our game is almost complete, let's finish it off by updating the HUD.
 
 	![](MyCharacter.png)
 
-	You can rename the variable from the **MyBlueprint** window by selecting the variable and pressing **F2**. 
+	You can rename the variable from the **My Blueprint** window by selecting the variable and pressing **F2**. 
 
-1.	On the **Designer** tab, click the value for **Pulls** then next to **Text** in the **Details** panel, click the **Bind** option. 
+1.	On the **Designer** tab, click the value for **Pulls** then next to **Text** in the **Details** panel, click the **Bind** option and choose **Create Binding**. 
 
 	![](Bind1.png)
 
@@ -89,7 +90,7 @@ Our game is almost complete, let's finish it off by updating the HUD.
 	[/OBJECT]
 
 
-That's it; your simple shooting gallery game is complete. Over the course of this guide, we have learned how to communicate between three different types of Blueprints using an **Event Dispatcher**, **Cast** nodes, and **Binding**. We have learned how to pass information from the level to a Character Blueprint as well as how to communicate variables from the character to a HUD. 
+That's it; your simple shooting gallery game is complete. To spawn a rock, press **F** as we bound it do that event in a previous section. Over the course of this guide, we have learned how to communicate between three different types of Blueprints using an **Event Dispatcher**, **Cast** nodes, and **Binding**. We have learned how to pass information from the level to a Character Blueprint as well as how to communicate variables from the character to a HUD. 
 
 In the final step, you will see some examples you can take to alter this project and make it your own. You will also be provided with links to additional documentation that covers topics in this guide. 
 
