#UE4doc #docs:
Added engine version field
Removed bolding from multiple terms

diff --git a/Engine/Documentation/Source/Engine/Blueprints/HowTo/BPHT_7/BPHT_7.INT.udn b/Engine/Documentation/Source/Engine/Blueprints/HowTo/BPHT_7/BPHT_7.INT.udn
--- a/Engine/Documentation/Source/Engine/Blueprints/HowTo/BPHT_7/BPHT_7.INT.udn
+++ b/Engine/Documentation/Source/Engine/Blueprints/HowTo/BPHT_7/BPHT_7.INT.udn
@@ -2,6 +2,7 @@ Availability: Public
 Title:7. On Your Own!
 Crumbs: %ROOT%, Engine, Engine/Blueprints, Engine/Blueprints/HowTo
 Description:In this final step, examples are given how you can take this project to the next level as well as links to additional documentation are provided. 
+version: 4.9
 
 
 [VAR:Steps]
@@ -44,15 +45,15 @@ For more information on the topics covered in this guide, and across the entire
 
 As for coverage related to the topics in this guide:
 
-* For more information on **Blueprints** in general, see: [Blueprints Visual Scripting](Engine/Blueprints)
+* For more information on Blueprints in general, see: [Blueprints Visual Scripting](Engine/Blueprints)
 
 * For more information on using Blueprint Communication and when to use it, see: [Blueprint Communication Usage](Engine/Blueprints/UserGuide/BlueprintCommsUsage)
 
-* For more information on **Event Dispatchers**, see: [Event Dispatchers](Engine/Blueprints/UserGuide/EventDispatcher)
+* For more information on Event Dispatchers, see: [Event Dispatchers](Engine/Blueprints/UserGuide/EventDispatcher)
 
-* For more information on **Casting**, see: [Casting](Engine/Blueprints/UserGuide/CastNodes)
+* For more information on Casting, see: [Casting](Engine/Blueprints/UserGuide/CastNodes)
 
-* For more information on creating HUD elements with **Unreal Motion Graphics** (UMG), see: [Unreal Motion Graphics](Engine/UMG)
+* For more information on creating HUD elements with Unreal Motion Graphics (UMG), see: [Unreal Motion Graphics](Engine/UMG)
 
 
 %Steps%
