#UE4doc #docs:
Added engine version field
Removed all IconSections as per documentation guidelines

diff --git a/Engine/Documentation/Source/Engine/Blueprints/Overview/BlueprintOverview.INT.udn b/Engine/Documentation/Source/Engine/Blueprints/Overview/BlueprintOverview.INT.udn
--- a/Engine/Documentation/Source/Engine/Blueprints/Overview/BlueprintOverview.INT.udn
+++ b/Engine/Documentation/Source/Engine/Blueprints/Overview/BlueprintOverview.INT.udn
@@ -2,6 +2,7 @@
 Title: Blueprint Overview
 Crumbs: %ROOT%, Engine, Engine/Blueprints
 Description: The Blueprint Overview page breaks down the anatomy of a Blueprint and the different types of Blueprints available. 
+version: 4.9
 
 [VAR:Topic]
 [OBJECT:Topic]
@@ -45,12 +46,9 @@
 
 [TOC (start:2)]
 
-## Overview
-
 [INCLUDE:Engine/Blueprints/GettingStarted/#BlueprintIntro]
-
-**Does that mean _Blueprints_ are a replacement for UnrealScript?** Yes and no. Gameplay programming and everything that UnrealScript 
-was used for in the past can still be handled through code using C++. At the same time, while _Blueprints_ are not meant as a replacement 
+Does that mean Blueprints are a replacement for UnrealScript? Yes and no. Gameplay programming and everything that UnrealScript 
+was used for in the past can still be handled through code using C++. At the same time, while Blueprints are not meant as a replacement 
 for UnrealScript, they do serve many of the same purposes that UnrealScript handled, such as:
 
 * Extending classes
@@ -58,12 +56,12 @@
 * Managing subobject (e.g. Components) instancing for classes
 
 The expectation is that gameplay programmers will set up base classes which expose a useful set of functions and properties 
-that _Blueprints_ made from those base classes can use and extend upon. 
+that Blueprints made from those base classes can use and extend upon. 
 
 [EXCERPT:GameplayEquivalents]
 The table below provides a comparison of how various aspects would be handled in UnrealScript (from Unreal Engine 3), 
 C++, and Blueprints to help those transitioning from previous versions of the engine as well as show how native code 
-and _Blueprints_ compare.
+and Blueprints compare.
 
 | UnrealScript (UE3)  | Blueprints (UE4)         | C++ (UE4)          |
 | ------------------- | ------------------------ | ------------------ |
@@ -81,159 +79,65 @@
 
 ## Types of Blueprints
 
-_Blueprints_ can be one of several types that each have their own specific use from creating new types to scripting level 
-events to defining interfaces or macros to be used by other _Blueprints_.
+Blueprints can be one of several types that each have their own specific use from creating new types to scripting level 
+events to defining interfaces or macros to be used by other Blueprints.
 
-[EXCERPT:Types]
-[OBJECT:IconSection]
-	[PARAM:icon]
-		![Blueprints](%ROOT%/blueprint_icon.png)(convert:false)
-	[/PARAM]
-	[PARAM:title]
-		[](Engine/Blueprints/UserGuide/Types/ClassBlueprint)
-	[/PARAM]
-	[PARAM:description]
-		[INCLUDE:Engine/Blueprints/UserGuide/Types/ClassBlueprint#Intro]
-	[/PARAM]
-[/OBJECT]
-[OBJECT:IconSection]
-	[PARAM:icon]
-		![Blueprints](%ROOT%/blueprint_icon.png)(convert:false)
-	[/PARAM]
-	[PARAM:title]
-		Data-Only Blueprint
-	[/PARAM]
-	[PARAM:description]
-		[INCLUDE:Engine/Blueprints/UserGuide/Types/ClassBlueprint#DataBlueprint]
-	[/PARAM]
-[/OBJECT]
-[OBJECT:IconSection]
-	[PARAM:icon]
-		![Blueprints](%ROOT%/levelscript_icon.png)(convert:false)
-	[/PARAM]
-	[PARAM:title]
-		[](Engine/Blueprints/UserGuide/Types/LevelBlueprint)
-	[/PARAM]
-	[PARAM:description]
-		[INCLUDE:Engine/Blueprints/UserGuide/Types/LevelBlueprint#Intro]
-	[/PARAM]
-[/OBJECT]
-[OBJECT:IconSection]
-	[PARAM:icon]
-		![Blueprints](%ROOT%/interface_icon.png)(convert:false)
-	[/PARAM]
-	[PARAM:title]
-		[](Engine/Blueprints/UserGuide/Types/Interface)
-	[/PARAM]
-	[PARAM:description]
-		[INCLUDE:Engine/Blueprints/UserGuide/Types/Interface#Intro]
-	[/PARAM]
-[/OBJECT]
-[OBJECT:IconSection]
-	[PARAM:icon]
-		![Blueprints](%ROOT%/library_icon.png)(convert:false)
-	[/PARAM]
-	[PARAM:title]
-		[](Engine/Blueprints/UserGuide/Types/MacroLibrary)
-	[/PARAM]
-	[PARAM:description]
-		[INCLUDE:Engine/Blueprints/UserGuide/Types/MacroLibrary#Intro]
-	[/PARAM]
-[/OBJECT]
-[OBJECT:IconSection]
-	[PARAM:icon]
-		![Blueprints](%ROOT%/blueprint_icon.png)(convert:false)
-	[/PARAM]
-	[PARAM:title]
-		Blueprint Utilities
-	[/PARAM]
-	[PARAM:description]
-		A **Blueprint Utility**, a.k.a. a **Blutility**, is an editor-only _Blueprint_ that can be used to perform editor actions or extend 
-		editor functionality. These can expose [](Engine/Blueprints/UserGuide/Events) with no parameters as buttons in the UI 
-		and have the ability to execute any functions exposed to _Blueprints_ and act on the current set of selected Actors in the viewport.
-	[/PARAM]
-[/OBJECT]
-[/EXCERPT:Types]
 
-## Blueprint Anatomy
+###Blueprint Class
 
-The functionality of _Blueprints_ is defined using various elements; some of which are present by default, 
-while others can be added on an as-needed basis. These provide the ability to define _Components_, perform 
+[INCLUDE:Engine/Blueprints/UserGuide/Types/ClassBlueprint#Intro]
+
+
+###Data-Only Blueprint
+
+[INCLUDE:Engine/Blueprints/UserGuide/Types/ClassBlueprint#DataBlueprint]
+
+###Level Blueprint
+
+[INCLUDE:Engine/Blueprints/UserGuide/Types/LevelBlueprint#Intro]
+
+###Blueprint Interface
+
+[INCLUDE:Engine/Blueprints/UserGuide/Types/Interface#Intro]
+
+###Blueprint Macro Library
+
+[INCLUDE:Engine/Blueprints/UserGuide/Types/MacroLibrary#Intro]
+
+###Blueprint Utilities
+
+A **Blueprint Utility**, a.k.a. a **Blutility**, is an editor-only Blueprint that can be used to perform editor actions or extend 
+editor functionality. These can expose [](Engine/Blueprints/UserGuide/Events) with no parameters as buttons in the UI 
+and have the ability to execute any functions exposed to _Blueprints_ and act on the current set of selected Actors in the viewport.
+
+##Blueprint Anatomy
+
+The functionality of Blueprints is defined using various elements; some of which are present by default, 
+while others can be added on an as-needed basis. These provide the ability to define Components, perform 
 initialization and setup operations, respond to events, organize and modularize operations, define properties, 
 and more.
 
-
 [EXCERPT:Sections]
-[OBJECT:IconSection]
-	[PARAM:icon]
-		[![%Engine/Blueprints/UserGuide/Components:title%](Engine/Blueprints/UserGuide/Components/components_list.png)(w:256)](Engine/Blueprints/UserGuide/Components)
-	[/PARAM]
-	[PARAM:title]
-		[](Engine/Blueprints/UserGuide/Components)
-	[/PARAM]
-	[PARAM:path]
-		[RELATIVE:Engine/Blueprints/UserGuide/Components]
-	[/PARAM]
-	[PARAM:description]
-		[INCLUDE:Engine/Blueprints/UserGuide/Components#Intro]
-	[/PARAM]
-[/OBJECT]
-[OBJECT:IconSection]
-	[PARAM:icon]
-		[![%Engine/Blueprints/UserGuide/UserConstructionScript:title%](Engine/Blueprints/UserGuide/UserConstructionScript/ucs_thumb.png)(w:256)](Engine/Blueprints/UserGuide/UserConstructionScript)
-	[/PARAM]
-	[PARAM:title]
-		[](Engine/Blueprints/UserGuide/UserConstructionScript)
-	[/PARAM]
-	[PARAM:path]
-		[RELATIVE:Engine/Blueprints/UserGuide/UserConstructionScript]
-	[/PARAM]
-	[PARAM:description]
-		[INCLUDE:Engine/Blueprints/UserGuide/UserConstructionScript#Intro]
-	[/PARAM]
-[/OBJECT]
-[OBJECT:IconSection]
-	[PARAM:icon]
-		[![%Engine/Blueprints/UserGuide/EventGraph:title%](Engine/Blueprints/UserGuide/EventGraph/eventgraph_thumb.png)(w:256)](Engine/Blueprints/UserGuide/EventGraph)
-	[/PARAM]
-	[PARAM:title]
-		[](Engine/Blueprints/UserGuide/EventGraph)
-	[/PARAM]
-	[PARAM:path]
-		[RELATIVE:Engine/Blueprints/UserGuide/EventGraph]
-	[/PARAM]
-	[PARAM:description]
-		[INCLUDE:Engine/Blueprints/UserGuide/EventGraph#Intro]
-	[/PARAM]
-[/OBJECT]
-[OBJECT:IconSection]
-	[PARAM:icon]
-		[![%Engine/Blueprints/UserGuide/Functions:title%](Engine/Blueprints/UserGuide/Functions/functions_thumb.png)(w:256)](Engine/Blueprints/UserGuide/Functions)
-	[/PARAM]
-	[PARAM:title]
-		[](Engine/Blueprints/UserGuide/Functions)
-	[/PARAM]
-	[PARAM:path]
-		[RELATIVE:Engine/Blueprints/UserGuide/Functions]
-	[/PARAM]
-	[PARAM:description]
-		[INCLUDE:Engine/Blueprints/UserGuide/Functions#Intro]
-	[/PARAM]
-[/OBJECT]
-[OBJECT:IconSection]
-	[PARAM:icon]
-		[![%Engine/Blueprints/UserGuide/Variables:title%](Engine/Blueprints/UserGuide/Variables/Variables.png)(w:256)](Engine/Blueprints/UserGuide/Variables)
-	[/PARAM]
-	[PARAM:title]
-		[Variables](Engine/Blueprints/UserGuide/Variables)
-	[/PARAM]
-	[PARAM:path]
-		[RELATIVE:Engine/Blueprints/UserGuide/Variables]
-	[/PARAM]
-	[PARAM:description]
-		[INCLUDE:Engine/Blueprints/UserGuide/Variables#Intro]
-	[/PARAM]
-[/OBJECT]
+###Components Window
+
+[INCLUDE:Engine/Blueprints/UserGuide/Components#Intro]
+
+###Construction Script
+
+[INCLUDE:Engine/Blueprints/UserGuide/UserConstructionScript#Intro]
+
+###Event Graph
+
+[INCLUDE:Engine/Blueprints/UserGuide/EventGraph#Intro]
+
+###Functions
+
+[INCLUDE:Engine/Blueprints/UserGuide/Functions#Intro]
+
+###Variables
+
+[INCLUDE:Engine/Blueprints/UserGuide/Variables#Intro]
+
 [/EXCERPT:Sections]
 
 <!--
