#UE4doc #docs:
Added engine version field

diff --git a/Engine/Documentation/Source/Engine/Blueprints/QuickStart/BPQuickStart.INT.udn b/Engine/Documentation/Source/Engine/Blueprints/QuickStart/BPQuickStart.INT.udn
--- a/Engine/Documentation/Source/Engine/Blueprints/QuickStart/BPQuickStart.INT.udn
+++ b/Engine/Documentation/Source/Engine/Blueprints/QuickStart/BPQuickStart.INT.udn
@@ -2,6 +2,7 @@ Availability:Public
 Title: Blueprints Quick Start Guide
 Crumbs: %ROOT%, Engine, Engine/Blueprints
 Description: Creating Your First Blueprint
+version: 4.9
 
 [VAR:Topic]
 [OBJECT:Topic]
