#UE4doc #docs:
Added engine version field
Moved a link to a note region as the placement didn't explain what it was linking to.
Bolded a term

diff --git a/Engine/Documentation/Source/Engine/Blueprints/QuickStart/1/BPQS_1.INT.udn b/Engine/Documentation/Source/Engine/Blueprints/QuickStart/1/BPQS_1.INT.udn
--- a/Engine/Documentation/Source/Engine/Blueprints/QuickStart/1/BPQS_1.INT.udn
+++ b/Engine/Documentation/Source/Engine/Blueprints/QuickStart/1/BPQS_1.INT.udn
@@ -2,6 +2,7 @@ Availability: Public
 Title:1. Required Project Setup
 Crumbs:%ROOT%, Engine, Engine/Blueprints, Engine/Blueprints/QuickStart
 Description:Use the Side Scroller Blueprint Template as a jumping-off point for your game.
+version:4.9
 
 
 [VAR:Steps]
@@ -25,9 +26,13 @@ Description:Use the Side Scroller Blueprint Template as a jumping-off point for
 
 1. [Create a new project](Engine/Basics/Projects/Browser#creatingnewprojects).
 
-1. Select the Side Scroller template in the Blueprint tab of the New Project wizard.
+1. Select the **Side Scroller** template in the **Blueprint** tab of the **New Project** wizard.
 
-1. Choose whatever [scalability and quality settings work best for your setup, but make sure to include starter content](Engine/Basics/Projects/Browser#projectsettings).
+1. Choose whatever scalability and quality settings work best for your setup, but make sure to include **Starter Content**.
+
+[REGION:note]
+If you are unsure of what settings are right for you, you can find more information in the [Project Settings](Engine/Basics/Projects/Browser#projectsettings) section
+[/REGION]
 
 1. Give your project a name, and then create it by clicking the **Create Project** button. 
 
