#UE4doc #docs:
Added engine version field
Bolded a term

diff --git a/Engine/Documentation/Source/Engine/Blueprints/QuickStart/2/BPQS_2.INT.udn b/Engine/Documentation/Source/Engine/Blueprints/QuickStart/2/BPQS_2.INT.udn
--- a/Engine/Documentation/Source/Engine/Blueprints/QuickStart/2/BPQS_2.INT.udn
+++ b/Engine/Documentation/Source/Engine/Blueprints/QuickStart/2/BPQS_2.INT.udn
@@ -2,6 +2,7 @@ Availability: Public
 Title:2. Construct your Launchpad
 Crumbs:%ROOT%, Engine, Engine/Blueprints, Engine/Blueprints/QuickStart
 Description:Build your first launchpad in the level.
+version: 4.9
 
 
 [VAR:Steps]
@@ -58,7 +59,7 @@ If at any time you need to change your Actor's properties, such as to scale up t
 
 ![](LaunchPadActor.png)
 
-Now that you have your Actor the way you want it, we will turn it into a Blueprint Class. You can add more components inside the Blueprint Editor, and tweak them just like you can here in the level.
+Now that you have your Actor the way you want it, we will turn it into a Blueprint Class. You can add more components inside the **Blueprint Editor**, and tweak them just like you can here in the level.
 
 
 ---
