#UE4doc #docs:
Added engine version field
Bolded a term

diff --git a/Engine/Documentation/Source/Engine/Blueprints/QuickStart/3/QS_3.INT.udn b/Engine/Documentation/Source/Engine/Blueprints/QuickStart/3/QS_3.INT.udn
--- a/Engine/Documentation/Source/Engine/Blueprints/QuickStart/3/QS_3.INT.udn
+++ b/Engine/Documentation/Source/Engine/Blueprints/QuickStart/3/QS_3.INT.udn
@@ -2,6 +2,7 @@ Availability: Public
 Title:3. Convert Your Actor to a Blueprint Class
 Crumbs:%ROOT%, Engine, Engine/Blueprints, Engine/Blueprints/QuickStart
 Description:Make your launchpad reusable and scriptable by converting it to a Blueprint Class
+version: 4.9
 
 
 [VAR:Steps]
@@ -23,7 +24,7 @@ Description:Make your launchpad reusable and scriptable by converting it to a Bl
 %Steps%
 
 ---
-Making changes inside the Blueprint means that each time you make a new launchpad, it will have the look and feel you have created in the Blueprint Editor. While you could duplicate your LaunchPad
+Making changes inside the Blueprint means that each time you make a new launchpad, it will have the look and feel you have created in the **Blueprint Editor**. While you could duplicate your LaunchPad
 Actor around the level, any changes you make on a particular launchpad would only affect that one copy. 
 
 1. In the **Details** panel, click on the **Blueprint/Add Script** button.
