#UE4doc #docs:
Added engine version field
Replaced an image as the UI wasn't up to date
Unbolded and bolded a couple terms

diff --git a/Engine/Documentation/Source/Engine/Blueprints/QuickStart/4/BPQS_4.INT.udn b/Engine/Documentation/Source/Engine/Blueprints/QuickStart/4/BPQS_4.INT.udn
--- a/Engine/Documentation/Source/Engine/Blueprints/QuickStart/4/BPQS_4.INT.udn
+++ b/Engine/Documentation/Source/Engine/Blueprints/QuickStart/4/BPQS_4.INT.udn
@@ -2,6 +2,7 @@ Availability: Public
 Title:4. Create Your Starting Point
 Crumbs:%ROOT%, Engine, Engine/Blueprints, Engine/Blueprints/QuickStart
 Description:Your launchpad will launch the Character when you overlap it, so start with creating a Begin Overlap node.
+version: 4.9
 
 [VAR:Steps]
 [OBJECT:Navigation]
@@ -25,13 +26,13 @@ Description:Your launchpad will launch the Character when you overlap it, so sta
 
 To start adding behavior to your Blueprint Class, you will need to open it in the Blueprint Editor.  
 
-1. **Double-click** the **Blueprint Class** in the **Content Browser**.
+1. **Double-click** the Blueprint Class in the **Content Browser**.
 
-1. The Blueprint Editor will open, and you can see your **Cube** and **Box** components in a viewport. 
+1. The **Blueprint Editor** will open, and you can see your **Cube** and **Box** components in a viewport. 
 
-	[REGION:tip]
+	[REGION:note]
 	At this point, if you adjust the **Box** Component's placement, it will be applied to all launchpads you make from this Blueprint Class. Just like when you were working with components on your 
-	LaunchPad Actor, you can select the **Box** component in the **Components** list and adjust the location. Try a **Location** of (0.0, 0.0, 500.0).
+	LaunchPad Actor, you can select the **Box** component in the **Components** list and adjust the location. Try a **Location** of (0.0, 0.0, 50.0).
 
 	![](AdjustedTrigger.png)
 	[/REGION]
@@ -48,8 +49,8 @@ To start adding behavior to your Blueprint Class, you will need to open it in th
 
 1. **Right-click** in an empty spot in your graph to bring up the context menu of nodes you could add to the graph.
 
-	[REGION:tip]
-	To move around in the graph, **right-click** and drag around. At this point, you could drag the graph to the left, moving the pre-placed event nodes off the left side of the screen and creating more 
+	[REGION:note]
+	To move around in the graph, **Right-click** and drag around. At this point, you could drag the graph to the left, moving the pre-placed event nodes off the left side of the screen and creating more 
 	blank space to create your launchpad logic in.
 	[/REGION]
 
