#UE4doc #docs:
Added engine version field

diff --git a/Engine/Documentation/Source/Engine/Blueprints/QuickStart/5/BPQS_5.INT.udn b/Engine/Documentation/Source/Engine/Blueprints/QuickStart/5/BPQS_5.INT.udn
--- a/Engine/Documentation/Source/Engine/Blueprints/QuickStart/5/BPQS_5.INT.udn
+++ b/Engine/Documentation/Source/Engine/Blueprints/QuickStart/5/BPQS_5.INT.udn
@@ -2,6 +2,7 @@ Availability: Public
 Title:5. Test the Overlapping Actor
 Crumbs:%ROOT%, Engine, Engine/Blueprints, Engine/Blueprints/QuickStart
 Description:Make sure it is your Character overlapping the launchpad, so you only launch the Character.
+version: 4.9
 
 [VAR:Steps]
 [OBJECT:Navigation]
