#UE4doc #docs:
Added engine version field
Changed Tip region to Note due to the  nature of the content

diff --git a/Engine/Documentation/Source/Engine/Blueprints/QuickStart/6/BPQS_6.INT.udn b/Engine/Documentation/Source/Engine/Blueprints/QuickStart/6/BPQS_6.INT.udn
--- a/Engine/Documentation/Source/Engine/Blueprints/QuickStart/6/BPQS_6.INT.udn
+++ b/Engine/Documentation/Source/Engine/Blueprints/QuickStart/6/BPQS_6.INT.udn
@@ -2,6 +2,7 @@ Availability: Public
 Title:6. Launch Your Character
 Crumbs:%ROOT%, Engine, Engine/Blueprints, Engine/Blueprints/QuickStart
 Description:Set up the actual launching behavior with a Launch Character node.
+version: 4.9
 
 
 
@@ -44,7 +45,7 @@ and it will succeed. However, not everything in your level is the Pawn represent
 
 	![](LaunchWired.png)
 
-	[REGION:tip]
+	[REGION:note]
 	Notice that the output execution pin for successfully casting automatically connected to the input execution pin of **Launch Character**. 
 	[/REGION]
 
