#UE4doc #docs:
Added engine version field
Bolded a term
Fixed some capitalization issues

diff --git a/Engine/Documentation/Source/Engine/Blueprints/QuickStart/7/BPQS_7.INT.udn b/Engine/Documentation/Source/Engine/Blueprints/QuickStart/7/BPQS_7.INT.udn
--- a/Engine/Documentation/Source/Engine/Blueprints/QuickStart/7/BPQS_7.INT.udn
+++ b/Engine/Documentation/Source/Engine/Blueprints/QuickStart/7/BPQS_7.INT.udn
@@ -2,6 +2,7 @@ Availability: Public
 Title:7. On Your Own!
 Crumbs:%ROOT%, Engine, Engine/Blueprints, Engine/Blueprints/QuickStart
 Description:Add even more to your launchpad, and read more about Blueprints.
+version: 4.9
 
 
 
@@ -28,7 +29,7 @@ Using what you have learned over the course of this Quick Start guide, try to do
 
 * Play a sound when your Character is launched using an [Audio Component](Engine/Components/Audio).
 * Create a [variable](Engine/Blueprints/UserGuide/Variables) to store your Launch Velocity, and expose it so you can set it on each copy in the level.
-* Add a [Particle System Component](Engine/Components/Rendering#particlesystemcomponent) to your Blueprint and use one of the [particle systems](Engine/Rendering/ParticleSystems) from the Starter Content.
+* Add a [Particle System Component](Engine/Components/Rendering#particlesystemcomponent) to your Blueprint and use one of the [Particle Systems](Engine/Rendering/ParticleSystems) from the **Starter Content**.
 * Add an [Arrow Component](Engine/Components/Shapes) and use its rotation to define the direction to launch the character.
 * Using [Timelines](Engine/Blueprints/UserGuide/Timelines), add some animation to the Box Mesh to represent it launching the character.
 
