#UE4doc #docs:
Added engine version field
Removed unnecessary quotations
Fixed a capitalization issue

diff --git a/Engine/Documentation/Source/Engine/Blueprints/TechnicalGuide/BlueprintsTechnicalGuide.INT.udn b/Engine/Documentation/Source/Engine/Blueprints/TechnicalGuide/BlueprintsTechnicalGuide.INT.udn
--- a/Engine/Documentation/Source/Engine/Blueprints/TechnicalGuide/BlueprintsTechnicalGuide.INT.udn
+++ b/Engine/Documentation/Source/Engine/Blueprints/TechnicalGuide/BlueprintsTechnicalGuide.INT.udn
@@ -2,6 +2,7 @@ Availability:Public
 Title:Blueprints Technical Guide
 Crumbs: %ROOT%, Programming
 Description:Technical guide for programmers working with Blueprints.
+version: 4.9
 
 
 Blueprints are a powerful new feature introduced in Unreal Engine 4.  
@@ -10,7 +11,7 @@ In short, Blueprints are a way to create new UClasses without the need for writi
 class.  You can then add, arrange, and customize components, implement custom logic using a visual scripting language, respond to events and interactions, define custom variables, handle input, and 
 generally create a fully custom object type.
 
-Each Blueprint has a "construction script", analogous to a constructor in C++, which is run when the object is created.  This script can dynamically construct the Actor instance based on any number 
+Each Blueprint has a Construction Script, analogous to a constructor in C++, which is run when the object is created.  This script can dynamically construct the Actor instance based on any number 
 of factors, such as a fence that automatically sizes itself to fill a gap between buildings.  In this sense, a Blueprint can also be thought of as a very powerful prefab system. 
 
 
