#UE4doc #docs:
Added engine version field
Changed font settings for many terms
Fixed capitalization in some areas

diff --git a/Engine/Documentation/Source/Engine/Blueprints/TechnicalGuide/Compiler/BlueprintCompilerInternals.INT.udn b/Engine/Documentation/Source/Engine/Blueprints/TechnicalGuide/Compiler/BlueprintCompilerInternals.INT.udn
--- a/Engine/Documentation/Source/Engine/Blueprints/TechnicalGuide/Compiler/BlueprintCompilerInternals.INT.udn
+++ b/Engine/Documentation/Source/Engine/Blueprints/TechnicalGuide/Compiler/BlueprintCompilerInternals.INT.udn
@@ -2,12 +2,13 @@ Availability:Public
 Title:Blueprint Compiler Overview
 Crumbs: %ROOT%, Programming, Engine/Blueprints/TechnicalGuide
 Description:The steps of the Blueprint compilation process
+version: 4.9
 
 [TOC(start:2 end:2)]
 
 Blueprints, like regular C++ classes, need to be compiled before they can be used in-game. When you hit 
 the **Compile** button in the **Blueprint Editor**, the process of converting the properties and graphs of 
-a _Blueprint_ asset into a class is performed.
+a Blueprint asset into a class is performed.
 
 ## Terminology
 
@@ -43,61 +44,61 @@ The basic process of compiling a Blueprint is outlined below:
 * [REGION:listitem][Reinstance](#Re-instance)[/REGION]
 [/REGION]
 [REGION:caption clear]
-**\* Full Compile Only**
+*Full Compile Only
 [/REGION]
 
 ### Clean the Class
 
-Classes are _compiled in place_, which means the same `UBlueprintGeneratedClass` is cleaned and reused over 
-and over, so that pointers to the class do not have to be fixed up. `CleanAndSanitizeClass()` moves properties 
+Classes are compiled in place, which means the same **UBlueprintGeneratedClass** is cleaned and reused over 
+and over, so that pointers to the class do not have to be fixed up. **CleanAndSanitizeClass()** moves properties 
 and functions off the class and into a trash class in the transient package, and then clears any data on the class.
 
 ### Create Class Properties
 
-The compiler iterates over the Blueprint's `NewVariables` array, as well as some other places (construction scripts, etc.) 
+The compiler iterates over the Blueprint's **NewVariables** array, as well as some other places (construction scripts, etc.) 
 to find all of the UProperties needed by the class and then creates UProperties on the UClass's scope in the 
-function CreateClassVariablesFromBlueprint().
+function **CreateClassVariablesFromBlueprint()**.
 
 ### Create Function List
 
 The compiler creates the function list for the class by processing the event graphs, processing the regular function 
-graphs, and _pre-compiling_ the functions, i.e. calling `PrecompileFunction()` for each context.
+graphs, and _pre-compiling_ the functions, i.e. calling **PrecompileFunction()** for each context.
 
 #### Process Event Graphs
 
-Processing of the event graphs is performed by the `CreateAndProcessUberGraph()` function. This 
+Processing of the event graphs is performed by the **CreateAndProcessUberGraph()** function. This 
 copies all event graphs into one big graph, after which nodes are given a chance to expand. Then, a function stub is 
-created for each Event node in the graph, and an `FKismetFunctionContext` is created for each event graph.
+created for each Event node in the graph, and an **FKismetFunctionContext** is created for each event graph.
 
 #### Process Function Graphs
 
-Processing of the regular function graphs is done by the `ProcessOneFunctionGraph()` function, which duplicates 
-each graph to a temporary graph where nodes are given a chance to expand. A `FKismetFunctionContext` is created for 
+Processing of the regular function graphs is done by the **ProcessOneFunctionGraph()** function, which duplicates 
+each graph to a temporary graph where nodes are given a chance to expand. A **FKismetFunctionContext** is created for 
 each function graph as well.
 
 #### Pre-compile Functions
 
-Pre-compiling of the functions is handled by the `PrecompileFunction()` of each context. This function performs the 
+Pre-compiling of the functions is handled by the **PrecompileFunction()** of each context. This function performs the 
 following actions:
 
 * Schedules execution and calculates data dependencies.
 * Prunes any nodes that are unscheduled or not a data dependency.
-* Runs the node handler's `RegisterNets()` on each remaining node.
-* This creates the `FKismetTerms` for values within the function.
-* Creates the `UFunction` and associated properties.
+* Runs the node handler's **RegisterNets()** on each remaining node.
+* This creates the **FKismetTerms** for values within the function.
+* Creates the **UFunction** and associated properties.
 
 
 ### Bind and Link the Class
 
 Now that the compiler is aware of all of the UProperties and UFunctions for the class, it can bind and link the class, 
 which involves filling out the property chain, the property size, function map, etc. At this point, it essentially has 
-a _class header_ - minus the final flags and metadata - as well as a class default object (CDO).
+a class header - minus the final flags and metadata - as well as a Class Default Object (CDO).
 
 ### Compile Functions
 
-The next step consists of generating `FKismetCompiledStatment` objects for the remaining nodes which is 
-accomplished through the node handler's `Compile()` function, using `AppendStatementForNode()`. This 
-function can create `FKismetTerm` objects in the compile function as long as they are only used locally.
+The next step consists of generating **FKismetCompiledStatment** objects for the remaining nodes which is 
+accomplished through the node handler's **Compile()** function, using **AppendStatementForNode()**. This 
+function can create **FKismetTerm** objects in the compile function as long as they are only used locally.
 
 ### Finish Compiling Class
 
@@ -109,21 +110,21 @@ parent class before finally performing a few final checks to make sure everythin
 The backends convert the collection of statements from each function context into code. There are two backends 
 in use:
 
-* `FKismetCompilerVMBackend` - Converts FKCS to UnrealScript VM bytecode which are then serialized into the function's script array.
-* `FKismetCppBackend` - Emits _C++-like_ code for debugging purposes only.
+* **FKismetCompilerVMBackend** - Converts FKCS to UnrealScript VM bytecode which are then serialized into the function's script array.
+* **FKismetCppBackend** - Emits _C++-like_ code for debugging purposes only.
 
 ### Copy Class Default Object Properties
 
-Using a special function, `CopyPropertiesForUnrelatedObjects()`, the compiler copies the values from the old 
-class default object (CDO) of the class into the new CDO. Properties are copied via tagged serialization, so as 
+Using a special function, **CopyPropertiesForUnrelatedObjects()**, the compiler copies the values from the old 
+CDO of the class into the new CDO. Properties are copied via tagged serialization, so as 
 long as the names are consistent, they should properly be transferred. Components of the CDO are re-instanced and 
 fixed up appropriately at this stage. The GeneratedClass CDO is authoritative.
 
 ### Re-instance
 
 Since the class may have changed size and properties may have been added or removed, the compiler needs to re-instance 
-all objects with the class that were just compiled. This process uses a `TObjectIterator` to find all instances of the 
-class, spawn a new one, and then uses the `CopyPropertiesForUnrelatedObjects()` function to copy from the old instance 
+all objects with the class that were just compiled. This process uses a **TObjectIterator** to find all instances of the 
+class, spawn a new one, and then uses the **CopyPropertiesForUnrelatedObjects()** function to copy from the old instance 
 to the new one.
 
-For details, see the `FBlueprintCompileReinstancer` class.
+For details, see the **FBlueprintCompileReinstancer** class.
