#UE4doc #docs:
Added engine version field
Changed font settings for a few terms
Fixed some capitalization

diff --git a/Engine/Documentation/Source/Engine/Blueprints/TechnicalGuide/ExtendingBlueprints/ExposingGameplay.INT.udn b/Engine/Documentation/Source/Engine/Blueprints/TechnicalGuide/ExtendingBlueprints/ExposingGameplay.INT.udn
--- a/Engine/Documentation/Source/Engine/Blueprints/TechnicalGuide/ExtendingBlueprints/ExposingGameplay.INT.udn
+++ b/Engine/Documentation/Source/Engine/Blueprints/TechnicalGuide/ExtendingBlueprints/ExposingGameplay.INT.udn
@@ -20,7 +20,7 @@ creating an API exposed to Blueprints, see the [](Engine/Blueprints/TechnicalGui
 ## Making Classes Blueprintable
 
 In order to create Blueprints that extend from a class, that class must be defined as **Blueprintable**; which 
-involves adding this keyword inside the `UCLASS()` macro preceding the class definition. This keyword makes the 
+involves adding this keyword inside the **UCLASS()** macro preceding the class definition. This keyword makes the 
 Blueprint system aware of the class so that it shows in the class list of the **New Blueprint** dialog and can 
 be selected as the parent for the Blueprint being created.
 
@@ -43,8 +43,8 @@ The simplest form of declaration for a Blueprintable class is similar to the fol
 ## Readable and Writable Properties
 
 In order to expose a variable defined in a C++ class to a Blueprint extending from that class, that variable must 
-be defined using one of the keywords listed below inside the `UPROPERTY()` macro preceding the variable definition. 
-These keywords makes the Blueprint system aware of the variable so that it shows in the **MyBlueprint** panel and its value 
+be defined using one of the keywords listed below inside the **UPROPERTY()** macro preceding the variable definition. 
+These keywords makes the Blueprint system aware of the variable so that it shows in the **My Blueprint** panel and its value 
 can be set or accessed.
 
 	//Character's Health
@@ -62,7 +62,7 @@ can be set or accessed.
 ## Executable and Overridable Functions
 
 In order to call a native function from a Blueprint, that function must be defined using one of the keywords listed 
-below inside the `UFUNCTION()` macro preceding the function definition. These keywords make the Blueprint system 
+below inside the **UFUNCTION()** macro preceding the function definition. These keywords make the Blueprint system 
 aware of the function so that it shows in the context menu or Palette and can be added to a graph and executed - or, 
 in the case of events, so that they can be overridden and executed.
 
@@ -79,8 +79,8 @@ The simplest form of declaration for a BlueprintCallable Function is similar to
 | **BlueprintCallable** | This is a native function that can be called from a Blueprint, which executes native code that changes something about the object it is being called on, or some other global state.  This means it has to be "scheduled," or told explicitly the order in which it executes relative to other nodes.  We do this by the white execution line.  All Blueprint callable functions will be called in the order in which they appear along the white execution line. |
 | **BlueprintPure** | This is a native function that can be called from a Blueprint, which executes native code that does not change anything about the object it is being called on, or any other global state.  This means that nothing is changed by calling this node, it just takes input, and tells you an output.  These are things like math nodes (+, - , *, etc), or variable getters, or anything that does not change anything permanently.  These do not  need to be scheduled, and do not have a connection for the white execution line.  They are automatically figured out by the compiler, based on which BlueprintCallable nodes need the data produced by these nodes. |
 | [REGION:tablesection]Native to Blueprint Communication[/REGION] ||
-| **BlueprintImplementableEvent** | This is the primary way we allow native functions to call up into Blueprints.  They are like virtual functions that you implement in the Blueprint themselves.  If there is no implementation, the function call is ignored.  It is important to note that if a BlueprintImplementableEvent does not  have a return value or an out parameter, it will appear as an event, available by **right-clicking** and selecting it in the event graph of the Blueprint.  If it does have a return value or any out parameters, it will be listed in the **My Blueprints** tab, and can then be overridden by **right-clicking** and selecting "implement" function.  Note that BlueprintImplementableEvents do not  have a native implementation of the function. |
-| **BlueprintNativeEvent** | These are the same as above, except there is a native default implementation of the function that is called if the Blueprint does not override the function.  This is useful for things where you want some sort of default behavior if the Blueprint does not implement it, but want the Blueprint to be able to override the functionality if desired.  These are more costly, so we only put them in where the functionality is needed.  When you override the BlueprintNativeEvent, you can still call the native implementation, if desired, by **right-clicking** on the event or function entry node, and selecting "Add call to parent". |
+| **BlueprintImplementableEvent** | This is the primary way we allow native functions to call up into Blueprints.  They are like virtual functions that you implement in the Blueprint themselves.  If there is no implementation, the function call is ignored.  It is important to note that if a BlueprintImplementableEvent does not  have a return value or an out parameter, it will appear as an event, available by **Right-clicking** and selecting it in the event graph of the Blueprint.  If it does have a return value or any out parameters, it will be listed in the **My Blueprints** tab, and can then be overridden by **Right-clicking** and selecting implement function.  Note that BlueprintImplementableEvents do not  have a native implementation of the function. |
+| **BlueprintNativeEvent** | These are the same as above, except there is a native default implementation of the function that is called if the Blueprint does not override the function.  This is useful for things where you want some sort of default behavior if the Blueprint does not implement it, but want the Blueprint to be able to override the functionality if desired.  These are more costly, so we only put them in where the functionality is needed.  When you override the BlueprintNativeEvent, you can still call the native implementation, if desired, by **Right-clicking** on the event or function entry node, and selecting "Add call to parent". |
 
 
 
