#UE4doc #docs:
Added engine version field

diff --git a/Engine/Documentation/Source/Engine/Blueprints/TechnicalGuide/Guidelines/Guidelines.INT.udn b/Engine/Documentation/Source/Engine/Blueprints/TechnicalGuide/Guidelines/Guidelines.INT.udn
--- a/Engine/Documentation/Source/Engine/Blueprints/TechnicalGuide/Guidelines/Guidelines.INT.udn
+++ b/Engine/Documentation/Source/Engine/Blueprints/TechnicalGuide/Guidelines/Guidelines.INT.udn
@@ -3,6 +3,7 @@ Title:Guidelines for Programming for Blueprints
 Crumbs: %ROOT%, Programming, Engine/Blueprints/TechnicalGuide
 Description:Tips and tricks for how best to write a Blueprint-friendly API
 Related:Engine/Blueprints/TechnicalGuide/ExtendingBlueprints 
+version: 4.9
 
 [TOC (start:2 end:2)]
 
