#UE4doc #docs:
Added engine version field
Removed Overview header
Changed font settings for a few terms

diff --git a/Engine/Documentation/Source/Engine/Blueprints/UserGuide/KismetNodeReference.INT.udn b/Engine/Documentation/Source/Engine/Blueprints/UserGuide/KismetNodeReference.INT.udn
--- a/Engine/Documentation/Source/Engine/Blueprints/UserGuide/KismetNodeReference.INT.udn
+++ b/Engine/Documentation/Source/Engine/Blueprints/UserGuide/KismetNodeReference.INT.udn
@@ -2,13 +2,11 @@ Availability:Public
 Title: Blueprint User Guide
 Crumbs: %ROOT%, Engine, Engine/Blueprints
 Description: The User Guide is the go-to source to learn the different parts of Blueprints and nodes that are available to use within Blueprint graphs.
-
-
-## Overview
+version: 4.9
 
 [INCLUDE:Engine\Blueprints#intro]
 
-This page outlines the various types of Blueprints and Graph Types, Blueprint Elements (such as _Variables_, _Components_, and _Graphs_), Communication Methods and how to pass information between Blueprints, as well as a breakdown of the different types of Graph Nodes you will be working with as you create your Blueprints. 
+This page outlines the various types of Blueprints and Graph Types, Blueprint Elements (such as Variables, Components, and Graphs), Communication Methods and how to pass information between Blueprints, as well as a breakdown of the different types of Graph Nodes you will be working with as you create your Blueprints. 
 
 [OBJECT:TopicButtonList]
 	[PARAM:icon]
@@ -38,7 +36,7 @@ This page outlines the various types of Blueprints and Graph Types, Blueprint El
 
 [OBJECT:TopicButtonList]
 	[PARAM:icon]
-		![](%ROOT%/Editor_icon.png)(convert:false)
+		![](%ROOT%/editor_icon.png)(convert:false)
 	[/PARAM]
 	[PARAM:title]
 		Blueprint Graph Types
@@ -121,7 +119,7 @@ This page outlines the various types of Blueprints and Graph Types, Blueprint El
 -->
 [OBJECT:TopicButtonList]
 	[PARAM:icon]
-		![](%ROOT%/Blueprint_icon.png)(convert:false)
+		![](%ROOT%/blueprint_icon.png)(convert:false)
 	[/PARAM]
 	[PARAM:title]
 		Graph Nodes
