#UE4doc #docs:
Added engine version field
Bolded a term
Fixed capitalization issues
Changed "zero" to "0" since it ompared it to "1" instead of "one"

diff --git a/Engine/Documentation/Source/Engine/Blueprints/UserGuide/Arrays/Arrays.INT.udn b/Engine/Documentation/Source/Engine/Blueprints/UserGuide/Arrays/Arrays.INT.udn
--- a/Engine/Documentation/Source/Engine/Blueprints/UserGuide/Arrays/Arrays.INT.udn
+++ b/Engine/Documentation/Source/Engine/Blueprints/UserGuide/Arrays/Arrays.INT.udn
@@ -1,15 +1,16 @@
 Availability:Public
 Title:Blueprint Arrays
 Crumbs: %ROOT%, Engine/Blueprints, Engine/Blueprints/UserGuide
-Description:Documentation over Blueprint arrays, including creation, editing, and use.
+Description:Documentation over Blueprint Arrays, including creation, editing, and use.
+version: 4.9
 
 ![](Arrays.png)
 
-Just as with variable values, Blueprints can store data within arrays. If you are not versed in programming terms, you can think of an array as a collection of variables all housed within a single unit. 
+Just as with variable values, Blueprints can store data within **Arrays**. If you are not versed in programming terms, you can think of an Array as a collection of variables all housed within a single unit. 
 
-An array can only hold values of a single type. For instance, a Boolean array can only hold Boolean values.
+An Array can only hold values of a single type. For instance, a Boolean Array can only hold Boolean values.
 
-Array variables will contain a 3x3 colored grid to show that they are arrays and not regular variables. In the case of a disconnected array, that grid will have a darkened center. Once connected, the entire grid will be visible, as seen below.
+Array variables will contain a 3x3 colored grid to show that they are Arrays and not regular variables. In the case of a disconnected Array, that grid will have a darkened center. Once connected, the entire grid will be visible, as seen below.
 
 ![](GridIcon.png)
 
@@ -17,40 +18,40 @@ Array variables will contain a 3x3 colored grid to show that they are arrays and
 
 ## Creating Arrays
 
-Creating an array is as simple as clicking on the Array grid icon when creating a variable. 
+Creating an Array is as simple as clicking on the Array grid icon when creating a variable. 
 
 ![](example_variable_array.png)
 
-Once this icon is selected, the new item will be an array rather than a standard variable.
+Once this icon is selected, the new item will be an Array rather than a standard variable.
 
 ![](array_selected.png)
 
 ## Editing Arrays
 
-You can edit the value of an array either via the Blueprint Default settings, or at any point along a Blueprint node network. These networks can be in the Construction Script, a function, a macro, or an Event Graph.
+You can edit the value of an Array either via the Blueprint Default settings, or at any point along a Blueprint node network. These networks can be in the Construction Script, a function, a macro, or an Event Graph.
 
 ### Array Defaults
 
-Setting the default values for an array is very easy. Simply create the necessary array and enter the **Class Defaults** tab or the **Defaults** mode of the Blueprint Editor. You will see a section named for your array, as shown below:
+Setting the default values for an Array is very easy. Simply create the necessary Array and enter the **Class Defaults** tab or the **Defaults** mode of the Blueprint Editor. You will see a section named for your Array, as shown below:
 
 ![](StringArrayDefaults.png)
 
-[REGION:warning]
-	If you do not see your array listed in the Class Defaults, make sure that you have compiled your Blueprint since the array was created.
+[REGION:note]
+	If you do not see your Array listed in the Class Defaults, make sure that you have compiled your Blueprint since the Array was created.
 [/REGION]
 
-To edit the values of an array's defaults:
+To edit the values of an Array's defaults:
 
-1. In the **Class Defaults** tab, click the ![](button_plus.png) button. This will create a new index. Repeat for as many items as you would like the array to contain.
+1. In the **Class Defaults** tab, click the ![](button_plus.png) button. This will create a new index. Repeat for as many items as you would like the Array to contain.
 
 	![](StringArray_IndicesAdded.png)	
 	In this example, 3 indices have been added.
 
-1. Set each value accordingly. Note that the way in which you set the value will depend on the type of array in use. 
+1. Set each value accordingly. Note that the way in which you set the value will depend on the type of Array in use. 
 
 	![](StringArrayDefaultsSet.png)
 
-1. To insert, remove, or duplicate an array index, click the ![](button_ellipse.png) button next to the entry to bring up the edit menu:
+1. To insert, remove, or duplicate an Array index, click the ![](button_ellipse.png) button next to the entry to bring up the edit menu:
 
 	![](StringArrayEdit.png)
 
@@ -61,15 +62,15 @@ To edit the values of an array's defaults:
 
 ### Setting Array Values via Node Networks
 
-If your array is intended to be populated at runtime, you will generally not use defaults. Rather, you will use nodes within the Construction Script or an Event Graph to populate each index. For instance, you could use the **Add** or **Insert** nodes to add a new value to the next available index, or to insert a value at a given index, respectively.
+If your Array is intended to be populated at runtime, you will generally not use defaults. Rather, you will use nodes within the Construction Script or an Event Graph to populate each index. For instance, you could use the **Add** or **Insert** nodes to add a new value to the next available index, or to insert a value at a given index, respectively.
 
 ![](AddArray.png)
 
-For a complete list of the array control nodes within Blueprints, please see the [](Engine/Blueprints/UserGuide/Arrays/ArrayNodes) page.
+For a complete list of the Array control nodes within Blueprints, please see the [](Engine/Blueprints/UserGuide/Arrays/ArrayNodes) page.
 
 ## Array Indices 
 
-It is also important to note that values within an array are given **zero-based** index values, meaning that the first index is **zero**, not 1. For example, if you were to add five items to an array, the list would, internally, look something like this:
+It is also important to note that values within an Array are given **zero-based** index values, meaning that the first index is 0, not 1. For example, if you were to add five items to an Array, the list would, internally, look something like this:
 
 
 [REGION:simpletable]
