#UE4doc #docs:
Added engine version field
Changed font settings for a few terms

diff --git a/Engine/Documentation/Source/Engine/Blueprints/UserGuide/Arrays/ArrayNodes/ArrayNodes.INT.udn b/Engine/Documentation/Source/Engine/Blueprints/UserGuide/Arrays/ArrayNodes/ArrayNodes.INT.udn
--- a/Engine/Documentation/Source/Engine/Blueprints/UserGuide/Arrays/ArrayNodes/ArrayNodes.INT.udn
+++ b/Engine/Documentation/Source/Engine/Blueprints/UserGuide/Arrays/ArrayNodes/ArrayNodes.INT.udn
@@ -2,6 +2,7 @@ Availability:Public
 Title:Array Nodes
 Crumbs: %ROOT%, Engine/Blueprints, Engine/Blueprints/UserGuide, Engine/Blueprints/UserGuide/Arrays
 Description: Nodes designed to aid in working with Blueprint arrays.
+version: 4.9
 
 [VAR:Topic]
 [OBJECT:Topic]
@@ -58,13 +59,13 @@ The **Clear** node will wipe out all data within a connected array, resetting th
 
 ## Contains Item
 
-The **Contains Item** node allows for an array to be searched for a particular item. If the item is found, the node will output _True_, if not, _False_. It does not return the index at which the item was found.
+The **Contains Item** node allows for an array to be searched for a particular item. If the item is found, the node will output True, if not, False. It does not return the index at which the item was found.
 
 ![](Array_Contains.png)
 
 ## Filter Array
 
-The **Filter Array** node takes in an array of type _Actor_ (or any derivative thereof) and filters it based on an incoming _Class_. The result is a new array, only containing those items from the original array that matched the type of the connected class. 
+The **Filter Array** node takes in an array of type Actor (or any derivative thereof) and filters it based on an incoming Class. The result is a new array, only containing those items from the original array that matched the type of the connected class. 
 
 ![](Array_Filter.png)
 
@@ -82,7 +83,7 @@ The **Get** node takes in an array and an integer, which it uses as an index num
 
 ## Insert Item
 
-Insert Item takes in an array, a variable value, and an integer (index number). It then inserts the value of the variable into the array at the index value location.
+**Insert Item** takes in an array, a variable value, and an integer (index number). It then inserts the value of the variable into the array at the index value location.
 
 ![](Array_Insert.png)
 
@@ -106,7 +107,7 @@ The **Remove Index** node removes an item within the array via that item's index
 
 ## Remove Item
 
-**Remove Item** will remove an item from the array by taking in a variable value. If that value is found within the array, it is removed. All remaining index numbers will update to close the gap. This node has a Boolean output which returns _True_ if the item was found and removed.
+**Remove Item** will remove an item from the array by taking in a variable value. If that value is found within the array, it is removed. All remaining index numbers will update to close the gap. This node has a Boolean output which returns True if the item was found and removed.
 
 ![](Array_RemoveItem.png)
 
