#UE4doc #docs:
Added engine version field
Replaced a couple images due to outdated UI
Removed unnecessary quotations

diff --git a/Engine/Documentation/Source/Engine/Blueprints/UserGuide/BlueprintComms/BlueprintComms.INT.udn b/Engine/Documentation/Source/Engine/Blueprints/UserGuide/BlueprintComms/BlueprintComms.INT.udn
--- a/Engine/Documentation/Source/Engine/Blueprints/UserGuide/BlueprintComms/BlueprintComms.INT.udn
+++ b/Engine/Documentation/Source/Engine/Blueprints/UserGuide/BlueprintComms/BlueprintComms.INT.udn
@@ -5,6 +5,7 @@ Description:How to use one Blueprint's functions, events, and variables from ano
 related:Engine/Blueprints/UserGuide/BlueprintCommsUsage
 related:Engine/Blueprints/UserGuide/Types/Interface/UsingInterfaces
 related:Engine/Blueprints/UserGuide/EventDispatcher
+version: 4.9
 
 
 [VAR:Topic]
@@ -49,15 +50,15 @@ related:Engine/Blueprints/UserGuide/EventDispatcher
 
 [TOC(start:2)]
 
-When you are working with a Blueprint, you may sometimes find that you want to access the functions, events, and variables that are contained in another Blueprint. The simplest way to do so is to reference the "target" Blueprint through a publicly exposed object variable, and then specify which Blueprint instance you want to use.
+When you are working with a Blueprint, you may sometimes find that you want to access the functions, events, and variables that are contained in another Blueprint. The simplest way to do so is to reference the target Blueprint through a publicly exposed object variable, and then specify which Blueprint instance you want to use.
 
 ## Simple Way to Reference a Target Blueprint
 
-The following sections give an example of the simplest and most common way to reference a target Blueprint from the Blueprint you are working with (the "working" Blueprint).
+The following sections give an example of the simplest and most common way to reference a target Blueprint from the Blueprint you are working with (the working Blueprint).
 
 ### Creating an Object Variable for the Target Blueprint
 
-When you are working in a Blueprint, creating an object variable that references a target Blueprint gives you access to its functions, events, and variables.
+Creating an object variable that references a target Blueprint gives you access to its functions, events, and variables.
 
 **To create an object variable referencing a target Blueprint**:
 
@@ -67,13 +68,13 @@ When you are working in a Blueprint, creating an object variable that references
 
 1. Give your new variable a name related to the target Blueprint, such as "TargetBlueprint."
 
-1. In the **Details** panel, click the drop-down box next to **Variable Type**. Expand the **Object Reference** section.
+1. In the **Details** panel, click the drop-down box next to **Variable Type**. Expand the **Object Types** section.
 
 	[REGION:raw]
 	![Variable Type](BP_VarType.png)
 	[/REGION]
 
-1. Scroll through the list and select the name of the Blueprint you want to reference.
+1. Scroll through the list and select the name of the Blueprint you want to reference then the **Reference** option.
 
 	![Blueprint Reference](BP_Blueprint2Ref.png)
 
@@ -119,7 +120,7 @@ If you are not sure which instance in the list corresponds to the one you want t
 
 ### Using Functions, Events, and Variables from the Target Blueprint
 
-After you have created an object variable for the target Blueprint, you can get a reference to it by holding down **Ctrl** and dragging it onto the graph, which creates a variable "getter" node for it. If you then drag out from the node's output pin, a context menu appears, displaying the list of events, functions, and variables that you can use from the target Blueprint. You can now use these elements as you would any other scripting elements.
+After you have created an object variable for the target Blueprint, you can get a reference to it by holding down **Ctrl** and dragging it onto the graph, which creates a variable Getter node for it. If you then drag out from the node's output pin, a context menu appears, displaying the list of events, functions, and variables that you can use from the target Blueprint. You can now use these elements as you would any other scripting elements.
 
 ![Context Menu](BP_Target_ContextMenu.png)
 
