#UE4doc #docs - Added Engine Version field. Added Skill Level Field. Updated two images. Replaced image table with comparison slider.

diff --git a/Engine/Documentation/Source/Engine/Rendering/PostProcessEffects/PostProcessMaterials/PostProcessDemoMaterialOverview/PostProcessDemoMaterialOverview.INT.udn b/Engine/Documentation/Source/Engine/Rendering/PostProcessEffects/PostProcessMaterials/PostProcessDemoMaterialOverview/PostProcessDemoMaterialOverview.INT.udn
--- a/Engine/Documentation/Source/Engine/Rendering/PostProcessEffects/PostProcessMaterials/PostProcessDemoMaterialOverview/PostProcessDemoMaterialOverview.INT.udn
+++ b/Engine/Documentation/Source/Engine/Rendering/PostProcessEffects/PostProcessMaterials/PostProcessDemoMaterialOverview/PostProcessDemoMaterialOverview.INT.udn
@@ -2,6 +2,8 @@ Availability:Public
 Title:Post Process Material Example
 Crumbs: %ROOT%, Engine, Engine/Rendering/PostProcessEffects, Engine/Rendering/PostProcessEffects/PostProcessMaterials
 Description:An overview of a video-scanline post process material
+Version: 4.9
+SkillLevel: Advanced
 
 [VAR:Topic]
 [OBJECT:Topic]
@@ -55,7 +57,7 @@ w:450
 
 Below are some basic instructions on how to set up a simple Post Process Material:
 
-1. Create a new Level by choosing **File > New Level** from the Main Menu Bar.
+1. Create a new Level by choosing **File** > **New Level** from the Main Menu Bar.
 
 1. Create a new Material by clicking the ![](CB_button_NewAsset.png) button in the **Content Browser** and choosing **Material**. Provide a name for the Material.
 
@@ -67,11 +69,11 @@ Below are some basic instructions on how to set up a simple Post Process Materia
 
 	![](PostMaterialProperties.png)
 
-1. In the **Post Process Material** category, be sure to set the **Blendable Location** property to _Before Tonemapping_. This makes for more of a performance hit, but prevents ghosting of the scan lines while the view is in motion. For more information about the Blendable Location property, see the [Critical Settings for Post Process Materials](Engine/Rendering\PostProcessEffects\PostProcessMaterials#CriticalSettingsforPostProcessMaterials) section below. 
+1. In the **Post Process Material** category, be sure to set the **Blendable Location** property to _Before Tonemapping_. This makes for more of a performance hit, but prevents ghosting of the scan lines while the view is in motion. For more information about the Blendable Location property, see the [**Critical Settings for Post Process Materials**](Engine/Rendering\PostProcessEffects\PostProcessMaterials#CriticalSettingsforPostProcessMaterials) section below. 
 
 	![](PostProcessMaterialsCategory.png)
 
-1. At this point, you would create some Material expression network that defined the look of your post process. As a test for this example, we will create a video scan line overlay effect. You can really make just about anything you like, however. If you would like an example of a cool video scan line post process material, see the [Material Setup](#MaterialSetup) section.
+1. At this point, you would create some Material expression network that defined the look of your post process. As a test for this example, we will create a video scan line overlay effect. You can really make just about anything you like, however. If you would like an example of a cool video scan line post process material, see the [**Material Setup**](#MaterialSetup) section.
 
 
 ## Material Setup
@@ -111,7 +113,7 @@ To facilitate putting this Material together, we have broken the key elements in
 1. Create a time-driven sine wave that runs between 0.6 and 1.0, at a frequency of 10 Hz. This will be used to flicker the scanlines.
 1. Take the noise texture and tile and pan it at 2 different amounts and speeds for variation.
 1. Make an average of the two noises, then multiply it by the ScanDistort texture. Multiply _that_ by 40 and add it into the result to create white static.
-1. We take the panning ScanDistort texture, raise it to a power of 4 (this boosts contrast), multiply it by 40 to make it very bright, but then _subtract_ that from the final result to push it to full black. This creates the black bar effect often seen when tracking an old-school [VCR](http://en.wikipedia.org/wiki/Vcr).
+1. We take the panning ScanDistort texture, raise it to a power of 4 (this boosts contrast), multiply it by 40 to make it very bright, but then _subtract_ that from the final result to push it to full black. This creates the black bar effect often seen when tracking an old-school [**VCR**](http://en.wikipedia.org/wiki/Vcr).
 
 [REGION:silver]If you actually needed to click the VCR link above, then we do not expect you to really appreciate this effect.[/REGION] 
 
@@ -128,12 +130,14 @@ To facilitate putting this Material together, we have broken the key elements in
 
 1. Select your new Material in the **Content Browser** and click the ![](button_CBApply.png) button to apply the Material into the element. Your effect is now applied.
 
-
-[REGION:imagetable]
-|![](BeforePost.png)(%PostTest%) | ![](AfterPost.png)(%PostTest%)|
-|:---:|:---:|
-| Before post process material applied | After post process material applied |
-[/REGION]
+[OBJECT:ComparisonSlider]
+	[PARAM:before]
+	![Before post process material applied](BeforePost.png)(w:600)
+	[/PARAM]
+	[PARAM:after]
+	![After post process material applied](AfterPost.png)(w:600)
+	[/PARAM]
+[/OBJECT]
 
 Now you can add in some of the property-based post process settings that are a part of the Post Process Volume. In this case, we did the following:
 	
