#UE4doc #docs - Added ENgine Version Field. Replaced the image table with comparison slider.

diff --git a/Engine/Documentation/Source/Engine/Rendering/PostProcessEffects/SceneFringe/SceneFringe.INT.udn b/Engine/Documentation/Source/Engine/Rendering/PostProcessEffects/SceneFringe/SceneFringe.INT.udn
--- a/Engine/Documentation/Source/Engine/Rendering/PostProcessEffects/SceneFringe/SceneFringe.INT.udn
+++ b/Engine/Documentation/Source/Engine/Rendering/PostProcessEffects/SceneFringe/SceneFringe.INT.udn
@@ -2,16 +2,21 @@ Availability:Public
 Title:Scene Fringe
 Crumbs: %ROOT%, Engine, Engine/Rendering/PostProcessEffects
 Description:Chromatic aberration effect that simulates the color shifts near the edges of real-world camera lenses.
+Version: 4.9
+
 
 [EXCERPT:SceneFringe_Desc]
 **Scene Fringe** is a chromatic aberration effect that simulates the color shifts in real-world camera lenses. The effect is mostly noticeable near the edges of the image.
 [/EXCERPT:SceneFringe_Desc]
 
-[REGION:imagetable]
-| ![Without Scene Fringe](scenefringe_0.png)(w:455) | ![With Scene Fringe](scenefringe_2.png)(w:455) |
-| :----: | :----: |
-| Without Scene Fringe | With Scene Fringe (extreme settings) |
-[/REGION]
+[OBJECT:ComparisonSlider]
+	[PARAM:before]
+	![Without Scene Frigne](scenefringe_0.png)(w:700)
+	[/PARAM]
+	[PARAM:after]
+	![With Scene Fringe](scenefringe_2.png)(w:700)
+	[/PARAM]
+[/OBJECT]
 
 
 | Property | Description |
