#UE4doc #docs - Added Engine Version Field.

diff --git a/Engine/Documentation/Source/Engine/Rendering/PostProcessEffects/ScreenSpaceReflection/ScreenSpaceReflection.INT.udn b/Engine/Documentation/Source/Engine/Rendering/PostProcessEffects/ScreenSpaceReflection/ScreenSpaceReflection.INT.udn
--- a/Engine/Documentation/Source/Engine/Rendering/PostProcessEffects/ScreenSpaceReflection/ScreenSpaceReflection.INT.udn
+++ b/Engine/Documentation/Source/Engine/Rendering/PostProcessEffects/ScreenSpaceReflection/ScreenSpaceReflection.INT.udn
@@ -2,6 +2,7 @@ Availability:Public
 Title:Screen Space Reflection
 Crumbs: %ROOT%, Engine, Engine/Rendering/PostProcessEffects
 Description:Effect which alters the reflection that appear on the surface of materials.
+Version: 4.9
 
 [REGION:fullwidth]
 ![](GlossyScreenSpaceReflection.png)
