#UE4doc #docs - Added Engine Version Field. Replaced image table with Comparison Slider.

diff --git a/Engine/Documentation/Source/Engine/Rendering/PostProcessEffects/Vignette/Vignette.INT.udn b/Engine/Documentation/Source/Engine/Rendering/PostProcessEffects/Vignette/Vignette.INT.udn
--- a/Engine/Documentation/Source/Engine/Rendering/PostProcessEffects/Vignette/Vignette.INT.udn
+++ b/Engine/Documentation/Source/Engine/Rendering/PostProcessEffects/Vignette/Vignette.INT.udn
@@ -2,6 +2,7 @@ Availability:Public
 Title:Vignette
 Crumbs: %ROOT%, Engine, Engine/Rendering/PostProcessEffects
 Description:Effect causing the brightness of the to decrease as the distance from the viewport's center increases.
+Version: 4.9
 
 [REGION:fullwidth]
 ![Vignette Effect](vignette.png)
@@ -17,10 +18,13 @@ Description:Effect causing the brightness of the to decrease as the distance fro
 
 <!--
 [EXCERPT:Intensity]
-[REGION:imagetable]
-| ![Vignette Intensity - 0](vignette_off.png)(w:375) | ![Vignette Intensity - 1](vignette_on.png)(w:375) |
-| ------ | ------ |
-| 0.0 | 1.0 |
-[/REGION]
+[OBJECT:ComparisonSlider]
+	[PARAM:before]
+	![Vignetter Intensity - 0](vignette_off.png)(w:600)
+	[/PARAM]
+	[PARAM:after]
+	![Vignette Intensity - 1](vignette_on.png)(w:600)
+	[/PARAM]
+[/OBJECT]
 [/EXCERPT:Intensity]
 -->
