#UE4doc #docs - Added Engine Version Field.

diff --git a/Engine/Documentation/Source/Engine/UI/Editor.INT.udn b/Engine/Documentation/Source/Engine/UI/Editor.INT.udn
--- a/Engine/Documentation/Source/Engine/UI/Editor.INT.udn
+++ b/Engine/Documentation/Source/Engine/UI/Editor.INT.udn
@@ -3,6 +3,7 @@ Title: Unreal Editor Interface
 Crumbs:%ROOT%, Engine
 Description:A breakdown of the Unreal Editor user interface, along with key principles on how the editor works.
 Navigation:topic
+Version: 4.9
 
 [EXCERPT:LevelEditor]
 
