#UE4doc #docs - Added Engine Version Field.

diff --git a/Engine/Documentation/Source/Engine/UI/ClassViewer/ClassViewer.INT.udn b/Engine/Documentation/Source/Engine/UI/ClassViewer/ClassViewer.INT.udn
--- a/Engine/Documentation/Source/Engine/UI/ClassViewer/ClassViewer.INT.udn
+++ b/Engine/Documentation/Source/Engine/UI/ClassViewer/ClassViewer.INT.udn
@@ -2,6 +2,7 @@ Availability:Public
 Title:Class Viewer
 Crumbs: %ROOT%, Engine, Engine/UI
 Description:Tool for examining UE4 classes.
+Version: 4.9
 
 [VAR:Topic]
 [OBJECT:Topic]
