#UE4doc #docs - Added Engine Version field.

diff --git a/Engine/Documentation/Source/Engine/UI/CurveEditor/CurveEditor.INT.udn b/Engine/Documentation/Source/Engine/UI/CurveEditor/CurveEditor.INT.udn
--- a/Engine/Documentation/Source/Engine/UI/CurveEditor/CurveEditor.INT.udn
+++ b/Engine/Documentation/Source/Engine/UI/CurveEditor/CurveEditor.INT.udn
@@ -2,6 +2,7 @@ Availability: Public
 Title:Curve Editor
 Crumbs:%ROOT%, Engine, Engine/UI
 Description:Tool for interactively setting up and modifying curves to drive values of Distribution properties.
+Version: 4.9
 
 [VAR:Topic]
 [OBJECT:Topic]
