#UE4doc #docs - Added Engine Version Field. Updated Engine Preferences Image to match latest version.

diff --git a/Engine/Documentation/Source/Engine/UI/EditorSettings/EditorPreferences.INT.udn b/Engine/Documentation/Source/Engine/UI/EditorSettings/EditorPreferences.INT.udn
--- a/Engine/Documentation/Source/Engine/UI/EditorSettings/EditorPreferences.INT.udn
+++ b/Engine/Documentation/Source/Engine/UI/EditorSettings/EditorPreferences.INT.udn
@@ -2,6 +2,7 @@ Availability:Public
 Title:Editor Preferences
 Crumbs:%ROOT%, Engine, Engine/UI
 Description:Settings for configuring general editor behavior for controls, viewports, source control, and much more.
+Version: 4.9
 
 [VAR:Topic]
 [OBJECT:Topic]
