#UE4doc #docs - Added Engine Version Fields.

diff --git a/Engine/Documentation/Source/Engine/UI/GlobalAssetPicker/GlobalAssetPicker.INT.udn b/Engine/Documentation/Source/Engine/UI/GlobalAssetPicker/GlobalAssetPicker.INT.udn
--- a/Engine/Documentation/Source/Engine/UI/GlobalAssetPicker/GlobalAssetPicker.INT.udn
+++ b/Engine/Documentation/Source/Engine/UI/GlobalAssetPicker/GlobalAssetPicker.INT.udn
@@ -2,6 +2,7 @@ Availability:Public
 Title:Global Asset Picker
 Crumbs:%ROOT%, Engine, Engine/UI
 Description:The Global Asset Picker allows you to quickly find assets from any folder in the Asset Tree.
+Version: 4.9
 
 [VAR:Topic]
 [OBJECT:Topic]
