#UE4doc #docs - Added Engine Version Field.

diff --git a/Engine/Documentation/Source/Engine/UI/InterfaceOverview/InterfaceOverview.INT.udn b/Engine/Documentation/Source/Engine/UI/InterfaceOverview/InterfaceOverview.INT.udn
--- a/Engine/Documentation/Source/Engine/UI/InterfaceOverview/InterfaceOverview.INT.udn
+++ b/Engine/Documentation/Source/Engine/UI/InterfaceOverview/InterfaceOverview.INT.udn
@@ -3,6 +3,7 @@ Title: Layout Customization
 Crumbs: %ROOT%, Engine, Engine/UI
 Description:A look at how to customize the look and feel of the user interface by rearranging and docking various panels and tabs.
 Navigation:topic
+Version: 4.9
 
 [VAR:Topic]
 [OBJECT:Topic]
@@ -64,7 +65,7 @@ be rearranged to suit your needs. These windows can either be top-level editors
 
 ### Editors
 
-**Editors** are like applications within the Unreal Editor, such as the [Level Editor](Engine/UI/LevelEditor) and [Blueprint Editor](Engine/Blueprints/Editor). 
+**Editors** are like applications within the Unreal Editor, such as the [**Level Editor**](Engine/UI/LevelEditor) and [**Blueprint Editor**](Engine/Blueprints/Editor). 
 They are targeted at a specific aspect of the development process. Some editors, such as the Level Editor, can only have 
 a single instance in existence at any time; while other editors, like the Blueprint editor, can have multiple instances created, 
 but only a single instance can be editing a specific asset at any time. 
@@ -146,7 +147,7 @@ closing the editor. Tabs also remember their last location and will reopen there
 [/REGION]
 
 [REGION:tip]
-To save the layout manually, you can use the Save layout command in the Window menu.
+To save the layout manually, you can use the **Save layout** command in the **Window** menu.
 [/REGION]
 
 
