#UE4doc #docs - Added Engine Version Field. Updated Editor Preferences window Image.

diff --git a/Engine/Documentation/Source/Engine/UI/KeyBindings/EditorKeyBindings.INT.udn b/Engine/Documentation/Source/Engine/UI/KeyBindings/EditorKeyBindings.INT.udn
--- a/Engine/Documentation/Source/Engine/UI/KeyBindings/EditorKeyBindings.INT.udn
+++ b/Engine/Documentation/Source/Engine/UI/KeyBindings/EditorKeyBindings.INT.udn
@@ -2,6 +2,7 @@ Availability:Public
 Title:Keybindings Editor
 Crumbs:%ROOT%, Engine, Engine/UI
 Description:Configurable controls for common commands in Unreal Editor.
+Version: 4.9
 
 [VAR:Topic]
 [OBJECT:Topic]
@@ -43,7 +44,7 @@ Description:Configurable controls for common commands in Unreal Editor.
 [/VAR]
 
 [EXCERPT:Intro]
-*Keyboard shortcuts* are combinations of key presses on your keyboard that execute specific commands or actions. Common commands, as well as those specific to the mainframe or other tools in Unreal Editor 4, are configurable. You can customize keyboard shortcuts in the Unreal Editor to suit your workflow and personal preferences by binding specific keys or key combinations to specific commands. 
+**Keyboard shortcuts** are combinations of key presses on your keyboard that execute specific commands or actions. Common commands, as well as those specific to the mainframe or other tools in Unreal Editor 4, are configurable. You can customize keyboard shortcuts in the Unreal Editor to suit your workflow and personal preferences by binding specific keys or key combinations to specific commands. 
 
 You can set up your own set of custom keyboard shortcuts in the **Editor Preferences** window.  To modify the keyboard shortcuts, click **Editor Preferences** on the **Edit** menu, and then click **Keyboard Shortcuts**.
 [/EXCERPT:Intro]
@@ -58,7 +59,7 @@ You can use the **Search** field to filter the list of displayed commands by typ
 
 ## Creating a New Keyboard Shortcut
 
-1. Click in the text field for the command you want to bind to a keyboard shortcut.
+1. Click in the **text field** for the command you want to bind to a keyboard shortcut.
 
 1. Press the combination of keys you want to use to execute the command.
 
