#UE4doc #docs - Added Engine Version Field. Updated Details panel  image to match latest UI.

diff --git a/Engine/Documentation/Source/Engine/UI/LevelEditor/LevelEditor.INT.udn b/Engine/Documentation/Source/Engine/UI/LevelEditor/LevelEditor.INT.udn
--- a/Engine/Documentation/Source/Engine/UI/LevelEditor/LevelEditor.INT.udn
+++ b/Engine/Documentation/Source/Engine/UI/LevelEditor/LevelEditor.INT.udn
@@ -2,6 +2,7 @@ Availability:Public
 Title:Level Editor
 Crumbs:%ROOT%, Engine, Engine/UI
 Description:A guide that overviews the design and construction of game levels and environments.
+Version: 4.9
 
 [VAR:playlistnav]
 [OBJECT:PlaylistNavigation]
@@ -100,11 +101,11 @@ Description:A guide that overviews the design and construction of game levels an
 [TOC(start:2)]
 
 [EXCERPT:Intro]
-The **Level Editor** provides the core level creation functionality for Unreal Editor. This is where levels are created, viewed, and modified mainly by placing, transforming, and editing the properties of [Actors](Engine/Actors).
+The **Level Editor** provides the core level creation functionality for Unreal Editor. This is where levels are created, viewed, and modified mainly by placing, transforming, and editing the properties of [**Actors**](Engine/Actors).
 [/EXCERPT:Intro]
 
 [EXCERPT:Levels]
-In Unreal Editor, the scenes in which you create your game experience are generally referred to as [Levels](Engine/Levels). You can think of a level as a 3D environment into which you place a series of objects and geometry to define the world your players will experience. Any object that is placed in your world, be it a light, a mesh, or a character, is considered to be an Actor. Technically speaking, _Actor_ is a programming class used within the Unreal Engine to define an object that has 3D position, rotation, and scale data. For sake of ease, however, it will be easiest to think of an Actor as _any object that can be placed in your levels_.
+In Unreal Editor, the scenes in which you create your game experience are generally referred to as [**Levels**](Engine/Levels). You can think of a level as a 3D environment into which you place a series of objects and geometry to define the world your players will experience. Any object that is placed in your world, be it a light, a mesh, or a character, is considered to be an Actor. Technically speaking, _Actor_ is a programming class used within the Unreal Engine to define an object that has 3D position, rotation, and scale data. For sake of ease, however, it will be easiest to think of an Actor as _any object that can be placed in your levels_.
 [/EXCERPT:Levels]
 
 [EXCERPT:Actors]
@@ -148,9 +149,9 @@ To the right of the Tab Bar is the name of the current project.
 
 ![Main Menu Bar](Engine/UI/LevelEditor/MenuBar/MainMenubar.png)
 
-See the [Menu Bar](Engine/UI/LevelEditor/MenuBar) page for descriptions of each menu and its contents.
+See the [**Menu Bar**](Engine/UI/LevelEditor/MenuBar) page for descriptions of each menu and its contents.
 
-The Console is a text field that allows special console commands to be entered that are recognized by the editor. The text field has an auto-complete feature that automatically lists all commands matching the text currently in the box.
+The **Console** (`) is a text field that allows special console commands to be entered that are recognized by the editor. The text field has an auto-complete feature that automatically lists all commands matching the text currently in the box.
 
 The button on the far right of the menu bar is:
 
@@ -170,7 +171,7 @@ The button on the far right of the menu bar is:
 
 The **Toolbar** panel, like in most applications, is a group of commands providing quick access to commonly used tools and operations.
 
-See the [Toolbar](Engine/UI/LevelEditor/Toolbar) page for descriptions of each of the items on the Toolbar.
+See the [**Toolbar**](Engine/UI/LevelEditor/Toolbar) page for descriptions of each of the items on the Toolbar.
 
 ### Viewport
 
@@ -180,7 +181,7 @@ The **Viewport** panel is your window into the worlds you create in UnrealEd.
 
 This panel contains a set of viewports, each of which can be maximized to fill the entire panel and offer the ability to display the world from one of three orthographic views (Top, Side, Front) or a perspective view giving you complete control over what you see as well as how you see it.
 
-See [Viewports](Engine/UI/LevelEditor\Viewports) for more information on working with viewports.
+See [**Viewports**](Engine/UI/LevelEditor\Viewports) for more information on working with viewports.
 
 ### Details
 
@@ -188,7 +189,7 @@ See [Viewports](Engine/UI/LevelEditor\Viewports) for more information on working
 
 The **Details** panel contains information, utilities, and functions specific to the current selection in the viewport. It contains transform edit boxes for moving, rotating, and scaling Actors, displays all of the editable properties for the selected Actors, and provides quick access to additional editing functionality depending on the type of Actor(s) selected in the viewport. For instance, selected Actors can be exported to FBX and converted to another compatible type. The Selection Details also allows you to view the materials used by the selected Actors, if any, and quickly open them for editing.
 
-See the [Details](Engine/UI/LevelEditor/Details) page for a more complete overview and guide to using the **Details** panel in the Level Editor.
+See the [**Details**](Engine/UI/LevelEditor/Details) page for a more complete overview and guide to using the **Details** panel in the Level Editor.
 
 ### Modes
 
@@ -203,12 +204,12 @@ See the [Details](Engine/UI/LevelEditor/Details) page for a more complete overvi
 
 The **World Outliner** panel displays all of the Actors within the scene in a hierarchical tree view. Actors can be selected and modified directly from the **World Outliner**. You can also use the Info drop down menu to turn on an extra column that shows Levels, Layers, or ID Names.
 
-See the [World Outliner](Engine/UI/LevelEditor/SceneOutliner) page for details on using the **World Outliner**.
+See the [**World Outliner**](Engine/UI/LevelEditor/SceneOutliner) page for details on using the **World Outliner**.
 
 ### Layers
 
 [INCLUDE:Engine/UI/LevelEditor/Layers#Overview]
 
-See the [](Engine/UI/LevelEditor/Layers) page for details on using the **Layers** panel.
+See the **[](Engine/UI/LevelEditor/Layers)** page for details on using the **Layers** panel.
 
 %playlistnav%
