#UE4doc #docs - Added Engine Version Field. Updated Image to match latest UI. Removed/Added Bold lettering where needed.

diff --git a/Engine/Documentation/Source/Engine/UI/LevelEditor/Details/LevelEditorDetails.INT.udn b/Engine/Documentation/Source/Engine/UI/LevelEditor/Details/LevelEditorDetails.INT.udn
--- a/Engine/Documentation/Source/Engine/UI/LevelEditor/Details/LevelEditorDetails.INT.udn
+++ b/Engine/Documentation/Source/Engine/UI/LevelEditor/Details/LevelEditorDetails.INT.udn
@@ -2,6 +2,7 @@ Availability:Public
 Title:Details Panel
 Crumbs: %ROOT%, Engine, Engine/UI, Engine/UI/LevelEditor
 Description:Panel providing information and customized editing tools for selected Actors in the Level Editor.
+Version: 4.9
 
 [VAR:playlistnav]
 [OBJECT:PlaylistNavigation]
@@ -29,18 +30,18 @@ Description:Panel providing information and customized editing tools for selecte
 The **Details** panel contains information, utilities, and functions specific to the current selection in the viewport. 
 It contains transform edit boxes for moving, rotating, and scaling Actors, displays all of the editable properties for 
 the selected Actors, and provides quick access to additional editing functionality depending on the types of Actors selected 
-in the viewport. For instance, selected Actors can be exported to FBX and converted to another compatible type. The **Details** 
+in the viewport. For instance, selected Actors can be exported to FBX and converted to another compatible type. The Details 
 panel also allows you to view the Materials used by the selected Actors, if any, and quickly open them for editing.
 [/EXCERPT:DetailsOverview]
 
-![Level Editor Details Panel](LE_DetailsUI2.png)
+![Level Editor Details Panel](details_panel.png)
 
-For more information about the **Details** Panel UI, see [](Engine/UI/LevelEditor\Details\Interface).
+For more information about the Details Panel UI, see **[](Engine/UI/LevelEditor\Details\Interface)**.
 
 ## Actor Names
 
 Actors can have _friendly_ names set directly in the editor by the designer. These names can be used to access these Actors or 
-find them using the search functionality in the [World Outliner panel](Engine/UI/LevelEditor\SceneOutliner).
+find them using the search functionality in the **[World Outliner panel](Engine/UI/LevelEditor\SceneOutliner)**.
 
 To edit the Actor Name, simply type the name into the text box in at the top of the **Details** panel.
 
@@ -48,7 +49,7 @@ To edit the Actor Name, simply type the name into the text box in at the top of
 
 ## Search Filter
 
-The properties displayed within the **Details** panel can be filtered using the **Search** box. As you type, the properties are 
+The properties displayed within the Details panel can be filtered using the **Search** box. As you type, the properties are 
 automatically filtered to display only those properties matching the text.
 
 ![Properties filtered](LE_DetailsSearch2.png)
@@ -100,7 +101,7 @@ Properties declared as `editconst`, which cannot be modified in the editor, disp
 
 ## Property Controls
 
-Many of the basic types of properties are displayed in the **Details** panel using familiar editing widgets. 
+Many of the basic types of properties are displayed in the Details panel using familiar editing widgets. 
 Other types use customized widgets to provide a more intuitive experience when modifying their values.
 
 [REGION:topics fourth]
@@ -165,13 +166,13 @@ Other types use customized widgets to provide a more intuitive experience when m
 
 Properties in the **Display** panel are displayed in _categories_. Generally, these categories - **Rendering**, 
 **Lighting**, **Collision**, etc. - are determined by how the property is declared in code, and are used as a 
-means of organizing related properties into groups. Other categories you see in the **Details** panel - **Transform**, 
+means of organizing related properties into groups. Other categories you see in the Details panel - **Transform**, 
 **Static Mesh**, **Materials**, **Actor**, **Code View**, **Layers** - are custom widgets designed to expose certain 
 properties or functionality in a conspicuous manner making them easy to find and modify or use.
 
 ### Advanced Properties
 
-Some properties are declared as advanced and are hidden away by default in the **Details** panel. Categories that 
+Some properties are declared as advanced and are hidden away by default in the Details panel. Categories that 
 contain advanced properties have an expander button ![](button_advanced.png). Clicking the button causes the advanced 
 properties to be displayed.
 
