#UE4doc #docs - Added Engine Version Field.

diff --git a/Engine/Documentation/Source/Engine/UI/LevelEditor/Details/Customizations/CodeView/CodeViewCategory.INT.udn b/Engine/Documentation/Source/Engine/UI/LevelEditor/Details/Customizations/CodeView/CodeViewCategory.INT.udn
--- a/Engine/Documentation/Source/Engine/UI/LevelEditor/Details/Customizations/CodeView/CodeViewCategory.INT.udn
+++ b/Engine/Documentation/Source/Engine/UI/LevelEditor/Details/Customizations/CodeView/CodeViewCategory.INT.udn
@@ -2,6 +2,7 @@ Availability:Public
 Title: Code View Category
 Crumbs: %ROOT%, Engine, Engine/UI, Engine/UI/LevelEditor, Engine/UI/LevelEditor/Details
 Description:Displays the class hierarchy of the currently selected Actor, along with any public functions within those classes.
+Version: 4.9
 
 
 The **Code View** category displays the class of the currently selected Actor, along with any public functions within 
