#UE4doc #docs -  Added Engine Version Field.

diff --git a/Engine/Documentation/Source/Engine/UI/LevelEditor/Details/Customizations/Materials/MaterialsCategory.INT.udn b/Engine/Documentation/Source/Engine/UI/LevelEditor/Details/Customizations/Materials/MaterialsCategory.INT.udn
--- a/Engine/Documentation/Source/Engine/UI/LevelEditor/Details/Customizations/Materials/MaterialsCategory.INT.udn
+++ b/Engine/Documentation/Source/Engine/UI/LevelEditor/Details/Customizations/Materials/MaterialsCategory.INT.udn
@@ -2,6 +2,7 @@ Availability:Public
 Title: Materials Category
 Crumbs: %ROOT%, Engine, Engine/UI, Engine/UI/LevelEditor, Engine/UI/LevelEditor/Details
 Description:Customized controls for viewing and modifying the materials applied to all components of an Actor.
+Version: 4.9
 
 The **Materials** category displays all the Materials applied to the Actor(s) that are selected, grouped by element. 
 For example, a Static Mesh with 3 elements would display 3 Materials when selected.
@@ -10,7 +11,7 @@ For example, a Static Mesh with 3 elements would display 3 Materials when select
 
 | Control | Description |
 | ------- | ----------- |
-| ![](Details_Mat_ObjRef.png)| Standard [Object reference](Engine/UI/LevelEditor/Details/Properties/ObjectReference) property controls for setting, finding, and resetting the assigned Material back to the base Material. Clicking the drop-down list that contains the Material's name and then clicking **Edit** on the context menu opens the Material in the Material Instance Editor for modification. |
+| ![](Details_Mat_ObjRef.png)| Standard [**Object reference**](Engine/UI/LevelEditor/Details/Properties/ObjectReference) property controls for setting, finding, and resetting the assigned Material back to the base Material. Clicking the drop-down list that contains the Material's name and then clicking **Edit** on the context menu opens the Material in the Material Instance Editor for modification. |
 | ![Material Thumbnail](button_Details_Mat_Image.png)| Thumbnail of the Material, which can be **double-clicked** to open the Material Instance in the Material Instance Editor for modification. |
 | ![Find Texture](button_textures.png) | List of Textures used within the Material. Selecting a Texture from the list finds that Texture in the **Content Browser**. |
 
diff --git a/Engine/Documentation/Source/Engine/UI/LevelEditor/Details/Customizations/StaticMeshes/StaticMeshesCategory.INT.udn b/Engine/Documentation/Source/Engine/UI/LevelEditor/Details/Customizations/StaticMeshes/StaticMeshesCategory.INT.udn
--- a/Engine/Documentation/Source/Engine/UI/LevelEditor/Details/Customizations/StaticMeshes/StaticMeshesCategory.INT.udn
+++ b/Engine/Documentation/Source/Engine/UI/LevelEditor/Details/Customizations/StaticMeshes/StaticMeshesCategory.INT.udn
@@ -2,11 +2,12 @@ Availability:Public
 Title: Static Meshes Category
 Crumbs: %ROOT%, Engine, Engine/UI, Engine/UI/LevelEditor, Engine/UI/LevelEditor/Details
 Description:Customized controls for viewing and modifying the StaticMesh assets applied to all components of an Actor.
+Version: 4.9
 
 The **Static Mesh** category in the **Details** panel is only displayed when the selected Actor(s) are StaticMeshActors 
 or InterpActors. It promotes the normally buried StaticMesh property of the Actor's StaticMeshComponent, making it a 
 first-class citizen and displays it using a customized widget that includes a thumbnail along with the usual 
-[Object reference](Engine/UI/LevelEditor/Details/Properties/ObjectReference) property controls.
+[**Object reference**](Engine/UI/LevelEditor/Details/Properties/ObjectReference) property controls.
 
 ![](staticmesh.png)
 
@@ -26,11 +27,11 @@ using various preset methods.
 [REGION:simpletable]
 | Method | Description |
 | ------ | ----------- |
-| Blocking Volume | Creates a BlockingVolume in the shape of a cube with the dimensions of the StaticMesh's bounding box. |
-| Heavy Convex Volume | Creates a convex BlockingVolume from the StaticMesh's geometry with a **NormalTolerance** of 0.01. |
-| Normal Convex Volume | Creates a convex BlockingVolume from the StaticMesh's geometry with a **NormalTolerance** of 0.15. |
-| Light Convex Volume | Creates a convex BlockingVolume from the StaticMesh's geometry with a **NormalTolerance** of 0.50. |
-| Rough Convex Volume | Creates a convex BlockingVolume from the StaticMesh's geometry with a **NormalTolerance** of 0.75. |
+| *Blocking Volume* | Creates a BlockingVolume in the shape of a cube with the dimensions of the StaticMesh's bounding box. |
+| *Heavy Convex Volume* | Creates a convex BlockingVolume from the StaticMesh's geometry with a **NormalTolerance** of 0.01. |
+| *Normal Convex Volume* | Creates a convex BlockingVolume from the StaticMesh's geometry with a **NormalTolerance** of 0.15. |
+| *Light Convex Volume* | Creates a convex BlockingVolume from the StaticMesh's geometry with a **NormalTolerance** of 0.50. |
+| *Rough Convex Volume* | Creates a convex BlockingVolume from the StaticMesh's geometry with a **NormalTolerance** of 0.75. |
 [/REGION]
 [REGION:note]
 $ NormalTolerance : The rejection tolerance.  When figuring out which planes to cut the blocking volume cube with
