#UE4doc #docs - Added Engine Version Field. Corrected Bold lettering where needed.

diff --git a/Engine/Documentation/Source/Engine/UI/LevelEditor/Details/Customizations/Transform/TransformCategory.INT.udn b/Engine/Documentation/Source/Engine/UI/LevelEditor/Details/Customizations/Transform/TransformCategory.INT.udn
--- a/Engine/Documentation/Source/Engine/UI/LevelEditor/Details/Customizations/Transform/TransformCategory.INT.udn
+++ b/Engine/Documentation/Source/Engine/UI/LevelEditor/Details/Customizations/Transform/TransformCategory.INT.udn
@@ -2,20 +2,21 @@ Availability:Public
 Title: Transform Category
 Crumbs: %ROOT%, Engine, Engine/UI, Engine/UI/LevelEditor, Engine/UI/LevelEditor/Details
 Description:Customized controls for modifying the location, rotation, and scale of an Actor.
+Version: 4.9
 
 [EXCERPT:Transform]
-The Transform section of the **Details** panel allows you to view and edit the transforms - **Location**, **Rotation**, and **Scale** - of selected Actor(s). Also, where applicable, this also contains the settings for Actor Mobility.
+The **Transform** section of the Details panel allows you to view and edit the transforms - **Location**, **Rotation**, and **Scale** - of selected Actor(s). Also, where applicable, this also contains the settings for Actor Mobility.
 
 ![Transform Properties](Transform_Properties.png)
 
-Each transform property has numeric entry fields for the X, Y, and Z axes. You may type precise values directly 
+Each **Transform** property has numeric entry fields for the X, Y, and Z axes. You may type precise values directly 
 into these fields to adjust the selected Actors. When more than one Actor is selected and the values of the properties differ, 
 the fields will display _Multiple Values_. In such cases, entering a number will cause that value to be entered for all selected Actors.
 
-The Rotation fields have the unique ability to be used as sliders. Clicking and dragging on the field allows you to 
+The **Rotation** fields have the unique ability to be used as sliders. Clicking and dragging on the field allows you to 
 _roll_ the value, increasing or decreasing depending on the distance the mouse is moved.
 
-The Scale fields can also be locked by clicking the ![](button_lock_scale.png) button. When locked, the ratio of 
+The **Scale** fields can also be locked by clicking the ![](button_lock_scale.png) button. When locked, the ratio of 
 the scale values to each other are maintained when any of the individual values are modified, allowing for uniform scaling.
 
 The transform properties default to _relative_ transformation. This means the transform is relative to the component's parent. 
@@ -32,4 +33,4 @@ using _absolute_ transformations, the transform is relative to the world instead
 
 [INCLUDE:Engine/Actors/Mobility#intro]
 
-See [](Engine/Actors/Mobility) for more information.
+See **[](Engine/Actors/Mobility)** for more information.
