#UE4doc #docs:
Removed unnecessary bolding from multiple terms

diff --git a/Engine/Documentation/Source/Engine/Blueprints/UserGuide/BlueprintCommsUsage/BPComHowTo/BPComHowTo.INT.udn b/Engine/Documentation/Source/Engine/Blueprints/UserGuide/BlueprintCommsUsage/BPComHowTo/BPComHowTo.INT.udn
--- a/Engine/Documentation/Source/Engine/Blueprints/UserGuide/BlueprintCommsUsage/BPComHowTo/BPComHowTo.INT.udn
+++ b/Engine/Documentation/Source/Engine/Blueprints/UserGuide/BlueprintCommsUsage/BPComHowTo/BPComHowTo.INT.udn
@@ -7,7 +7,7 @@ Description:Examples of how to setup and use different methods of Blueprint Comm
 
 On this page you will learn step-by-step how to set up and use different methods of **Blueprint Communication**. 
 
-For examples of when to use each type of **Blueprint Communication**, refer to the [](Engine\Blueprints\UserGuide\BlueprintCommsUsage) page.
+For examples of when to use each type of Blueprint Communication, refer to the [](Engine\Blueprints\UserGuide\BlueprintCommsUsage) page.
 
 ### Direct Blueprint Communications
 
@@ -15,12 +15,12 @@ Below we have two **Blueprints** in our level that we want to have communicate w
 
 ![](2_1.png)(w:400)
 
-* The cube above is a **Blueprint** created using the **Shape_Cube** mesh and its collision has been set to **OverlapOnlyPawn** so it acts as a trigger. Also enable **Generate Overlap Events**
+* The cube above is a Blueprint created using the **Shape_Cube** mesh and its collision has been set to **OverlapOnlyPawn** so it acts as a trigger. Also enable **Generate Overlap Events**
 * The sparks above are the **Blueprint_Effect_Sparks** asset (included with starter content).
 
-Using **Direct Blueprint Communication**, we would do the following: 
+Using Direct Blueprint Communication, we would do the following: 
 
-1. In the **Shape_Cube** Blueprint, under **MyBlueprint** click the ![](Engine\Blueprints\Editor\UIComponents\MyBlueprint\plus_Button.png) button on the Variables category: ![](Engine\Blueprints\Editor\UIComponents\MyBlueprint\myblueprint_variable.png). 
+1. In the **Shape_Cube** Blueprint, under **My Blueprint** click the ![](Engine\Blueprints\Editor\UIComponents\MyBlueprint\plus_Button.png) button on the Variables category: ![](Engine\Blueprints\Editor\UIComponents\MyBlueprint\myblueprint_variable.png). 
 
 1. In the **Details** panel, under **Variable Type** select the type of Blueprint you wish to access.
 
@@ -37,7 +37,7 @@ Using **Direct Blueprint Communication**, we would do the following:
 	![](2_4.png)
 
 	1. **Variable Name** - give the variable a descriptive name such as **TargetBlueprint**. 
-	1. **Variable Type** - this should be the type of **Blueprint** that you want to access. 
+	1. **Variable Type** - this should be the type of Blueprint that you want to access. 
 	1. **Editable** - make sure this is checked to expose the variable and make it public which will allow you to access it in the Level Editor. 
 	1. **Tool Tip** - add a short description of what the variable does or what it will reference. 
 
@@ -45,22 +45,22 @@ Using **Direct Blueprint Communication**, we would do the following:
 
 	![](2_6.png)
 
-	1. Click the **None** drop-down box to assign a **Target Blueprint**. 
+	1. Click the **None** drop-down box to assign a Target Blueprint. 
 	1. All instances of the Blueprint placed in the level are displayed here, allowing you to specify which instance is the **Target Blueprint**. 
 
-	Here we are stating which **Blueprint_Effect_Sparks** Blueprint Actor we want to affect that is placed in our level; this considered the **Instance** Actor. If we had multiple sparks in our level and we only wanted to turn off one of them, we could choose which instance of the Blueprint we want to affect here by setting it as the **Target Blueprint**.
+	Here we are stating which **Blueprint_Effect_Sparks** Blueprint Actor we want to affect that is placed in our level; this considered the Instance Actor. If we had multiple sparks in our level and we only wanted to turn off one of them, we could choose which instance of the Blueprint we want to affect here by setting it as the **Target Blueprint**.
 			
 1. Instead of using the drop-down, you can click the ![](2_7.png) icon, then click on the object you want to reference that is placed in your level. 
 
 	![](2_8.png)
 
-	You can only set the **Target Blueprint** to target the specified **Variable Type** (in our case **Blueprint_Effect_Sparks**).
+	You can only set the **Target Blueprint** to target the specified Variable Type (in our case **Blueprint_Effect_Sparks**).
 
-1. In the **Shape_Cube** Blueprint, while holding down the **Control** key, drag the variable into the graph. 
+1. In the **Shape_Cube** Blueprint, while holding down the **Ctrl** key, drag the variable into the graph. 
 
 	![](2_9.png)
 
-	1. This creates a "getter" node which allows you to access Events, Variables, Functions, etc. from the **Target Blueprint**. 
+	1. This creates a Getter node which allows you to access Events, Variables, Functions, etc. from the **Target Blueprint**. 
 	1. Drag off the out pin to view the context menu.  
 	
 	Here we are searching for the spark effect and spark audio components from our **Target Blueprint** as we want to access those components. 
@@ -84,23 +84,23 @@ In the example below, inside our **MyCharacter** Blueprint, we have stated that
 
 We can then access the **Blueprint_Effect_Smoke** Blueprint (contained in the yellow boxes above) and get the Smoke Effect and Smoke Audio components from our **Target Blueprint** and **Deactivate** them when **F** is pressed a second time (which is what the **Flip/Flop** node does). Therefore we are accessing one Blueprint from inside another Blueprint through **Direct Blueprint Communication**. 
 
-Refer to the **[](Engine\Blueprints\UserGuide\BlueprintComms)** documentation for more information. 
+Refer to the [](Engine\Blueprints\UserGuide\BlueprintComms) documentation for more information. 
 
 
 ### Blueprint Casting
 
 [EXCERPT:CastExample]
 
-For this example, we have a fire effect **Blueprint** in our level (which is an Actor) and we want it to communicate with the playable **Character Blueprint** the player is using. When the player enters the fire, we want to send a signal to the **Character Blueprint** that the player has entered the fire and that they should now take damage. By using the **Return Value** of an **OverlapEvent**, we can **Cast To** our **Character Blueprint** and access the Events, Functions, or Variables within it. 
+For this example, we have a fire effect Blueprint in our level (which is an Actor) and we want it to communicate with the playable Character Blueprint the player is using. When the player enters the fire, we want to send a signal to the Character Blueprint that the player has entered the fire and that they should now take damage. By using the Return Value of an OverlapEvent, we can Cast To our Character Blueprint and access the Events, Functions, or Variables within it. 
 
  ![](3_0a.png)
 
 * The fire effect above is the **Blueprint_Effect_Fire** asset (included with starter content).
 * A sphere component named **Trigger** was added to the Blueprint and was set to **OverlapOnlyPawn** for its collision. 
 
- Using **Blueprint Casting**, we would do the following: 
+ Using Blueprint Casting, we would do the following: 
 
-1. The **Character Blueprint** assigned to the **Default Pawn Class** (the playable character) is our **Target Blueprint** we want to access. 
+1. The Character Blueprint assigned to the **Default Pawn Class** (the playable character) is our Target Blueprint we want to access. 
 
 	![](3_0b.png)
 
diff --git a/Engine/Documentation/Source/Engine/Blueprints/UserGuide/BlueprintCommsUsage/BlueprintCommsUsage.INT.udn b/Engine/Documentation/Source/Engine/Blueprints/UserGuide/BlueprintCommsUsage/BlueprintCommsUsage.INT.udn
--- a/Engine/Documentation/Source/Engine/Blueprints/UserGuide/BlueprintCommsUsage/BlueprintCommsUsage.INT.udn
+++ b/Engine/Documentation/Source/Engine/Blueprints/UserGuide/BlueprintCommsUsage/BlueprintCommsUsage.INT.udn
@@ -170,7 +170,7 @@ Here are a few examples when you could use **Blueprint Casting**:
 	*	You have several lights in your level and you want them all to turn on or off when an event occurs. 
 [/EXCERPT:Castusage]
 
-Refer to [**Setting Up Blueprint Casts**](Engine\Blueprints\BP_HowTo\Casting) for a walk through of how to set one up. 
+Refer to [Setting Up Blueprint Casts](Engine\Blueprints\BP_HowTo\Casting) for a walk through of how to set one up. 
 
 
 
