#UE4doc #docs:
Added engine version field
Removed a related page as it is only published to docs, therefore resulting in a 404 in the publicly published version
Removed bolding from a few terms

diff --git a/Engine/Documentation/Source/Engine/Blueprints/UserGuide/CastNodes/CastNodes.INT.udn b/Engine/Documentation/Source/Engine/Blueprints/UserGuide/CastNodes/CastNodes.INT.udn
--- a/Engine/Documentation/Source/Engine/Blueprints/UserGuide/CastNodes/CastNodes.INT.udn
+++ b/Engine/Documentation/Source/Engine/Blueprints/UserGuide/CastNodes/CastNodes.INT.udn
@@ -3,8 +3,8 @@ Title:Casting
 Crumbs: %ROOT%, Engine, Engine/Blueprints, Engine/Blueprints/UserGuide
 Description:Examples of different Casting Nodes and usage cases. 
 Related: Engine/Blueprints/UserGuide/BlueprintCommsUsage
-Related: Engine/Blueprints/UserGuide/BlueprintCommsUsage/BPComHowTo
 Related: Engine/Blueprints/HowTo
+version: 4.9
 
 [VAR:Topic]
 [OBJECT:Topic]
@@ -47,13 +47,13 @@ Related: Engine/Blueprints/HowTo
 
 [TOC (start:2 end:2)]
 
-When using a **Cast To** node in **Blueprints**, put simply, you are attempting to check if the object you are casting from is the specific object you **Cast To**. In other words, say you have created a special Character Blueprint (called MyCharacter for example) that has variables in it or other custom functionality contained within it and assigned it as the Default Pawn Class (or the Character Blueprint all Player Characters use by default). 
+When using a **Cast To** node in **Blueprints**, put simply, you are attempting to check if the object you are casting from is the specific object you casted to. In other words, say you have created a special Character Blueprint (called MyCharacter for example) that has variables in it or other custom functionality contained within it and assigned it as the Default Pawn Class (or the Character Blueprint all Player Characters use by default). 
 
 If you want to access properties of the Player Character from another Blueprint, you can use the **Get Player Character** node and affect the Player Character in very general ways (setting/getting its location, rotation, etc.) but you cannot access any of the custom functionality you have added to the MyCharacter Blueprint you have set up as you are getting the Player Character but not a specific type of Character. 
 
 By using the **Get Player Character** node, then using a **Cast To MyCharacter** node (the special Character Blueprint), you can then say if the Player Character is MyCharacter, let me access the Variables, Functions, Events or any other special functionality contained within that Blueprint.
 
-On this page, you will learn about some of the use cases for Casting, see an example for using a Cast To node, read about Target Blueprint Casting as well as other types of casting. 
+On this page, you will learn about some of the use cases for Casting, see an example for using a **Cast To** node, read about Target Blueprint Casting as well as other types of casting. 
 
 
 
