#UE4doc #docs:
Added engine version field
Added an explaination of what LMB stands for

diff --git a/Engine/Documentation/Source/Engine/Blueprints/UserGuide/CheatSheet/BlueprintEditorCheatSheet.INT.udn b/Engine/Documentation/Source/Engine/Blueprints/UserGuide/CheatSheet/BlueprintEditorCheatSheet.INT.udn
--- a/Engine/Documentation/Source/Engine/Blueprints/UserGuide/CheatSheet/BlueprintEditorCheatSheet.INT.udn
+++ b/Engine/Documentation/Source/Engine/Blueprints/UserGuide/CheatSheet/BlueprintEditorCheatSheet.INT.udn
@@ -2,10 +2,11 @@ Availability:Public
 Title:Blueprint Editor Cheat Sheet
 Crumbs: %ROOT%, Engine, Engine/Blueprints
 Description:Reference covering Blueprint shortcuts and useful actions.
+version: 4.9
 
 [TOC(start:2 end:2)]
 
-The **Blueprint** editor has a lot of productivity boosting shortcuts built-in, and while many will come naturally as you use the editor, others are a little bit hidden. This page covers the most commonly used actions and the shortcuts that can be used to perform those actions more efficiently. You can also download this list of shortcuts below. 
+The **Blueprint Editor** has a lot of productivity boosting shortcuts built-in, and while many will come naturally as you use the editor, others are a little bit hidden. This page covers the most commonly used actions and the shortcuts that can be used to perform those actions more efficiently. You can also download this list of shortcuts below. 
 
 Click to download: [Color Version](BlueprintCheatSheet-1989117414.pdf) | [Black/White Version](BlueprintCheatSheet_BlackAndWhite-1276559864.pdf)
 
@@ -31,7 +32,7 @@ Click to download: [Color Version](BlueprintCheatSheet-1989117414.pdf) | [Black/
 | Select Node | **Click** on a Node |
 | Add to Selection | **Shift + Click** on a Node |
 | Toggle Selection | **Ctrl + Click** on a Node |
-| Marquee Select (Replace) | **LMB** Drag |
+| Marquee Select (Replace) | **Left Mouse Button (LMB)** Drag |
 | Marquee Select (Add) | **Shift + LMB** Drag |
 | Marquee Select (Remove) | **Ctrl + LMB** Drag |
 
