#UE4doc #docs:
Added engine version field
Added a tip field to inform the user about keybindings
Fixed capitalization for a few terms

diff --git a/Engine/Documentation/Source/Engine/Blueprints/UserGuide/Comments/Comments.INT.udn b/Engine/Documentation/Source/Engine/Blueprints/UserGuide/Comments/Comments.INT.udn
--- a/Engine/Documentation/Source/Engine/Blueprints/UserGuide/Comments/Comments.INT.udn
+++ b/Engine/Documentation/Source/Engine/Blueprints/UserGuide/Comments/Comments.INT.udn
@@ -2,6 +2,7 @@ Availability:Public
 Title:Comments
 Crumbs: %ROOT%, Engine/Blueprints, Engine/Blueprints/UserGuide
 Description:Containers that can be used to group related nodes and provide descriptions about their functionality.
+version: 4.9
 
 [VAR:Topic]
 [OBJECT:Topic]
@@ -43,7 +44,7 @@ Description:Containers that can be used to group related nodes and provide descr
 
 
 [EXCERPT:Intro]
-**Comments** can be added directly to single _Blueprint_ nodes, or they can be added as comment boxes to group related nodes and provide descriptions 
+**Comments** can be added directly to single Blueprint nodes, or they can be added as comment boxes to group related nodes and provide descriptions 
 about their functionality. These can be used solely for organizational purposes to make the 
 **Graphs** more readable, but they can also be used for informational purposes as they 
 allow for textual descriptions to be added just like adding comments to code.
@@ -55,7 +56,7 @@ allow for textual descriptions to be added just like adding comments to code.
 
 ### Creating Node Comments
 
-To add a comment directly to a node in a _Blueprint_ graph:
+To add a comment directly to a node in a Blueprint graph:
 
 1. **Right-click** on the node.
 
@@ -69,7 +70,7 @@ To add a comment directly to a node in a _Blueprint_ graph:
 
 ### Editing and Removing Node Comments
 
-**Right-click** on the node again and change the text in the **Node Comment** text box in the **right-click** menu to change the comment.  Deleting the text and pressing **Enter** will remove the comment.
+**Right-click** on the node again and change the text in the **Node Comment** text box in the **Right-click** menu to change the comment.  Deleting the text and pressing **Enter** will remove the comment.
 
 ## Comment Boxes
 
@@ -82,7 +83,9 @@ Creating a comment box is easy:
 1. **Right-click** and choose **Create Comment from Selection** (or **Add Comment...** if no nodes are selected). A new comment box will appear.
 
 	![](create_comment_from_selection.png)
-
+	[REGION:tip]
+	You can also press **C** to create a comment.
+	[/REGION]
 1. Enter your comment into the header of the comment box that appears.  
 
 	![](type_box_comment.png)
@@ -101,7 +104,7 @@ Note that the text of comment boxes will scale as the graph is zoomed, so these
 
 ### Editing Comment Box Text
 
-**Right-click** on the header of the comment box, and change the text in the **Node Comment** text box of the **right-click** menu to change the text of your comment. Press **Enter** to save your new comment.
+**Right-click** on the header of the comment box, and change the text in the **Node Comment** text box of the **Right-click** menu to change the text of your comment. Press **Enter** to save your new comment.
 
 ![](rtclick_comment_box_change.png)
 
@@ -113,7 +116,7 @@ Alternatively:
 
 ### Deleting Comment Boxes
 
-Click on any comment box and press the **Delete** key, or **right-click** on the header of the comment box and select **Delete** from the **right-click** menu.  Both of these options will delete the comment box but
+Click on any comment box and press the **Delete** key, or **Right-click** on the header of the comment box and select **Delete** from the **Right-click** menu.  Both of these options will delete the comment box but
 have no effect on the nodes contained within the comment box.
 
 
