#UE4doc #docs:
Added engine version field
Changed font settings for a few terms and fixed some spacing

diff --git a/Engine/Documentation/Source/Engine/Blueprints/UserGuide/Components/ComponentsList.INT.udn b/Engine/Documentation/Source/Engine/Blueprints/UserGuide/Components/ComponentsList.INT.udn
--- a/Engine/Documentation/Source/Engine/Blueprints/UserGuide/Components/ComponentsList.INT.udn
+++ b/Engine/Documentation/Source/Engine/Blueprints/UserGuide/Components/ComponentsList.INT.udn
@@ -6,6 +6,7 @@ Related: Engine/Blueprints/UserGuide/Graphs
 Related: Engine/Blueprints/UserGuide/Variables
 Related: Engine/Blueprints/UserGuide/Events
 Related: Engine/Blueprints/UserGuide/Functions
+version: 4.9
 
 [VAR:Topic]
 [OBJECT:Topic]
@@ -50,7 +51,7 @@ Related: Engine/Blueprints/UserGuide/Functions
 [INCLUDE:Shared\Glossary\C#glossary_component]
 
 [REGION:note]
-For more information on **Components**, please see the [Components Overview](Engine\Components) documentation. 
+For more information on Components, please see the [Components Overview](Engine\Components) documentation. 
 [/REGION]
 
 
@@ -59,10 +60,10 @@ For more information on **Components**, please see the [Components Overview](Eng
 ## Components Window
 
 [EXCERPT:Intro]
-With an understanding of Components, the **Components** window inside the **Blueprint Editor** allows you to add Components to your **Blueprint**. This provides a means 
-of adding collision geometry via _CapsuleComponents_, _BoxComponents_, or _SphereComponents_, 
-adding rendered geometry in the form of _StaticMeshComponents_ or _SkeletalMeshComponents_, 
-controlling movement using _MovementComponents_, etc. The Components added in the Components list can 
+With an understanding of Components, the **Components** window inside the **Blueprint Editor** allows you to add Components to your Blueprint. This provides a means 
+of adding collision geometry via CapsuleComponents, BoxComponents, or SphereComponents, 
+adding rendered geometry in the form of StaticMeshComponents or SkeletalMeshComponents, 
+controlling movement using MovementComponents, etc. The Components added in the Components list can 
 also be assigned to instance variables providing access to them in the graphs of this or other Blueprints.
 [/EXCERPT:Intro]
 
@@ -89,7 +90,7 @@ Components can also be added by dragging-and-dropping assets from the **Content
 ![](drag_asset_components_list.png)
 
 [REGION:note]
-Assets this method applies to include: _StaticMeshes_, _SoundCues_, _SkeletalMeshes_, and _ParticleSystems_. 
+Assets this method applies to include: StaticMeshes, SoundCues, SkeletalMeshes, and ParticleSystems. 
 [/REGION]
 
 ## Removing Components
@@ -119,22 +120,22 @@ Transforming, Rotating, or Scaling a parent component will propagate the transfo
 
 ## Component Assets
 
-Once you have added a Component, you may need to specify the asset to occupy the Component (such as assigning a Static Mesh to use for a _StaticMeshComponent_). There are a couple of different ways you can assign an asset to use for a Component type as described below. 
+Once you have added a Component, you may need to specify the asset to occupy the Component (such as assigning a Static Mesh to use for a StaticMeshComponent). There are a couple of different ways you can assign an asset to use for a Component type as described below. 
 
 
 ### Assigning Component Assets
 
 To assign an asset to a Component in the **Components** window: 
 
-1. With the **Component** selected, in the **Details** panel find the section that corresponds to your Component type. 
+1. With the Component selected, in the **Details** panel find the section that corresponds to your Component type. 
 	
 	[REGION:raw]
 	![](emptyComponent.png)
 	[/REGION]
 
-	Above we have added a _StaticMeshComponent_ and under **StaticMesh** is where we would assign the asset to use.
+	Above we have added a StaticMeshComponent and under **StaticMesh** is where we would assign the asset to use.
 
-1. Click the **Current Asset** drop-down box, then select the asset to use from the context menu. 
+1. Click the **Static Mesh** drop-down box, then select the asset to use from the context menu. 
 
 	[REGION:raw]
 	![](selectComponentList.png)
@@ -152,9 +153,9 @@ Another method of assigning an asset can be done using the **Content Browser**:
 	![](emptyComponent.png)
 	[/REGION]
 
-	Above we have added a _StaticMeshComponent_ and under **StaticMesh** is where we would assign the asset to use. 
+	Above we have added a StaticMeshComponent and under **StaticMesh** is where we would assign the asset to use. 
 
-1.	Instead of clicking the **Current Asset** box, since an asset is selected in the **Content Browser**, click the ![](k2_button_use_scs.png) button. 
+1.	Instead of clicking the **Static Mesh** box, since an asset is selected in the **Content Browser**, click the ![](k2_button_use_scs.png) button. 
 
 	![](copy_asset.png)
 
@@ -217,11 +218,11 @@ To change the name of these variables:
 
 You can quickly add events and/or functions based a Component to the **Event Graph** of the Blueprint through different methods. Any events or functions created in this manner are specific to that particular function and do not have to be tested to verify which Component is involved.
 
-1. Add a Component for which events can be created, such as a _BoxComponent_.
+1. Add a Component for which events can be created, such as a BoxComponent.
 
 	![](AddBoxComponent.png)
 
-1. Provide a name for your Component, here we called it _Trigger_. 
+1. Provide a name for your Component, here we called it Trigger. 
 
 	![](NamedBoxComponent.png)
 
@@ -240,9 +241,9 @@ You can quickly add events and/or functions based a Component to the **Event Gra
 	![](EventAdded.png)
 
 
-Events and Functions can also be added for a Component from the **Event Graph** through the **MyBlueprint** window: 
+Events and Functions can also be added for a Component from the **Event Graph** through the **My Blueprint** window: 
 
-1. In the **MyBlueprint** window, under _Components_, select your Component. 
+1. In the **My Blueprint** window, under **Components**, select your Component. 
 	
 1. **Right-click** in the graph to bring up the context menu. 
 
@@ -253,10 +254,10 @@ Events and Functions can also be added for a Component from the **Event Graph**
 	If the component has any associated events or functions, they will be listed at the top.
 
 	[REGION:note]
-		Not all Components have associated events. For instance, a _PointLightComponent_ only has functions.
+		Not all Components have associated events. For instance, a PointLightComponent only has functions.
 	[/REGION]
 
-	You can also **Right-click** on the Component in the **MyBlueprint** window to access the **Add Event** context menu. 
+	You can also **Right-click** on the Component in the **My Blueprint** window to access the **Add Event** context menu. 
 
 	[REGION:raw]
 	![](RightClickMyBlueprint.png)
