#UE4doc #docs:
Added engine version field
Changed font settings for a few terms
Fixed capitalization for multiple terms

diff --git a/Engine/Documentation/Source/Engine/Blueprints/UserGuide/Debugging/KismetDebugging.INT.udn b/Engine/Documentation/Source/Engine/Blueprints/UserGuide/Debugging/KismetDebugging.INT.udn
--- a/Engine/Documentation/Source/Engine/Blueprints/UserGuide/Debugging/KismetDebugging.INT.udn
+++ b/Engine/Documentation/Source/Engine/Blueprints/UserGuide/Debugging/KismetDebugging.INT.udn
@@ -1,10 +1,11 @@
 Availability:Public
 Title: Blueprint Debugging
 Crumbs:%ROOT%, Engine/Blueprints, Engine/Blueprints/UserGuide
-Description:Pause execution using breakpoints to inspect graphs and the values of variables.
+Description:Pause execution using Breakpoints to inspect graphs and the values of variables.
+version: 4.9
 
 [EXCERPT:Overview]
-Blueprint debugging is a very powerful feature that allows you to pause execution of the game in Play In Editor or Simulate In Editor mode and step through any graph of a Blueprint or Level Blueprint through the use of breakpoints.
+Blueprint debugging is a very powerful feature that allows you to pause execution of the game in Play In Editor or Simulate In Editor mode and step through any graph of a Blueprint or Level Blueprint through the use of Breakpoints.
 [/EXCERPT:Overview]
 
 [TOC (start:2)]
@@ -12,7 +13,7 @@ Blueprint debugging is a very powerful feature that allows you to pause executio
 ## Debugging Controls
 
 The debug controls allow execution of the game in Play In Editor and Simulate In Editor mode to be controlled. The controls become enabled in the Blueprint Editor when the game is running. Different debugging-specific buttons appear on the 
-[Toolbar](Engine/Blueprints/Editor/UIComponents\Toolbar) depending on the type of _Blueprint_ being debugged and the current state of the debugging. 
+[Toolbar](Engine/Blueprints/Editor/UIComponents\Toolbar) depending on the type of Blueprint being debugged and the current state of the debugging. 
 
 [REGION:fullwidth]
 ![](debugging_editor.png)
@@ -21,33 +22,33 @@ The debug controls allow execution of the game in Play In Editor and Simulate In
 Both the **Debug** tab, which can be opened from the **Blueprint Editor's** window menu, and the **Blueprint Debugger** will also display the context-sensitive debugging buttons when Simulate in Editor or Play in Editor modes are active.
 
 The debug controls allow execution of the game in Play In Editor and Simulate In Editor mode to be controlled. The controls become enabled when the game is running. Different debugging-specific buttons appear on the 
-[Toolbar](Engine/Blueprints/Editor/UIComponents\Toolbar) depending on the type of _Blueprint_ being debugged and the current state of the debugging. Some controls only become enabled when relevant, i.e. when a **breakpoint** is hit.
+**Toolbar** depending on the type of Blueprint being debugged and the current state of the debugging. Some controls only become enabled when relevant, i.e. when a **Breakpoint** is hit.
 
 ![](blueprint_debugger.png)
 
 ## Breakpoints
 
-**Breakpoints** allow execution of the game to be paused at a certain point when running the game in the editor. All breakpoints for all graphs belonging to a Blueprint or Level Blueprint are displayed in the **Debug** tab. The names of the breakpoints in this list are hyperlinked allowing you to click on them and be taken directly to the node associated with the breakpoint in the **Graph** tab.
+Breakpoints allow execution of the game to be paused at a certain point when running the game in the editor. All Breakpoints for all graphs belonging to a Blueprint or Level Blueprint are displayed in the **Debug** tab. The names of the Breakpoints in this list are hyperlinked allowing you to click on them and be taken directly to the node associated with the Breakpoint in the **Graph** tab.
 
-Breakpoints can be set on any node in the graph that gets executed, such as **Function Calls** and **Flow Control** operations, by **right-clicking** on the node and choosing **Add breakpoint**. When a breakpoint has been set on a node, a graphic is displayed in the upper-left corner of the node.
+Breakpoints can be set on any node in the graph that gets executed, such as **Function Calls** and **Flow Control** operations, by **Right-clicking** on the node and choosing **Add Breakpoint**. When a Breakpoint has been set on a node, a graphic is displayed in the upper-left corner of the node.
 
 ![Blueprint Debugging - Breakpoint Enabled](k2_breakpoint.png)
 
-Breakpoints can be disabled as well by **right-clicking** on the node and choosing **Disable breakpoint** or clicking the icon next to the breakpoint in the **Debug** to toggle it off. This can also be done by **right-clicking** the breakpoint in the **Debug** tab and choosing **Disable breakpoint**.
+Breakpoints can be disabled as well by **Right-clicking** on the node and choosing **Disable Breakpoint** or clicking the icon next to the Breakpoint in the **Debug** to toggle it off. This can also be done by **Right-clicking** the Breakpoint in the **Debug** tab and choosing **Disable Breakpoint**.
 
-![Blueprint Debugging - Breakpoint Disable Option](k2_breakpoint_toggleoff.png)
+![Blueprint Debugging - Breakpoint Disable Option](k2_Breakpoint_toggleoff.png)
 
-This leaves the breakpoint in place, but it will not stop the game when the node is executed. A disabled breakpoint is displayed as just an outline.
+This leaves the Breakpoint in place, but it will not stop the game when the node is executed. A disabled Breakpoint is displayed as just an outline.
 
 ![Blueprint Debugging - Breakpoint Disabled](k2_breakpoint_disabled.png)
 
-To enable a disabled breakpoint, **right-click** on the node and choose **Enable breakpoint** or click the icon next to the breakpoint in the **Debug** tab. This can also be done by **right-clicking** the breakpoint in the **Debug** tab and choosing **Enable breakpoint**.
+To enable a disabled Breakpoint, **Right-click** on the node and choose **Enable Breakpoint** or click the icon next to the Breakpoint in the **Debug** tab. This can also be done by **Right-clicking** the Breakpoint in the **Debug** tab and choosing **Enable Breakpoint**.
 
 ![Blueprint Debugging - Breakpoint Enable Option](k2_breakpoint_toggleon.png)
 
-Breakpoints can be removed by **right-clicking** the node or **right-clicking** the breakpoint in the **Debug** tab and choosing **Remove breakpoint**.
+Breakpoints can be removed by **Right-clicking** the node or **Right-clicking** the Breakpoint in the **Debug** tab and choosing **Remove Breakpoint**.
 
-When a breakpoint is hit, the graph will focus on the node containing the breakpoint and an indicator is displayed on the breakpoint.
+When a Breakpoint is hit, the graph will focus on the node containing the Breakpoint and an indicator is displayed on the Breakpoint.
 
 ![Blueprint Debugging - Breakpoint Hit Indicator](k2_breakpoint_indicator.png)
 
@@ -62,12 +63,12 @@ In addition, the viewport will display an overlay to alert the player that the g
 [EXCERPT:DebugTab]
 
 [EXCERPT:DebugPanel]
-The **Debug** tab shows information designated as important by the designer in the form of breakpoints and watch values as well as a trace stack of all nodes belonging to the Blueprint that have been executed. This window also shows controls for controlling execution of the game when using breakpoints.
+The **Debug** tab shows information designated as important by the designer in the form of Breakpoints and watch values as well as a trace stack of all nodes belonging to the Blueprint that have been executed. This window also shows controls for controlling execution of the game when using Breakpoints.
 [/EXCERPT:DebugPanel]
 
 ### Watch Values
 
-**Watch values** allow you to quickly and easily see the values of variables during execution. All watch values for the selected Actor as well as the Level Blueprint are displayed in the Debug window when debugging after a breakpoint is hit and update as you step through the graph.
+**Watch values** allow you to quickly and easily see the values of variables during execution. All watch values for the selected Actor as well as the Level Blueprint are displayed in the Debug window when debugging after a Breakpoint is hit and update as you step through the graph.
 
 ![Blueprint Debugging - Watch Values](k2_watch_value.png)
 
