#UE4doc #docs:
Added engine version field
Changed font settings for multiple terms

diff --git a/Engine/Documentation/Source/Engine/Blueprints/UserGuide/EventDispatcher/EventDispatchers.INT.udn b/Engine/Documentation/Source/Engine/Blueprints/UserGuide/EventDispatcher/EventDispatchers.INT.udn
--- a/Engine/Documentation/Source/Engine/Blueprints/UserGuide/EventDispatcher/EventDispatchers.INT.udn
+++ b/Engine/Documentation/Source/Engine/Blueprints/UserGuide/EventDispatcher/EventDispatchers.INT.udn
@@ -3,9 +3,13 @@ Title:Event Dispatchers
 Crumbs: %ROOT%, Engine, Engine/Blueprints, Engine/Blueprints/UserGuide
 Description: Allows a Blueprint Class to report on its state to the Level Blueprint.
 related:Engine/Blueprints/UserGuide/BlueprintCommsUsage
-related:Engine/Blueprints/UserGuide/BlueprintCommsUsage/BPComHowTo
 related:Engine/Blueprints/UserGuide/Types/Interface/UsingInterfaces
 related:Engine/Blueprints/UserGuide/BlueprintComms
+version: 4.9
+<!--related:Engine/Blueprints/UserGuide/BlueprintCommsUsage/BPComHowTo
+Commented out due to page not being published publicly-->
+
+
 
 [VAR:Topic]
 [OBJECT:Topic]
@@ -68,7 +72,7 @@ To create a new Event Dispatcher:
 
 ### Setting Properties
 
-By selecting the Event Dispatcher in the **My Blueprint** tab, you can edit its properties in the **Details** tab.  You can set the tooltip and category for your Event Dispatcher, as well as add inputs.
+By selecting the Event Dispatcher in the **My Blueprint** tab, you can edit its properties in the **Details** panel.  You can set the tooltip and category for your Event Dispatcher, as well as add inputs.
 
 Adding inputs to your Event Dispatcher allows you to send variables to each event bound to your Event Dispatcher.  This allows data flow not only within your Blueprint Class, but also between your Blueprint Class
 and the Level Blueprint.
@@ -82,7 +86,7 @@ To add your own inputs to the Event Dispatcher:
 	![](new_input_Dispatcher.png)
 	[/REGION]
 
-1. Name the new input and set its type using the dropdown menu.  In this example, there is a _String_ input parameter named **MyStringParam**. 
+1. Name the new input and set its type using the dropdown menu.  In this example, there is a String input parameter named **MyStringParam**. 
 
 	[REGION:raw]
 	![](named_new_Dispatcher.png)
@@ -100,7 +104,7 @@ To add your own inputs to the Event Dispatcher:
 
 ## Using Event Dispatchers
 
-Once you have created the Event Dispatcher, you can add event nodes, bind nodes, and unbind nodes linked to it. Although you can **double-click** on the Event Dispatcher entry in the **My Blueprint** tab to 
+Once you have created the Event Dispatcher, you can add event nodes, bind nodes, and unbind nodes linked to it. Although you can **Double-click** on the Event Dispatcher entry in the **My Blueprint** tab to 
 open the Event Dispatcher's graph, the graph is locked and you cannot modify the Event Dispatcher directly.  The bind, unbind, and assign methods all allow you to add events to the Event Dispatcher's 
 event list, while the call method activates all the events stored in the event list. 
 
