#UE4doc #docs:
Added engine version field
Changed font settings for multiple terms

diff --git a/Engine/Documentation/Source/Engine/Blueprints/UserGuide/EventDispatcher/BindingAndUnbinding/BindingAndUnbinding.INT.udn b/Engine/Documentation/Source/Engine/Blueprints/UserGuide/EventDispatcher/BindingAndUnbinding/BindingAndUnbinding.INT.udn
--- a/Engine/Documentation/Source/Engine/Blueprints/UserGuide/EventDispatcher/BindingAndUnbinding/BindingAndUnbinding.INT.udn
+++ b/Engine/Documentation/Source/Engine/Blueprints/UserGuide/EventDispatcher/BindingAndUnbinding/BindingAndUnbinding.INT.udn
@@ -2,6 +2,7 @@ Availability:Public
 Title:Binding and Unbinding Events
 Crumbs: %ROOT%, Engine, Engine/Blueprints, Engine/Blueprints/UserGuide, Engine/Blueprints/UserGuide/EventDispatcher
 Description: Adding and removing events from an Event Dispatcher's events list.
+version: 4.9
 
 [TOC(start:2 end:3)]
 
@@ -11,11 +12,11 @@ Description: Adding and removing events from an Event Dispatcher's events list.
 Calling the Event Dispatcher will have no effect if there are no events bound to it. Consider that each Event Dispatcher has a list of events associated with it. The way to add an event to this list is by using a **Bind Event** node, and the way to remove an event from this list is by using
 an **Unbind Event** node.  It is also possible to unbind all of the events currently bound to an Event Dispatcher with an **Unbind All Events** node.  
 [/EXCERPT:Intro]
-Each event can be bound only once, even if the **Bind Event** node is executed multiple times.  Also, events in the _Blueprint Class_ and the _Level Blueprint_ are both added to the same event list, so
-an **Unbind All Events** node will unbind events in both the _Blueprint Class_ and the _Level Blueprint_.  
+Each event can be bound only once, even if the **Bind Event** node is executed multiple times.  Also, events in the Blueprint Class and the Level Blueprint are both added to the same event list, so
+an **Unbind All Events** node will unbind events in both the Blueprint Class and the Level Blueprint.  
 
-* An **Unbind All Events** node executed in the _Blueprint Class_ will unbind events in both the _Blueprint Class_ and the _Level Blueprint_ for all instances of the class.
-* An **Unbind All Events** node executed in the _Level Blueprint_ will unbind events in both the _Blueprint Class_ and the _Level Blueprint_, but just for the **Target** supplied.
+* An **Unbind All Events** node executed in the Blueprint Class will unbind events in both the Blueprint Class and the Level Blueprint for all instances of the class.
+* An **Unbind All Events** node executed in the Level Blueprint will unbind events in both the Blueprint Class and the Level Blueprint, but just for the **Target** supplied.
 
 ## Creating Bind, Unbind, and Unbind All Nodes
 
