#UE4doc #docs:
Added engine version field
Changed font settings for multiple terms
Replaced a few images due to incorrect UI

diff --git a/Engine/Documentation/Source/Engine/Blueprints/UserGuide/EventDispatcher/Calling/CallingEventDispatchers.INT.udn b/Engine/Documentation/Source/Engine/Blueprints/UserGuide/EventDispatcher/Calling/CallingEventDispatchers.INT.udn
--- a/Engine/Documentation/Source/Engine/Blueprints/UserGuide/EventDispatcher/Calling/CallingEventDispatchers.INT.udn
+++ b/Engine/Documentation/Source/Engine/Blueprints/UserGuide/EventDispatcher/Calling/CallingEventDispatchers.INT.udn
@@ -2,6 +2,7 @@ Availability:Public
 Title:Calling Event Dispatchers
 Crumbs: %ROOT%, Engine, Engine/Blueprints, Engine/Blueprints/UserGuide, Engine/Blueprints/UserGuide/EventDispatcher
 Description: Calling the Event Dispatcher executes all of the currently bound events in the events list.
+version: 4.9
 
 
 
@@ -12,7 +13,7 @@ Description: Calling the Event Dispatcher executes all of the currently bound ev
 
 [EXCERPT:Intro]
 Calling an Event Dispatcher with a **Call** node causes all of the events bound to the Event Dispatcher to fire. You can have more than one **Call** node for each Event Dispatcher, and you can also
-call the Event Dispatcher in both _Blueprint Classes_ and _Level Blueprints_.
+call the Event Dispatcher in both Blueprint Classes and Level Blueprints.
 [/EXCERPT:Intro]
 
 
