#UE4doc #docs:
Added engine version field

diff --git a/Engine/Documentation/Source/Engine/Blueprints/UserGuide/EventDispatcher/CreatingEvents/CreatingEvents.INT.udn b/Engine/Documentation/Source/Engine/Blueprints/UserGuide/EventDispatcher/CreatingEvents/CreatingEvents.INT.udn
--- a/Engine/Documentation/Source/Engine/Blueprints/UserGuide/EventDispatcher/CreatingEvents/CreatingEvents.INT.udn
+++ b/Engine/Documentation/Source/Engine/Blueprints/UserGuide/EventDispatcher/CreatingEvents/CreatingEvents.INT.udn
@@ -2,6 +2,7 @@ Availability:Public
 Title:Creating Dispatcher Events
 Crumbs: %ROOT%, Engine, Engine/Blueprints, Engine/Blueprints/UserGuide, Engine/Blueprints/UserGuide/EventDispatcher
 Description: Creating events that can be bound and added to the Event Dispatcher's events list.
+version: 4.9
 
 [VAR:Topic]
 [OBJECT:Topic]
