#UE4doc #docs:
Added engine version field
Changed font settings for multiple terms

diff --git a/Engine/Documentation/Source/Engine/Blueprints/UserGuide/EventGraph/BlueprintEventGraph.INT.udn b/Engine/Documentation/Source/Engine/Blueprints/UserGuide/EventGraph/BlueprintEventGraph.INT.udn
--- a/Engine/Documentation/Source/Engine/Blueprints/UserGuide/EventGraph/BlueprintEventGraph.INT.udn
+++ b/Engine/Documentation/Source/Engine/Blueprints/UserGuide/EventGraph/BlueprintEventGraph.INT.udn
@@ -2,6 +2,7 @@ Availability:Public
 Title:EventGraph
 Crumbs: %ROOT%, Engine, Engine/Blueprints, Engine/Blueprints/UserGuide
 Description:Uses events and function calls to perform actions in response to Blueprint events.
+version: 4.9
 
 
 [VAR:Topic]
@@ -52,19 +53,19 @@ Description:Uses events and function calls to perform actions in response to Blu
 
 [/EXCERPT:Intro]
 
-The **EventGraph** of a _Level Blueprint_ contains a node graph that uses events and function calls to perform actions in response 
+The **EventGraph** of a Level Blueprint contains a node graph that uses events and function calls to perform actions in response 
 to gameplay events. This is used to handle events for the level as a whole and to add functionality for specific instances of 
-_Actors_ and _Blueprints_ within the world. 
+Actors and Blueprints within the world. 
 
-In either case, an **EventGraph** is used by adding one or more events to act as entry points and then connecting _Function 
-Calls_, _Flow Control_ nodes, and _Variables_ to perform the desired actions.
+In either case, an **EventGraph** is used by adding one or more events to act as entry points and then connecting Function 
+Calls, Flow Control nodes, and Variables to perform the desired actions.
 
 ## Working with Graphs
 
-The **Graph** panel displays the visual representation of a particular graph of nodes as it shows all of the nodes contained in the graph as well as the connections between them. It provides editing capabilities for adding and removing nodes, arranging nodes, and creating links between nodes. Breakpoints can also be set in the **Graph** panel to aid in debugging Blueprints. 
+The **Graph** tab displays the visual representation of a particular graph of nodes as it shows all of the nodes contained in the graph as well as the connections between them. It provides editing capabilities for adding and removing nodes, arranging nodes, and creating links between nodes. Breakpoints can also be set in the **Graph** tab to aid in debugging Blueprints. 
 
 See the [](Engine/Blueprints/Editor/UIComponents/GraphEditor) for a detailed guide to working with the **EventGraph** and other **Graphs** with Blueprints.
 
-[REGION:tip]
+[REGION:note]
 For information about using Matinee to control specific types of Blueprint components and properties, and to create events, see [](Engine/Matinee\UserGuide\BlueprintComponents).
 [/REGION]
