Fixing a grammatical error on the StatCommand Documentation page

diff --git a/Engine/Documentation/Source/Engine/Performance/StatCommands/StatCommands.INT.udn b/Engine/Documentation/Source/Engine/Performance/StatCommands/StatCommands.INT.udn
--- a/Engine/Documentation/Source/Engine/Performance/StatCommands/StatCommands.INT.udn
+++ b/Engine/Documentation/Source/Engine/Performance/StatCommands/StatCommands.INT.udn
@@ -45,7 +45,7 @@ Note: removed until this can be checked against profiler info. StatsViewer didn'
 | **stat StartFile** **stat StopFile** | Generates a .ue4stats file that can be analyzed offline (With the StatsViewer tool from: UE4\Engine\Binaries\DotNET\StatsViewer). A GPU stall can appear as high workload the frame before.
 
 [EXCERPT:statSlow]
-This command will display stats for the game thread and the render thread. All stats will rendered as one big stats group. This command can be useful, if we don’t have access to the profiler or the log file, or we want to test the basic performance of our game.
+This command will display stats for the game thread and the render thread. All stats will rendered as one big stats group. This command can be useful, if we don't have access to the profiler or the log file, or we want to test the basic performance of our game.
 
 Here is the example in the running game:
 
diff --git a/Engine/Documentation/Source/Engine/Rendering/PostProcessEffects/ColorGrading/ColorGrading.INT.udn b/Engine/Documentation/Source/Engine/Rendering/PostProcessEffects/ColorGrading/ColorGrading.INT.udn
--- a/Engine/Documentation/Source/Engine/Rendering/PostProcessEffects/ColorGrading/ColorGrading.INT.udn
+++ b/Engine/Documentation/Source/Engine/Rendering/PostProcessEffects/ColorGrading/ColorGrading.INT.udn
@@ -122,6 +122,6 @@ Various color corrections are possible. Here are a few examples:
 * Higher quality contrast (e.g. curves with a steeper linear part in the middle)
 * Selective changes to the darks, midtones, and bright areas of an image (e.g. curves)
 * Selective changes to specific colors (best expressed in some color space where the brightness is in a separate channel e.g. LAB)
-* Adjustments can be even done in a different color space (e.g. LAB keeps brightness and color independent which can be very useful) 
+* Adjustments can even be done in a different color space (e.g. LAB keeps brightness and color independent which can be very useful) 
 
 [/EXCERPT:LookupTexture]
