#UE4doc #docs:
Added engine version field
Replaced a couple images due to outdated UI
Changed font settings for multiple terms
Changed a tip region into a note due to the content of the region

diff --git a/Engine/Documentation/Source/Engine/Blueprints/UserGuide/Events/Events.INT.udn b/Engine/Documentation/Source/Engine/Blueprints/UserGuide/Events/Events.INT.udn
--- a/Engine/Documentation/Source/Engine/Blueprints/UserGuide/Events/Events.INT.udn
+++ b/Engine/Documentation/Source/Engine/Blueprints/UserGuide/Events/Events.INT.udn
@@ -2,6 +2,7 @@ Availability:Public
 Title:Events
 Crumbs: %ROOT%, Engine, Engine/Blueprints, Engine/Blueprints/UserGuide
 Description: Nodes that are called from gameplay code to begin execution of an individual network within the EventGraph.
+version: 4.9
 
 [VAR:Topic]
 [OBJECT:Topic]
@@ -48,9 +49,9 @@ Description: Nodes that are called from gameplay code to begin execution of an i
 
 [EXCERPT:Intro]
 **Events** are nodes that are called from gameplay code to begin execution of an individual network 
-within the **EventGraph**. They allow _Blueprints_ to perform a series of actions in response to 
+within the **EventGraph**. They allow Blueprints to perform a series of actions in response to 
 certain events that occur within the game, such as the game beginning, the level being reset, taking 
-damage, etc. These events can be accessed within _Blueprints_ in order to implement new functionality or 
+damage, etc. These events can be accessed within Blueprints in order to implement new functionality or 
 override or augment the default functionality. Any number of **Events** can be used within a 
 single **EventGraph**; though only one of each type may be used.
 [/EXCERPT:Intro]
@@ -62,7 +63,7 @@ single **EventGraph**; though only one of each type may be used.
 	event, you will need to string them together linearly.
 [/REGION]
 
-[REGION:tip]
+[REGION:note]
 For information about using Matinee to create events to be triggered at specific times during the playback of a cinematic sequence, see [](Engine/Matinee\UserGuide\BlueprintComponents).
 [/REGION]
 
@@ -90,7 +91,7 @@ This Blueprint Event node executes only on the server. For single player games,
 
 ![LevelReset.png](LevelReset.png)
 
-The LevelReset event sends out an execution signal when the level restarts. This is useful when you 
+The **Level Reset** event sends out an execution signal when the level restarts. This is useful when you 
 need something to be triggered once a level has reloaded, such as in a gaming situation when the player 
 has died but the level does not need to reload.
 
@@ -104,19 +105,19 @@ has died but the level does not need to reload.
 This event will execute when a number of conditions are met at the same time:
 
 * Collision response between the actors must allow Overlaps.
-* Both Actors that are to execute the event have **Generate Overlap Events** set to _true_.
+* Both Actors that are to execute the event have **Generate Overlap Events** set to true.
 * And finally, both Actors' collision starts overlapping; moving together or one is created overlapping the other.
 
-For more information on collision, see: [**Collision Responses**](Engine\Physics\Collision).
+For more information on collision, see: [Collision Responses](Engine\Physics\Collision).
 
 ![](BeginOverlapEX.png)
 
-_When this Blueprint Actor overlaps the Actor stored in the Player Actor variable, it will increment the Counter integer variable._
+_When this Blueprint Actor overlaps the Actor stored in the Player Actor variable, it will increment the Counter integer variable_.
 
 | Item | Description |
 | -------- | ----------- |
 | [REGION:tablesection]Output Pins[/REGION] ||
-| **Other Actor** | **Actor** - This is the **Actor** that is overlapping this Blueprint. |
+| **Other Actor** | Actor - This is the Actor that is overlapping this Blueprint. |
 [/EXCERPT:EventBeginOverlap]
 
 ## Event Actor End Overlap
@@ -126,10 +127,10 @@ _When this Blueprint Actor overlaps the Actor stored in the Player Actor variabl
 This event will execute when a number of conditions are met at the same time:
 
 * Collision response between the actors must allow Overlaps.
-* Both Actors that are to execute the event have **Generate Overlap Events** set to _true_.
+* Both Actors that are to execute the event have **Generate Overlap Events** set to true.
 * And finally, both Actors' collision stop overlapping; moving apart or if one is destroyed.
 
-For more information on collision, see: [**Collision Responses**](Engine\Physics\Collision).
+For more information on collision, see: [Collision Responses](Engine\Physics\Collision).
 
 ![](EndOverlapEX.png)
 
@@ -138,30 +139,30 @@ _When this Blueprint Actor stops overlapping any other Actor except the Actor st
 | Item | Description |
 | -------- | ----------- |
 | [REGION:tablesection]Output Pins[/REGION] ||
-| **Other Actor** | **Actor** - This is the **Actor** that is overlapping this Blueprint. |
+| **Other Actor** | Actor - This is the Actor that is overlapping this Blueprint. |
 [/EXCERPT:EventEndOverlap]
 
 ## Event Hit
 
 ![](EventHit.png)
 
-This event will only execute if the collision settings on both Actors involved in a collision have **Simulation Generates Hit Events** set to _true_.
+This event will only execute if the collision settings on both Actors involved in a collision have **Simulation Generates Hit Events** set to true.
 
 | Item | Description |
 | -------- | ----------- |
 | [REGION:tablesection]Output Pins[/REGION] ||
-| **My Comp** | **PrimitiveComponent** - The Component on the executing Actor that was hit. |
-| **Other** | **Actor** - The other Actor involved in the collision. |
-| **Other Comp** | **PrimitiveComponent** - The component on the other Actor involved in the collision that was hit. |
-| **Self Moved** | **Boolean** - Used when receiving a hit from another object's movement (if false), the directions of Hit Normal and Hit Impact Normal will be adjusted to indicate force from the other object against the object being hit. |
-| **Hit Location** | **Vector** - The location of contact between the two colliding Actors. |
-| **Hit Normal** | **Vector** - The direction of the collision. |
-| **Normal Impulse** | **Vector** - The force that the Actors collided with. |
-| **Hit** | **Struct** `HitResult` - All the data collected in a Hit, you can pull off and "break" this result to gain access to the individual bits of data. |
+| **My Comp** | PrimitiveComponent - The Component on the executing Actor that was hit. |
+| **Other** | Actor - The other Actor involved in the collision. |
+| **Other Comp** | PrimitiveComponent - The component on the other Actor involved in the collision that was hit. |
+| **Self Moved** | Boolean - Used when receiving a hit from another object's movement (if false), the directions of Hit Normal and Hit Impact Normal will be adjusted to indicate force from the other object against the object being hit. |
+| **Hit Location** | Vector - The location of contact between the two colliding Actors. |
+| **Hit Normal** | Vector - The direction of the collision. |
+| **Normal Impulse** | Vector - The force that the Actors collided with. |
+| **Hit** | Struct HitResult - All the data collected in a Hit, you can pull off and "break" this result to gain access to the individual bits of data. |
 
 ![](EventHitEX.png)
 
-_In this example, when this Blueprint executes_ `Hit` _it will spawn an explosion effect at the point of impact._
+_In this example, when this Blueprint executes **Hit** it will spawn an explosion effect at the point of impact._
 
 ## Event Any Damage
 
@@ -174,10 +175,10 @@ This event is passed along when general damage is to be dealt. Like drowning or
 | Item | Description |
 | -------- | ----------- |
 | [REGION:tablesection]Output Pins[/REGION] ||
-| **Damage** | **Float** - The amount of damage being passed into the Actor. |
-| **Damage Type** | **Object 'DamageType'** - This is the object that contains additional data on the Damage being dealt. |
-| **Instigated By** | **Actor** - The Actor that is responsible for the damage. This would be the Actor that fired a gun or threw the grenade to deal the damage. |
-| **Damage Causer** | **Actor** - The Actor that caused the damage. This would be like a bullet or explosion. |
+| **Damage** | Float - The amount of damage being passed into the Actor. |
+| **Damage Type** | Object DamageType - This is the object that contains additional data on the Damage being dealt. |
+| **Instigated By** | Actor - The Actor that is responsible for the damage. This would be the Actor that fired a gun or threw the grenade to deal the damage. |
+| **Damage Causer** | Actor - The Actor that caused the damage. This would be like a bullet or explosion. |
 
 [REGION:lightbox]
 [![](AnyDamageEX.png)(w:940)](AnyDamageEX.png)
@@ -191,20 +192,20 @@ _Here if the damage being dealt to the Actor is coming from Water, it will subtr
 
 ![](PointDamage.png)
 
-Point Damage is meant to represent damage dealt by projectiles, hit scan weapons, or even melee weaponry.
+**Point Damage** is meant to represent damage dealt by projectiles, hit scan weapons, or even melee weaponry.
 
 | Item | Description |
 | -------- | ----------- |
 | [REGION:tablesection]Output Pins[/REGION] ||
-| **Damage** | **Float** - The amount of damage being passed into the Actor. |
-| **Damage Type** | **Object** `DamageType` - This is the object that contains additional data on the Damage being dealt. |
-| **Hit Location** | **Vector** - The location of where the damage is being applied. |
-| **Hit Normal** | **Vector** - The direction of the collision. |
-| **Hit Component** | **PrimitiveComponent** - The Component on the executing Actor that was hit. |
-| **Bone Name** | **Name** - The name of the bone that was hit. |
-| **Shot from Direction** | **Vector** - The direction the damage came from. |
-| **Instigated By** | **Actor** - The Actor that is responsible for the damage. This would be the Actor that fired a gun or threw the grenade to deal the damage. |
-| **Damage Causer** | **Actor** - The Actor that caused the damage. This would be like a bullet or explosion. |
+| **Damage** | Float - The amount of damage being passed into the Actor. |
+| **Damage Type** | Object DamageType - This is the object that contains additional data on the Damage being dealt. |
+| **Hit Location** | Vector - The location of where the damage is being applied. |
+| **Hit Normal** | Vector - The direction of the collision. |
+| **Hit Component** | PrimitiveComponent - The Component on the executing Actor that was hit. |
+| **Bone Name** | Name - The name of the bone that was hit. |
+| **Shot from Direction** | Vector - The direction the damage came from. |
+| **Instigated By** | Actor - The Actor that is responsible for the damage. This would be the Actor that fired a gun or threw the grenade to deal the damage. |
+| **Damage Causer** | Actor - The Actor that caused the damage. This would be like a bullet or explosion. |
 
 [REGION:lightbox]
 [![](PointDamageEX.png)(w:940)](PointDamageEX.png)
@@ -218,17 +219,17 @@ _In this example, when any damage is received, the damage dealt is subtracted fr
 
 ![](RadialDamage.png)
 
-The ReceiveRadiusDamage Actor is called whenever the parent Actor for this sequence receives radial damage. This is useful for handling events based on explosion damage, or damage caused indirectly.
+The **Radial Damage** Event is called whenever the parent Actor for this sequence receives radial damage. This is useful for handling events based on explosion damage, or damage caused indirectly.
 
 | Item | Description |
 | -------- | ----------- |
 | [REGION:tablesection]Output Pins[/REGION] ||
-| **Damage Received** | **Float** - The amount of damage received from the event. |
-| **Damage Type** | **Object** `DamageType` - This is the object that contains additional data on the Damage being dealt.  |
-| **Origin** | **Vector** - This gives the location in 3D space where the damage originated.|
-| **Hit Info** | **Struct** `HitResult` - All the data collected in a Hit, you can pull off and "break" this result to gain access to the individual bits of data.  |
-| **Instigated By** | **Controller** - The Controller (AI or Player) that instigated the damage. |
-| **Damage Causer** | **Actor** - The Actor that caused the damage. Could be a bullet, rocket, laser, or a character punching another. |
+| **Damage Received** | Float - The amount of damage received from the event. |
+| **Damage Type** | Object DamageType - This is the object that contains additional data on the Damage being dealt.  |
+| **Origin** | Vector - This gives the location in 3D space where the damage originated.|
+| **Hit Info** | Struct HitResult - All the data collected in a Hit, you can pull off and "break" this result to gain access to the individual bits of data.  |
+| **Instigated By** | Controller - The Controller (AI or Player) that instigated the damage. |
+| **Damage Causer** | Actor - The Actor that caused the damage. Could be a bullet, rocket, laser, or a character punching another. |
 
 [REGION:lightbox]
 [![](RadialDamageEX.png)(w:940)](RadialDamageEX.png)
@@ -280,7 +281,7 @@ _Once this Actor is no longer in the World, a String will output to indicate the
 | Item | Description |
 | -------- | ----------- |
 | [REGION:tablesection]Output Pins[/REGION] ||
-| **End Play Reason** | **enum** `EEndPlayReason` - An enum indicating the reason for why Event End Play is being called.|
+| **End Play Reason** | enum EEndPlayReason - An enum indicating the reason for why Event End Play is being called.|
 [/EXCERPT:EventEnd]
 
 ## Event Destroyed
@@ -306,7 +307,7 @@ This is a simple event that is called on every frame of gameplay.
 | Item | Description |
 | -------- | ----------- |
 | [REGION:tablesection]Output Pins[/REGION] ||
-| **Delta Seconds** | This outputs the amount of time between frames. |
+| **Delta Seconds** | Float - This outputs the amount of time between frames. |
 
 [![](TickEX.png)(w:940)](TickEX.png)
 
@@ -326,8 +327,8 @@ This is a specialized event that enables Blueprints to draw to the HUD. The HUD
 | Item | Description |
 | -------- | ----------- |
 | [REGION:tablesection]Output Pins[/REGION] ||
-| **Size X** | **Int** - The width of the render window in pixels. |
-| **Size Y** | **Int** - The height of the render window in pixels.  |
+| **Size X** | Int - The width of the render window in pixels. |
+| **Size Y** | Int - The height of the render window in pixels.  |
 
 
 [REGION:lightbox]
@@ -341,4 +342,4 @@ _This will create a clickable **Hit Box** in the center of the render window, wi
 
 ![Custom Event Node](Engine\Blueprints\UserGuide\Events\Custom\name_custom_event.png)
 
-The Custom Event node is a specialized node with its own workflow, see: [**Custom Event Nodes**](Engine\Blueprints\UserGuide\Events\Custom) for more information.
+The Custom Event node is a specialized node with its own workflow, see: [Custom Event Nodes](Engine\Blueprints\UserGuide\Events\Custom) for more information.
