#UE4doc #docs:
--WIP--
Submitting changes for the weekend

diff --git a/Engine/Documentation/Source/Engine/Blueprints/UserGuide/Events/Custom/CustomEvents.INT.udn b/Engine/Documentation/Source/Engine/Blueprints/UserGuide/Events/Custom/CustomEvents.INT.udn
--- a/Engine/Documentation/Source/Engine/Blueprints/UserGuide/Events/Custom/CustomEvents.INT.udn
+++ b/Engine/Documentation/Source/Engine/Blueprints/UserGuide/Events/Custom/CustomEvents.INT.udn
@@ -2,6 +2,7 @@ Availability:Public
 Title:Custom Events
 Crumbs: %ROOT%, Engine, Engine/Blueprints, Engine/Blueprints/UserGuide
 Description: Custom user-created events that can be fired off from within a Graph.
+version: 4.9
 
 [VAR:Topic]
 [OBJECT:Topic]
@@ -46,7 +47,6 @@ Description: Custom user-created events that can be fired off from within a Grap
 
 [TOC(start:2 end:3)]
 
-## Overview
 
 [EXCERPT:Intro]
 Just like **Events**, **Custom Events** have an output pin for execution and optional output data pins.  However, they are created by the user,
@@ -59,13 +59,13 @@ call or via the `CE` or `KE` console commands.
 
 ##Custom Events
 
-_Custom Events_ provide you with a way to create your own events that can be called at any point in 
-your _Blueprint_'s sequence. In cases where you are connecting multiple output execution wires to the input execution pin 
-of one particular node, _Custom Events_ can simplify the wire network of your graph. _Custom Events_ can even be set up in one graph of a _Blueprint_ and called in another graph.
+Custom Events provide you with a way to create your own events that can be called at any point in 
+your Blueprint's sequence. In cases where you are connecting multiple output execution wires to the input execution pin 
+of one particular node, Custom Events can simplify the wire network of your graph. Custom Events can even be set up in one graph of a Blueprint and called in another graph.
 
 They have a fairly straightforward workflow:
 
-* Create and name the _Custom Event_.
+* Create and name the Custom Event.
 * Establish any input parameters the event should have, along with any default values.
 * Create a special function node that calls the custom event.
 * Feed in any input parameters required.
@@ -73,7 +73,7 @@ They have a fairly straightforward workflow:
 
 ### Creating Custom Events
 
-1. Create a _Custom Event_ node by **right-clicking** and choosing **Add Custom Event...** from the context menu. 
+1. Create a Custom Event node by **right-clicking** and choosing **Add Custom Event...** from the context menu. 
 
 	![](add_custom_event.png)
 
@@ -91,7 +91,7 @@ To add input parameters:
 
 	![](new_input_custom_event.png)
 
-1. Name the new input and set its type using the dropdown menu.  In this example, there is a _String_ input parameter named **MyStringParam**. 
+1. Name the new input and set its type using the dropdown menu.  In this example, there is a String input parameter named **MyStringParam**. 
 
 	![](named_new_variable.png)
 
@@ -102,23 +102,23 @@ To add input parameters:
 To change the location of the pin for this parameter on the edge of the node, use the up and down arrows in the expanded **Details** pane entry.
 
 
-Now, just like any other _Event_ or execution node, you can attach other nodes to the output execution pin of your _Custom Event_, and execution of that network will begin when your _Custom Event_ is triggered.
-This _Custom Event_ example prints a _String_ to the screen.
+Now, just like any other Event or execution node, you can attach other nodes to the output execution pin of your Custom Event, and execution of that network will begin when your Custom Event is triggered.
+This Custom Event example prints a String to the screen.
 
 ![](small_custom_event.png)
 
 ### Calling Custom Events
 
-Your _Custom Event_ and its associated network have been created, but unlike regular _Events_, there are no preset conditions for the _Custom Event_ to fire. To call your _Custom Event_, **right-click** and choose **Call Function > [Custom Event Name]** from the  context menu. 
+Your Custom Event and its associated network have been created, but unlike regular Events, there are no preset conditions for the Custom Event to fire. To call your Custom Event, **right-click** and choose **Call Function > [Custom Event Name]** from the  context menu. 
 
 ![](call_ce_context_menu.png)
 
 ![](call_custom_event_node.png)
 
-Any input parameters that were set up on the _Custom Event_ will appear as input data pins in the new node so that they can be passed into the _Custom Event_. Connect any input data pins with data wires to variables or other data pins as needed.
+Any input parameters that were set up on the Custom Event will appear as input data pins in the new node so that they can be passed into the Custom Event. Connect any input data pins with data wires to variables or other data pins as needed.
 
 [REGION:tip]
-Unlike regular _Events_, which can only be called once per graph per _Event_ type, you can call a _Custom Event_ multiple times throughout a graph.  In this way, _Custom Events_ allow branching of multiple
+Unlike regular Events, which can only be called once per graph per Event type, you can call a Custom Event multiple times throughout a graph.  In this way, Custom Events allow branching of multiple
 execution outputs to a single execution input without needing to directly connect wires. 
 [/REGION]
 
@@ -126,20 +126,20 @@ execution outputs to a single execution input without needing to directly connec
 ![](connected_custom_event.png)
 [/REGION:fullwidth]
 
-In this example, if the **IsSuccess** boolean variable is _false_, an error message will be printed.  This graph has the same function as wiring the **Print String** node in sequence after the **Branch** node,
+In this example, if the **IsSuccess** boolean variable is false, an error message will be printed.  This graph has the same function as wiring the **Print String** node in sequence after the **Branch** node,
 but with the added functionality that other sections of the graph can utilize the **Print String** node, and the two network sections do not have to be located near each other in the graph.
 
 ### Troubleshooting Custom Events
 
-If you get a **Warning!** bar on your _Custom Event_ node with an **"Unable to find function with name [CustomEvent]"** message, **Compile** your _Blueprint_. 
+If you get a **Warning!** bar on your Custom Event node with an **"Unable to find function with name [CustomEvent]"** message, **Compile** your Blueprint. 
 
 ![](custom_event_warning.png)
 
-If you change the number of input parameters on your _Custom Event_, all nodes that call the _Custom Event_ will have an error when you compile your _Blueprint_. 
+If you change the number of input parameters on your Custom Event, all nodes that call the Custom Event will have an error when you compile your Blueprint. 
 
 ![](parameter_change_error.png)
 
-You must refresh all the nodes that call your _Custom Event_. To refresh an individual node or a selected group of nodes, **right-click** on the node(s) and select **Refresh Nodes**.
+You must refresh all the nodes that call your Custom Event. To refresh an individual node or a selected group of nodes, **Right-click** on the node(s) and select **Refresh Nodes**.
 
 ![](refresh_nodes.png)
 
