#UE4doc #docs - Added Engine Version Field. Bolded/unbolded text where needed.

diff --git a/Engine/Documentation/Source/Engine/Rendering/Materials/Editor/MaterialEditor.INT.udn b/Engine/Documentation/Source/Engine/Rendering/Materials/Editor/MaterialEditor.INT.udn
--- a/Engine/Documentation/Source/Engine/Rendering/Materials/Editor/MaterialEditor.INT.udn
+++ b/Engine/Documentation/Source/Engine/Rendering/Materials/Editor/MaterialEditor.INT.udn
@@ -2,6 +2,7 @@ Availability: Public
 Title:Material Editor Reference
 Crumbs: %ROOT%, Engine, Engine/Rendering/Materials
 Description:Guide to using the Material Editor for creating shaders.
+Version: 4.9
 
 [VAR:Topic]
 [OBJECT:Topic]
@@ -44,10 +45,10 @@ For more information about the Material Editor UI, see [](Engine/Rendering/Mater
 ## Opening the Material Editor
 
 The Material Editor can be opened by **double-clicking** any Material asset or through the **right-click** context menu of a Material 
-asset in the **Content Browser**. Either of these will open up that particular Material in the Material Editor for editing. 
+asset in the Content Browser. Either of these will open up that particular Material in the Material Editor for editing. 
 
-The Material Editor is also opened whenever you create a new Material in the **Content Browser**, which is done either by **right-clicking** 
-and choosing New Material from the context menu, or clicking the **New** button at the top of the **Content Browser** and choosing 
+The Material Editor is also opened whenever you create a new Material in the Content Browser, which is done either by **right-clicking** 
+and choosing New Material from the context menu, or clicking the **Add New** button at the top of the **Content Browser** and choosing 
 **Material** from the drop-down list.
 
 
@@ -73,7 +74,7 @@ them will allow you to stretch a frame around a group of multiple nodes, making
 sections of a network, rather than using the **Desc** property to describe a single node at a time. 
 
 You can create a Comment at any time by pressing **C** while hovering the mouse over the Graph. You can also **right-click** in 
-any blank area of the Graph and choose **New Comment**.
+any blank area of the Graph and choose **Add Comment** from the context menu.
 
 The comments themselves are displayed as banners of text along the top of the comment box. This text will not scale 
 with the Graph's zoom level, so even when zoomed far away, you will still be able to read the description.
@@ -97,10 +98,10 @@ Comments can be renamed by selecting the comment and then modifying the **Text**
 When making changes to a network, it can be useful to get immediate feedback from each change in realtime. The Material Editor 
 offers two features in the Graph to give you instantaneous feedback for your network: Live Nodes and Live Update.
 
-The difference between these two options may not be apparent at first. Live Nodes allows for any constant changes that the nodes 
+The difference between these two options may not be apparent at first. **Live Nodes** allows for any constant changes that the nodes 
 make - due to such expressions as Panners - to playback in real time. 
 
-Live Update, on the other hand, allows for each node within the network to have its shader recompiled whenever a change is made. 
+**Live Update**, on the other hand, allows for each node within the network to have its shader recompiled whenever a change is made. 
 These changes include creation of new nodes, deletion of nodes, node connections and disconnections, and changes in properties. 
 This recompilation is necessary so that the material preview drawn at that node is up to date. However, recompiling these 
 intermediate shaders can be time-consuming, especially if your Material contains many nodes.  If you are experiencing long wait 
