#UE4doc #doc - Added Engine version field. Minor changes with bolding where needed.

diff --git a/Engine/Documentation/Source/Engine/Rendering/Materials/Editor/Interface/MaterialEditorInterface.INT.udn b/Engine/Documentation/Source/Engine/Rendering/Materials/Editor/Interface/MaterialEditorInterface.INT.udn
--- a/Engine/Documentation/Source/Engine/Rendering/Materials/Editor/Interface/MaterialEditorInterface.INT.udn
+++ b/Engine/Documentation/Source/Engine/Rendering/Materials/Editor/Interface/MaterialEditorInterface.INT.udn
@@ -2,6 +2,7 @@ Availability: Public
 Title:Material Editor UI
 Crumbs: %ROOT%, Engine, Engine/Rendering/Materials, Engine/Rendering/Materials/Editor
 Description:Guide to using the Material Editor UI.
+Version: 4.9
 
 
 [TOC(start:2)]
@@ -150,7 +151,7 @@ For more information about material functions, see [Material Functions](Engine/R
 
 ![](MaterialEditor_HLSLPanel.png)
 
-This panel shows the High Level Shader Language (HLSL) code that is generated by the current Material. Be 
+This panel shows the **High Level Shader Language** (HLSL) code that is generated by the current Material. Be 
 aware that this is not an editor; you will not be able to change the HLSL code. Rather, this is merely a viewer to 
 allow you to visualize the code defined by your material network.
 
