#UE4doc #docs - Added Engine version Field. Removed Overview Header.

diff --git a/Engine/Documentation/Source/Engine/Rendering/Materials/ExpressionReference/ExpressionReference.INT.udn b/Engine/Documentation/Source/Engine/Rendering/Materials/ExpressionReference/ExpressionReference.INT.udn
--- a/Engine/Documentation/Source/Engine/Rendering/Materials/ExpressionReference/ExpressionReference.INT.udn
+++ b/Engine/Documentation/Source/Engine/Rendering/Materials/ExpressionReference/ExpressionReference.INT.udn
@@ -2,6 +2,7 @@ Availability: Public
 title:Material Expression Reference
 Crumbs: %ROOT%, Engine, Engine/Rendering/Materials
 Description:Reference for all material expressions available for use in the Material Editor for creating complex node-based shader networks.
+Version: 4.9
 
 [REGION:banner]
 ![Materials Compendium](compendium.png)
@@ -9,7 +10,6 @@ Description:Reference for all material expressions available for use in the Mate
 
 [TOC (start:2 end:3)]
 
-## Overview
 
 This page is a reference for all material expression nodes available in the Material Editor. Material Expressions are the building blocks used to create fully functional materials in Unreal Engine 4. Each material expression is a self-contained black box that either outputs a set of one or more specific values or performs a single operation on one or more inputs and then outputs the results of that operation.
 
