#UE4doc #docs - Added Engine Version Field. Bolded term.

diff --git a/Engine/Documentation/Source/Engine/Rendering/Materials/ExpressionReference/Atmosphere/Atmosphere.INT.udn b/Engine/Documentation/Source/Engine/Rendering/Materials/ExpressionReference/Atmosphere/Atmosphere.INT.udn
--- a/Engine/Documentation/Source/Engine/Rendering/Materials/ExpressionReference/Atmosphere/Atmosphere.INT.udn
+++ b/Engine/Documentation/Source/Engine/Rendering/Materials/ExpressionReference/Atmosphere/Atmosphere.INT.udn
@@ -2,6 +2,7 @@ Availability:Public
 Title:Atmosphere Expressions
 Crumbs: %ROOT%, Engine, Engine/Rendering/Materials, Engine/Rendering/Materials/ExpressionReference
 Description:Expressions that affect fog and other atmospheric level effects.
+Version:4.9
 
 [EXCERPT:contents]
 [TOC(start:2 end:2)]
@@ -11,7 +12,7 @@ Description:Expressions that affect fog and other atmospheric level effects.
 
 ## Atmospheric Fog Color
 
-The Atmospheric Fog Material Expression allows you to query the current color of the level's Atmospheric Fog, at any position in World Space. If no World Position is fed into it, then the world position of the pixel in question is used. This is useful when you need Materials to appear to fade into a distant fog color.
+The **Atmospheric Fog Color** Material Expression allows you to query the current color of the level's Atmospheric Fog, at any position in World Space. If no World Position is fed into it, then the world position of the pixel in question is used. This is useful when you need Materials to appear to fade into a distant fog color.
 
 In the example below, The Base Color is being set using an Atmospheric Fog node, with the World Position receiving a simple network that queries the location that is always 50,000 units behind the object, relative to the position of the camera.
 
