#UE4doc #docs - Added Engine Version Field.

diff --git a/Engine/Documentation/Source/Engine/Rendering/Materials/ExpressionReference/Color/ColorExpressions.INT.udn b/Engine/Documentation/Source/Engine/Rendering/Materials/ExpressionReference/Color/ColorExpressions.INT.udn
--- a/Engine/Documentation/Source/Engine/Rendering/Materials/ExpressionReference/Color/ColorExpressions.INT.udn
+++ b/Engine/Documentation/Source/Engine/Rendering/Materials/ExpressionReference/Color/ColorExpressions.INT.udn
@@ -2,6 +2,7 @@ Availability:Public
 Title:Color Expressions
 Crumbs: %ROOT%, Engine, Engine/Rendering/Materials, Engine/Rendering/Materials/ExpressionReference
 Description:Expressions that perform actions on color inputs.
+Version: 4.9
 
 [EXCERPT:TOC]
 [TOC(start:2)]
