#UE4doc #docs - Added Engine Version Field. Bolded Terms where needed.

diff --git a/Engine/Documentation/Source/Engine/Rendering/Materials/ExpressionReference/Constant/ConstantExpressions.INT.udn b/Engine/Documentation/Source/Engine/Rendering/Materials/ExpressionReference/Constant/ConstantExpressions.INT.udn
--- a/Engine/Documentation/Source/Engine/Rendering/Materials/ExpressionReference/Constant/ConstantExpressions.INT.udn
+++ b/Engine/Documentation/Source/Engine/Rendering/Materials/ExpressionReference/Constant/ConstantExpressions.INT.udn
@@ -2,6 +2,7 @@ Availability:Public
 Title:Constant Expressions
 Crumbs: %ROOT%, Engine, Engine/Rendering/Materials, Engine/Rendering/Materials/ExpressionReference
 Description:Expressions where the values of the outputs generally do not change once set in the editor or when play begins.
+Version: 4.9
 
 [EXCERPT:TOC]
 [TOC(start:2)]
@@ -11,7 +12,7 @@ Description:Expressions where the values of the outputs generally do not change
 
 ## Constant
 
-The Constant expression outputs a single float value. It is one of the most commonly used expressions and can be connected to any input, regardless of the number of channels the input expects. For instance, if you connect a Constant to an input expecting a 3 Vector, the constant value will be used for all 3 elements. When supplying a single number, it can be useful to collapse the node using the small triangle icon in the description area.
+The **Constant** expression outputs a single float value. It is one of the most commonly used expressions and can be connected to any input, regardless of the number of channels the input expects. For instance, if you connect a Constant to an input expecting a 3 Vector, the constant value will be used for all 3 elements. When supplying a single number, it can be useful to collapse the node using the small triangle icon in the description area.
 
 | Item | Description |
 | -------- | ----------- |
@@ -24,13 +25,13 @@ The Constant expression outputs a single float value. It is one of the most comm
 ![ConstantExample.png](ConstantExample.png)
 
 [REGION:tip]
-You can quickly create a Constant node by holding the **1** key and clicking in the graph area of the Material Editor.
+You can quickly create a Constant node by holding the **1** key and **Left Mouse Click** in the graph area of the Material Editor.
 [/REGION]
 
 [EXCERPT:Constant2vector]
 ## Constant2Vector
 
-The Constant2Vector expression outputs a two-channel vector value, in other words, two constant numbers.
+The **Constant2Vector** expression outputs a two-channel vector value, in other words, two constant numbers.
 
 | Item | Description |
 | -------- | ----------- |
@@ -47,7 +48,7 @@ The Constant2Vector expression outputs a two-channel vector value, in other word
 [/REGION]
 
 [REGION:tip]
-You can quickly create a Constant2Vector node by holding the **2** key and clicking in the graph area of the Material Editor.
+You can quickly create a Constant2Vector node by holding the **2** key and **Left Mouse Click** in the graph area of the Material Editor.
 [/REGION]
 
 [/EXCERPT:Constant2vector]
@@ -55,7 +56,7 @@ You can quickly create a Constant2Vector node by holding the **2** key and click
 [EXCERPT:Constant3vector]
 ## Constant3Vector
 
-The Constant3Vector expression outputs a three-channel vector value, in other words, three constants numbers.  An RGB color can be thought of as a Constant3Vector, where each channel is assigned to a color (red, green, blue).
+The **Constant3Vector** expression outputs a three-channel vector value, in other words, three constants numbers.  An RGB color can be thought of as a Constant3Vector, where each channel is assigned to a color (red, green, blue).
 
 | Item | Description |
 | -------- | ----------- |
@@ -72,7 +73,7 @@ The Constant3Vector expression outputs a three-channel vector value, in other wo
 [/REGION]
 
 [REGION:tip]
-You can quickly create a Constant3Vector node by holding the **3** key and clicking in the graph area of the Material Editor.
+You can quickly create a Constant3Vector node by holding the **3** key and **Left Mouse Click** in the graph area of the Material Editor.
 [/REGION]
 
 [/EXCERPT:Constant3vector]
@@ -80,7 +81,7 @@ You can quickly create a Constant3Vector node by holding the **3** key and click
 [EXCERPT:Constant4Vector]
 ## Constant4Vector
 
-The Constant4Vector expression outputs a four-channel vector value, in other words, four constants numbers.  An RGBA color can be thought of as a Constant4Vector, where each channel is assigned to a color (red, green, blue, alpha).
+The **Constant4Vector** expression outputs a four-channel vector value, in other words, four constants numbers.  An RGBA color can be thought of as a Constant4Vector, where each channel is assigned to a color (red, green, blue, alpha).
 
 | Item | Description |
 | -------- | ----------- |
@@ -95,7 +96,7 @@ The Constant4Vector expression outputs a four-channel vector value, in other wor
 ![Constant4Example.png](Constant4Example.png)
 
 [REGION:tip]
-You can quickly create a Constant4Vector node by holding the **4** key and clicking in the graph area of the Material Editor.
+You can quickly create a Constant4Vector node by holding the **4** key and **Left Mouse Click** in the graph area of the Material Editor.
 [/REGION]
 
 [/EXCERPT:Constant4Vector]
@@ -135,7 +136,7 @@ This network will cause the object to fade in (instead of "popping") once the ca
 
 ## PerInstanceFadeAmount
 
-The PerInstanceFadeAmount expression outputs a float value associated with the amount of fade applied to an instanced Static Mesh, such as foliage. It is constant, but can be a different number for each individual instance of a mesh.
+The **PerInstanceFadeAmount** expression outputs a float value associated with the amount of fade applied to an instanced Static Mesh, such as foliage. It is constant, but can be a different number for each individual instance of a mesh.
 
 [REGION:note]
 This only works when applied to an InstancedStaticMesh Actor or other Actor which utilizes InstancedStaticMeshComponents.
@@ -143,7 +144,7 @@ This only works when applied to an InstancedStaticMesh Actor or other Actor whic
 
 ## PerInstanceRandom
 
-The PerInstanceRandom expression outputs a different random float value per Static Mesh instance to which the material is applied. InstancedStaticMeshComponent sets a random float for instance, which is exposed so that it can be used for whatever is desired (e.g. random light level behind a window). It is constant, but different, for each instance of the mesh.
+The **PerInstanceRandom** expression outputs a different random float value per Static Mesh instance to which the material is applied. InstancedStaticMeshComponent sets a random float for instance, which is exposed so that it can be used for whatever is desired (e.g. random light level behind a window). It is constant, but different, for each instance of the mesh.
 
 The output value will be a whole number between 0 and RAND_MAX for the target platform.
 
@@ -159,7 +160,7 @@ This only works when applied to an InstancedStaticMesh Actor or other Actor whic
 (#Time)
 ## Time
 
-The Time node is used to add the passage of time to a material, such as a [Panner](Engine/Rendering/Materials/ExpressionReference/Coordinates#Panner), [Cosine](Engine/Rendering/Materials/ExpressionReference/Math#Cosine), or other time-dependent operation.
+The **Time** node is used to add the passage of time to a material, such as a [Panner](Engine/Rendering/Materials/ExpressionReference/Coordinates#Panner), [Cosine](Engine/Rendering/Materials/ExpressionReference/Math#Cosine), or other time-dependent operation.
 
 | Item | Description |
 | -------- | ----------- |
@@ -177,14 +178,14 @@ The network described in the image above would create a Material that would chan
 
 ## TwoSidedSign
 
-The TwoSidedSign expression is useful for flipping the normal on backfaces of two sided custom lighting materials to match the functionality of Phong. +1 for frontfaces, -1 for backfaces of a twosided material.
+The **TwoSidedSign** expression is useful for flipping the normal on backfaces of two sided custom lighting materials to match the functionality of Phong. +1 for frontfaces, -1 for backfaces of a twosided material.
 
 ![exp_twosidedsign.png](exp_twosidedsign.png)
 
 
 ## VertexColor
 
-The VertexColor expression is the access point for the material to the outputs of color modules affecting sprite particles emitters.
+The **VertexColor** expression is the access point for the material to the outputs of color modules affecting sprite particles emitters.
 
 | Item | Description |
 | -------- | ----------- |
diff --git a/Engine/Documentation/Source/Engine/Rendering/Materials/ExpressionReference/Coordinates/CoordinatesExpressions.INT.udn b/Engine/Documentation/Source/Engine/Rendering/Materials/ExpressionReference/Coordinates/CoordinatesExpressions.INT.udn
--- a/Engine/Documentation/Source/Engine/Rendering/Materials/ExpressionReference/Coordinates/CoordinatesExpressions.INT.udn
+++ b/Engine/Documentation/Source/Engine/Rendering/Materials/ExpressionReference/Coordinates/CoordinatesExpressions.INT.udn
@@ -2,6 +2,7 @@ Availability:Public
 Title:Coordinates Expressions
 Crumbs: %ROOT%, Engine, Engine/Rendering/Materials
 Description:Expressions that perform actions on texture coordinates or output values that can be used as, or to modify, texture coordinates.
+Version:4.9
 
 [EXCERPT:TOC]
 [TOC(start:2)]
@@ -16,7 +17,7 @@ Description:Expressions that perform actions on texture coordinates or output va
 
 ## LightmapUVs
 
-The LightmapUVs expression outputs the lightmap UV texture coordinates in the forum of a two-channel vector value. If lightmap UVs are unavailable, it will output a two-channel vector value of (0,0).
+The **LightmapUVs** expression outputs the lightmap UV texture coordinates in the forum of a two-channel vector value. If lightmap UVs are unavailable, it will output a two-channel vector value of (0,0).
 
 
 [INCLUDE:Engine/Rendering/Materials/ExpressionReference/Vector#ObjectOrientation]
@@ -27,7 +28,7 @@ The LightmapUVs expression outputs the lightmap UV texture coordinates in the fo
 
 ## ObjectRadius
 
-The object radius outputs a value equal to the radius of a given object in Unreal units. Scaling is taken into account and the results can be unique for each individual object.
+The **Object Radius** outputs a value equal to the radius of a given object in Unreal units. Scaling is taken into account and the results can be unique for each individual object.
 
 [REGION:fullwidth]
 ![ObjectRadiusExample.png](ObjectRadiusExample.png)
@@ -38,7 +39,7 @@ In this example, both meshes are receiving this material in which the ObjectRadi
 
 ## Panner
 
-The Panner expression outputs UV texture coordinates that can be used to create panning, or moving, textures.
+The **Panner** expression outputs UV texture coordinates that can be used to create panning, or moving, textures.
 
 | Item | Description |
 | -------- | ----------- |
@@ -65,7 +66,7 @@ Panner generates UVs that change according to the Time input.  The Coordinate in
 
 ## Rotator
 
-The Rotator expression outputs UV texture coordinates in the form of a two-channel vector value that can be used to create rotating textures.
+The **Rotator** expression outputs UV texture coordinates in the form of a two-channel vector value that can be used to create rotating textures.
 
 | Item | Description |
 | -------- | ----------- |
@@ -83,11 +84,11 @@ The Rotator expression outputs UV texture coordinates in the form of a two-chann
 
 ## SceneTexelSize
 
-The SceneTexelSize expression allows you to offset by texel sizes, as you would when using the SceneColor and SceneDepth expressions. This is useful for edge detection in multi-resolution systems, as without this calculation you would be forced to use a small static value, resulting in inconsistent results at lower resolutions.
+The **SceneTexelSize** expression allows you to offset by texel sizes, as you would when using the SceneColor and SceneDepth expressions. This is useful for edge detection in multi-resolution systems, as without this calculation you would be forced to use a small static value, resulting in inconsistent results at lower resolutions.
 
 ## ScreenPosition
 
-The ScreenPosition expression outputs the screen-space position of the pixel currently being rendered.  
+The **ScreenPosition** expression outputs the screen-space position of the pixel currently being rendered.  
 
 <!--| Item | Description |
 | -------- | ----------- |
@@ -100,7 +101,7 @@ The ScreenPosition expression outputs the screen-space position of the pixel cur
 
 ## TextureCoordinate
 
-The TextureCoordinate expression outputs UV texture coordinates in the form of a two-channel vector value allowing materials to use different UV channels, specify tiling, and otherwise operate on the UVs of a mesh.
+The **TextureCoordinate** expression outputs UV texture coordinates in the form of a two-channel vector value allowing materials to use different UV channels, specify tiling, and otherwise operate on the UVs of a mesh.
 
 | Item | Description |
 | -------- | ----------- |
@@ -123,7 +124,7 @@ The TextureCoordinate expression outputs UV texture coordinates in the form of a
 
 ## ViewSize
  
- The ViewSize expression outputs a 2D vector giving the size of the current view in pixels. This is useful for causing various changes in your materials based on the current resolution of the screen.
+ The **ViewSize** expression outputs a 2D vector giving the size of the current view in pixels. This is useful for causing various changes in your materials based on the current resolution of the screen.
 
  ![](ViewSizeNetwork.png)
 
@@ -140,7 +141,7 @@ The TextureCoordinate expression outputs UV texture coordinates in the form of a
 ## WorldPosition
 
 
-The WorldPosition expression outputs the position of the current pixel in world space. To visualize, simply plug the output into Emissive:
+The **WorldPosition** expression outputs the position of the current pixel in world space. To visualize, simply plug the output into Emissive:
 
 [REGION:fullwidth]
 ![WorldPositionExample.png](WorldPositionExample.png)
diff --git a/Engine/Documentation/Source/Engine/Rendering/Materials/ExpressionReference/Vector/VectorExpressions.INT.udn b/Engine/Documentation/Source/Engine/Rendering/Materials/ExpressionReference/Vector/VectorExpressions.INT.udn
--- a/Engine/Documentation/Source/Engine/Rendering/Materials/ExpressionReference/Vector/VectorExpressions.INT.udn
+++ b/Engine/Documentation/Source/Engine/Rendering/Materials/ExpressionReference/Vector/VectorExpressions.INT.udn
@@ -10,7 +10,7 @@ Description:Expressions that output a vector value, such as a position or normal
 [EXCERPT:ActorPositionWS]
 ## ActorPositionWS
 
-ActorPositionWS outputs Vector3 (RGB) data representing the location of the object with this material on it in world-space. 
+**ActorPositionWS** outputs Vector3 (RGB) data representing the location of the object with this material on it in world-space. 
 
 ![](ActorPositionColors.png)(w:434) ![](ActorPositionNetwork.png)(w:490)
 
@@ -23,7 +23,7 @@ In this example, you can see that ActorPositionWS is being fed directly into the
 ## CameraPositionWS
 
 
-The CameraWorldPosition expression outputs a three-channel vector value representing the camera's position in world space.
+The **CameraWorldPosition** expression outputs a three-channel vector value representing the camera's position in world space.
 
 The preview sphere changes color as the camera rotates.
 
@@ -33,7 +33,7 @@ The preview sphere changes color as the camera rotates.
 
 ## CameraVectorWS
 
-The CameraVector expression outputs a three-channel vector value representing the direction of the camera with respect to the surface, in other words, the direction from the pixel to the camera.
+The **CameraVector** expression outputs a three-channel vector value representing the direction of the camera with respect to the surface, in other words, the direction from the pixel to the camera.
 
 **Example Usage:** CameraVector is often used to fake environment maps by connecting the CameraVector to a ComponentMask and use the x and y channels of the CameraVector as texture coordinates.
 
@@ -71,7 +71,7 @@ The LightVector expression is a three-channel vector value representing the dire
 
 ## Object Bounds
 
-The Object Bounds expression outputs the size of the object in each axis. If used as color, the X, Y, and Z axes correspond to R, G, and B, respectively. 
+The **Object Bounds** expression outputs the size of the object in each axis. If used as color, the X, Y, and Z axes correspond to R, G, and B, respectively. 
 
 ![](ObjectBoundsScene.png)(w:411) ![](ObjectBoundsNetwork.png)
 
@@ -81,7 +81,7 @@ In the example above, you can see that the colors of the objects correspond to t
 ## ObjectOrientation
 
 
-The ObjectOrientation expression outputs the world-space up vector of the object. In other words, the object's local positive z-axis is pointing in this direction.
+The **ObjectOrientation** expression outputs the world-space up vector of the object. In other words, the object's local positive z-axis is pointing in this direction.
 
 ![ObjectOrientationExample.png](ObjectOrientationExample.png)
 [/EXCERPT:ObjectOrientation]
@@ -91,7 +91,7 @@ The ObjectOrientation expression outputs the world-space up vector of the object
 ## ObjectPositionWS
 
 
-The ObjectPositionWS expression outputs the world-space center position of the object's bounds.  For example, this is useful for creating spherical lighting for foliage.
+The **ObjectPositionWS** expression outputs the world-space center position of the object's bounds.  For example, this is useful for creating spherical lighting for foliage.
 
 
 ![](WorldPositionWSScene.png)(w:460) ![](WorldPositionNetwork.png)(w:470)
@@ -128,7 +128,7 @@ The ReflectionVectorWS expression is similar in spirit to [CameraVectorWS](#Came
 ## VertexNormalWS
 
 
-The VertexNormalWS expression outputs the world-space vertex normal. It can only be used in material inputs that are executed in the vertex shader, like WorldPositionOffset.  This is useful for making a mesh grow or shrink.  Note that offsetting position along the normal will cause the geometry to split apart along UV seams.
+The **VertexNormalWS** expression outputs the world-space vertex normal. It can only be used in material inputs that are executed in the vertex shader, like WorldPositionOffset.  This is useful for making a mesh grow or shrink.  Note that offsetting position along the normal will cause the geometry to split apart along UV seams.
 
 [REGION:fullwidth]
 ![VertexNormalWSExample.png](VertexNormalWSExample.png)
