#UE4doc #docs - Updated Engine Version field.

diff --git a/Engine/Documentation/Source/Engine/Rendering/Materials/ExpressionReference/Custom/CustomExpressions.INT.udn b/Engine/Documentation/Source/Engine/Rendering/Materials/ExpressionReference/Custom/CustomExpressions.INT.udn
--- a/Engine/Documentation/Source/Engine/Rendering/Materials/ExpressionReference/Custom/CustomExpressions.INT.udn
+++ b/Engine/Documentation/Source/Engine/Rendering/Materials/ExpressionReference/Custom/CustomExpressions.INT.udn
@@ -3,11 +3,11 @@ Title:Custom Expressions
 Crumbs: %ROOT%, Engine, Engine/Rendering/Materials, Engine/Rendering/Materials/ExpressionReference
 Description:Expressions that allow the use of custom, plain shader code.
 Skilllevel: Advanced
-Version: 4.8
+Version: 4.9
 
 ## Custom
 
-The Custom expression allows you to write custom HLSL shader code operating on an arbitrary amount of inputs and outputting the result of the operation.
+The **Custom** expression allows you to write custom HLSL shader code operating on an arbitrary amount of inputs and outputting the result of the operation.
 
 | Item | Description |
 | -------- | ----------- |
