#UE4doc #docs - Added Engine Version field. Replaced Image Table with Comparison Slider.

diff --git a/Engine/Documentation/Source/Engine/Rendering/Materials/ExpressionReference/Depth/DepthExpressions.INT.udn b/Engine/Documentation/Source/Engine/Rendering/Materials/ExpressionReference/Depth/DepthExpressions.INT.udn
--- a/Engine/Documentation/Source/Engine/Rendering/Materials/ExpressionReference/Depth/DepthExpressions.INT.udn
+++ b/Engine/Documentation/Source/Engine/Rendering/Materials/ExpressionReference/Depth/DepthExpressions.INT.udn
@@ -2,6 +2,7 @@ Availability:Public
 Title:Depth Expressions
 Crumbs: %ROOT%, Engine, Engine/Rendering/Materials, Engine/Rendering/Materials/ExpressionReference
 Description:Expressions that deal with depth of the pixel being rendered.
+Version: 4.9
 
 [EXCERPT:TOC]
 [TOC(start:2)]
@@ -11,7 +12,7 @@ Description:Expressions that deal with depth of the pixel being rendered.
 [EXCERPT:DepthFade]
 ## DepthFade
 
-The DepthFade expression is used to hide unsightly seams that take place when translucent objects intersect with opaque ones. 
+The **DepthFade** expression is used to hide unsightly seams that take place when translucent objects intersect with opaque ones. 
 
 | Item | Description |
 | -------- | ----------- |
@@ -19,13 +20,21 @@ The DepthFade expression is used to hide unsightly seams that take place when tr
 | **Fade Distance** | World space distance over which the fade should take place. This is used if the FadeDistance input is unconnected. | 
 | [REGION:tablesection]Inputs[/REGION] ||
 | **Opacity** | Takes in the existing opacity for the object prior to the depth fade. |
-| **FadeDistance** | World space distance over which the fade should take place.  |
+| **FadeDistance** | World space distance over which the fade should take place. |
+
+[OBJECT:ComparisonSlider]
+	[PARAM:before]
+	![Before Depth Fade](DepthFade0.png)(w:340)
+	[/PARAM]
+	[PARAM:after]
+	![After Depth Fade](DepthFade1.png)(w:340)
+	[/PARAM]
+[/OBJECT]
+
+| [INCLUDE:#DepthFadeNetwork] |
+| :---: | 
+| Network (Click for full size) |
 
-[REGION:imagetable]
-| ![](DepthFade0.png)(w:340) | ![](DepthFade1.png)(w:340) | [INCLUDE:#DepthFadeNetwork] |
-| :---: | :---: | :---: |
-| Before Depth Fade | After Depth Fade | Network (Click for full size) |
-[/REGION]
 
 <!--
 [EXCERPT:DepthFadeNetwork]
@@ -40,7 +49,7 @@ The DepthFade expression is used to hide unsightly seams that take place when tr
 
 ## PixelDepth
 
-The PixelDepth expression outputs the depth, or distance from the camera, of the pixel currently being rendered.  
+The **PixelDepth** expression outputs the depth, or distance from the camera, of the pixel currently being rendered.  
 
 
 
@@ -66,7 +75,7 @@ In this example, the material network has been applied to the floor. Notice how
 
 ## SceneDepth
 
-The SceneDepth expression outputs the existing scene depth.  This is similar to [PixelDepth](#PixelDepth), except that PixelDepth can sample the depth only at the pixel currently being drawn, whereas SceneDepth can sample depth at **any location**.
+The **SceneDepth** expression outputs the existing scene depth.  This is similar to [PixelDepth](#PixelDepth), except that PixelDepth can sample the depth only at the pixel currently being drawn, whereas SceneDepth can sample depth at **any location**.
 
 [REGION:warning]
 	Only translucent materials may utilize SceneDepth.
