#UE4doc #docs - Added Engine Version field. Bolded terms.

diff --git a/Engine/Documentation/Source/Engine/Rendering/Materials/ExpressionReference/Font/FontExpressions.INT.udn b/Engine/Documentation/Source/Engine/Rendering/Materials/ExpressionReference/Font/FontExpressions.INT.udn
--- a/Engine/Documentation/Source/Engine/Rendering/Materials/ExpressionReference/Font/FontExpressions.INT.udn
+++ b/Engine/Documentation/Source/Engine/Rendering/Materials/ExpressionReference/Font/FontExpressions.INT.udn
@@ -2,6 +2,8 @@ Availability:Public
 Title:Font Expressions
 Crumbs: %ROOT%, Engine, Engine/Rendering/Materials, Engine/Rendering/Materials/ExpressionReference
 Description:Expressions that sample and output font assets.
+Version: 4.9
+SkillLevel:
 
 [EXCERPT:TOC]
 [TOC(start:2)]
@@ -10,7 +12,7 @@ Description:Expressions that sample and output font assets.
 [EXCERPT:FontSample]
 ## FontSample
 
-The FontSample expression allows you to sample the texture pages out of a font resource as regular 2D textures. The alpha channel of the font will contain the font outline value.  Only valid font pages are allowed to be specified.
+The **FontSample** expression allows you to sample the texture pages out of a font resource as regular 2D textures. The alpha channel of the font will contain the font outline value.  Only valid font pages are allowed to be specified.
 
 | Item | Description |
 | -------- | ----------- |
@@ -25,7 +27,7 @@ The FontSample expression allows you to sample the texture pages out of a font r
 [EXCERPT:FontSampleParameter]
 ## FontSampleParameter
 
-The FontSampleParameter expression provides a way to expose a font-based parameter in a material instance constant, making it easy to use different fonts in different instances. The alpha channel of the font will contain the font outline value.  Only valid font pages are allowed to be specified.
+The **FontSampleParameter** expression provides a way to expose a font-based parameter in a material instance constant, making it easy to use different fonts in different instances. The alpha channel of the font will contain the font outline value.  Only valid font pages are allowed to be specified.
 
 | Item | Description |
 | -------- | ----------- |
