#UE4doc #docs - Added Engine Version Field. Bolded Terms. Corrected mislabeled term.

diff --git a/Engine/Documentation/Source/Engine/Rendering/Materials/ExpressionReference/Functions/FunctionsExpressions.INT.udn b/Engine/Documentation/Source/Engine/Rendering/Materials/ExpressionReference/Functions/FunctionsExpressions.INT.udn
--- a/Engine/Documentation/Source/Engine/Rendering/Materials/ExpressionReference/Functions/FunctionsExpressions.INT.udn
+++ b/Engine/Documentation/Source/Engine/Rendering/Materials/ExpressionReference/Functions/FunctionsExpressions.INT.udn
@@ -2,6 +2,7 @@ Availability:Public
 Title:Material Function Expressions
 Crumbs: %ROOT%, Engine, Engine/Rendering/Materials, Engine/Rendering/Materials/ExpressionReference
 Description:Expressions that are used to create or execute Materials Functions.
+Version: 4.9
 
 [EXCERPT:TOC]
 [TOC(start:2)]
@@ -9,7 +10,7 @@ Description:Expressions that are used to create or execute Materials Functions.
 
 ## FunctionInput
 
-The FunctionInput expression can only be placed in a material function, where it defines one of the function's inputs.
+The **FunctionInput** expression can only be placed in a material function, where it defines one of the function's inputs.
 
 | Item | Description |
 | -------- | ----------- |
@@ -27,7 +28,7 @@ This node is used with [](Engine/Rendering/Materials/Functions).
 
 ## FunctionOutput
 
-The FunctionInput expression can only be placed in a material function, where it defines one of the function's outputs.
+The **FunctionOutput** expression can only be placed in a material function, where it defines one of the function's outputs.
 
 | Item | Description |
 | -------- | ----------- |
@@ -42,7 +43,7 @@ This node is used with [](Engine/Rendering/Materials/Functions).
 
 ## MaterialFunctionCall
 
-The MaterialFunctionCall expression allows you to use an external [MaterialFunction](Engine/Rendering/Materials/Functions) from another material or function. The external function's input and output nodes become inputs and outputs of the function call node. If a MaterialFunction is selected in the **Content Browser** when placing one of these expressions, it will automatically be assigned.
+The **MaterialFunctionCall** expression allows you to use an external [MaterialFunction](Engine/Rendering/Materials/Functions) from another material or function. The external function's input and output nodes become inputs and outputs of the function call node. If a MaterialFunction is selected in the **Content Browser** when placing one of these expressions, it will automatically be assigned.
 
 **Shortcut:** _F + Left Mouse Click_
 
@@ -57,7 +58,7 @@ This node is used with [](Engine/Rendering/Materials/Functions).
 
 ## StaticBool
 
-The StaticBool expression is used to provide a default bool value for a static bool function input within a function. This node does not switch between anything, so it must be used in conjunction with a StaticSwitch node.
+The **StaticBool** expression is used to provide a default bool value for a static bool function input within a function. This node does not switch between anything, so it must be used in conjunction with a StaticSwitch node.
 
 | Item | Description |
 | -------- | ----------- |
@@ -70,7 +71,7 @@ This node is used with [](Engine/Rendering/Materials/Functions).
 
 ## StaticSwitch
 
-The StaticSwitch expression works like a StaticSwitchParameter, except that it only implements the switch and does not create a parameter.
+The **StaticSwitch** expression works like a StaticSwitchParameter, except that it only implements the switch and does not create a parameter.
 
 | Item | Description |
 | -------- | ----------- |
@@ -88,7 +89,7 @@ This node is used with [](Engine/Rendering/Materials/Functions).
 [EXCERPT:TextureObject]
 ## TextureObject
 
-The TextureObject expression is used to provide a default texture for a texture function input within a function. This node does not actually sample the texture, so it must be used in conjunction with a TextureSample node.
+The **TextureObject** expression is used to provide a default texture for a texture function input within a function. This node does not actually sample the texture, so it must be used in conjunction with a TextureSample node.
 
 | Item | Description |
 | -------- | ----------- |
