#UE4doc #docs - Added Engine Version Field.

diff --git a/Engine/Documentation/Source/Engine/Rendering/Materials/ExpressionReference/Landscape/LandscapeExpressions.INT.udn b/Engine/Documentation/Source/Engine/Rendering/Materials/ExpressionReference/Landscape/LandscapeExpressions.INT.udn
--- a/Engine/Documentation/Source/Engine/Rendering/Materials/ExpressionReference/Landscape/LandscapeExpressions.INT.udn
+++ b/Engine/Documentation/Source/Engine/Rendering/Materials/ExpressionReference/Landscape/LandscapeExpressions.INT.udn
@@ -2,6 +2,7 @@ Availability:Public
 Title:Landscape Expressions
 Crumbs: %ROOT%, Engine, Engine/Rendering/Materials, Engine/Rendering/Materials/ExpressionReference
 Description:Expressions for creating materials to be applied to Landscape terrains.
+Version: 4.9
 
 [EXCERPT:TOC]
 [TOC(start:2)]
