#UE4doc #docs - Added Engine Version Field. Bolded Terms.

diff --git a/Engine/Documentation/Source/Engine/Rendering/Materials/ExpressionReference/MaterialAttributes/MaterialAttributes.INT.udn b/Engine/Documentation/Source/Engine/Rendering/Materials/ExpressionReference/MaterialAttributes/MaterialAttributes.INT.udn
--- a/Engine/Documentation/Source/Engine/Rendering/Materials/ExpressionReference/MaterialAttributes/MaterialAttributes.INT.udn
+++ b/Engine/Documentation/Source/Engine/Rendering/Materials/ExpressionReference/MaterialAttributes/MaterialAttributes.INT.udn
@@ -2,7 +2,7 @@ Availability:Public
 Title:Material Attributes Expressions
 Crumbs: %ROOT%, Engine, Engine/Rendering/Materials, Engine/Rendering/Materials/ExpressionReference
 Description:These expression nodes serve as a way to separate or combine various Material attributes, which is particularly useful when creating Layered Materials.
-
+Version: 4.9
 
 [EXCERPT:TOC]
 [TOC(start:2)]
@@ -14,7 +14,7 @@ Description:These expression nodes serve as a way to separate or combine various
 
 ## Break Material Attributes
 
-The Break Material Attributes expression is ideal when using a [Layered Material](Engine/Rendering/Materials\LayeredMaterials) - a feature of the Material Functions system. When using a Material Layer Function within a Material, you may want to use only one aspect of the layer. For example, you may have a Material Layer that defines a nice looking generic Material, such as steel. You may want to use only the Roughness and Base Color attributes from that layer in your final Material, rather than using the whole thing. In such cases, you can use a Break Material Attributes node to split up all of the incoming attributes of the Material Layer, and then just plug in the ones you want. This also allows for complex blending of various Material Attributes.
+The **Break Material Attributes** expression is ideal when using a [Layered Material](Engine/Rendering/Materials\LayeredMaterials) - a feature of the Material Functions system. When using a **Material Layer Function** within a Material, you may want to use only one aspect of the layer. For example, you may have a Material Layer that defines a nice looking generic Material, such as steel. You may want to use only the Roughness and Base Color attributes from that layer in your final Material, rather than using the whole thing. In such cases, you can use a Break Material Attributes node to split up all of the incoming attributes of the Material Layer, and then just plug in the ones you want. This also allows for complex blending of various Material Attributes.
 
 [REGION:fullwidth]
 ![](BreakMaterialAttributes.png)
@@ -24,7 +24,7 @@ The Break Material Attributes expression is ideal when using a [Layered Material
 ## Make Material Attributes
 
 
-The Make Material Attributes node does exactly the opposite of the Break Material Attributes node. Instead of splitting attributes apart, this brings them together. This is useful when creating your own [Material Layer](Engine/Rendering/Materials\LayeredMaterials) functions, as you will have access to all of the standard attributes for your output. This can also be used for complex Material setups in which you want to define more than one type of Material and blend them together, all within one Material.
+The **Make Material Attributes** node does exactly the opposite of the Break Material Attributes node. Instead of splitting attributes apart, this brings them together. This is useful when creating your own [Material Layer](Engine/Rendering/Materials\LayeredMaterials) functions, as you will have access to all of the standard attributes for your output. This can also be used for complex Material setups in which you want to define more than one type of Material and blend them together, all within one Material.
 
 [REGION:fullwidth]
 ![](MakeMaterialAttributes.png)
