#UE4doc #docs - Added Engine Version Field. Bolded Terms. Fixed error reported on AnswerHub with the IF node.

diff --git a/Engine/Documentation/Source/Engine/Rendering/Materials/ExpressionReference/Math/MathExpressions.INT.udn b/Engine/Documentation/Source/Engine/Rendering/Materials/ExpressionReference/Math/MathExpressions.INT.udn
--- a/Engine/Documentation/Source/Engine/Rendering/Materials/ExpressionReference/Math/MathExpressions.INT.udn
+++ b/Engine/Documentation/Source/Engine/Rendering/Materials/ExpressionReference/Math/MathExpressions.INT.udn
@@ -2,6 +2,7 @@ Availability:Public
 Title:Math Expressions
 Crumbs: %ROOT%, Engine, Engine/Rendering/Materials, Engine/Rendering/Materials/ExpressionReference
 Description:Expressions that perform mathematical operations on one or more inputs.
+Version: 4.9
 
 [EXCERPT:TOC]
 [TOC(start:2)]
@@ -9,7 +10,7 @@ Description:Expressions that perform mathematical operations on one or more inpu
 
 ## Abs
 
-Abs is an abbreviation for the mathematical term "absolute value".  The Abs expression outputs the absolute, or unsigned, value of the input it receives. Essentially, this means it turns negative numbers into positive numbers by dropping the minus sign, while positive numbers and zero remain unchanged.
+**Abs** is an abbreviation for the mathematical term "absolute value".  The Abs expression outputs the absolute, or unsigned, value of the input it receives. Essentially, this means it turns negative numbers into positive numbers by dropping the minus sign, while positive numbers and zero remain unchanged.
 
 **Examples:** Abs of -0.7 is 0.7; Abs of -1.0 is 1.0; Abs of 1.0 is also 1.0
 
@@ -20,7 +21,7 @@ Abs is an abbreviation for the mathematical term "absolute value".  The Abs expr
 
 ## Add
 
-The Add expression takes two inputs, adds them together and outputs the result.  This addition operation is performed on a per channel basis, meaning that the inputs' R channels get added, G channels get added, B channels get added, etc. Both inputs must have the same number of channels unless one of them is a single Constant value. Constants can be added to a vector with any number of inputs.
+The **Add** expression takes two inputs, adds them together and outputs the result.  This addition operation is performed on a per channel basis, meaning that the inputs' R channels get added, G channels get added, B channels get added, etc. Both inputs must have the same number of channels unless one of them is a single Constant value. Constants can be added to a vector with any number of inputs.
 
 | Item | Description |
 | -------- | ----------- |
@@ -44,7 +45,7 @@ The Add expression takes two inputs, adds them together and outputs the result.
 
 ## Ceil
 
-The Ceil expression takes in value(s), rounds them **up** to the next integer, and outputs the result.  See also [Floor](#Floor) and [Frac](#Frac).
+The **Ceil** expression takes in value(s), rounds them **up** to the next integer, and outputs the result.  See also [Floor](#Floor) and [Frac](#Frac).
 
 **Examples:** Ceil of 0.2 is 1.0; Ceil of (0.2,1.6) is (1.0,2.0).
 
@@ -53,7 +54,7 @@ The Ceil expression takes in value(s), rounds them **up** to the next integer, a
 
 ## Clamp
 
-The Clamp expression takes in value(s) and constrains them to a specified range, defined by a minimum and maximum value.  A minimum value of 0.0 and maximum value of 0.5 means that the resulting value(s) will never be less than 0.0 and never greater than 0.5.
+The **Clamp** expression takes in value(s) and constrains them to a specified range, defined by a minimum and maximum value.  A minimum value of 0.0 and maximum value of 0.5 means that the resulting value(s) will never be less than 0.0 and never greater than 0.5.
 
 | Item | Description |
 | -------- | ----------- |
@@ -76,7 +77,7 @@ The Clamp expression takes in value(s) and constrains them to a specified range,
 
 ## Cosine
 
-The Cosine expression outputs the cosine of the value input (in radians). Most commonly, this is used to output a continuous oscillating waveform by connecting a [Time](Engine/Rendering/Materials/ExpressionReference/Constant#Time) expression to its input. The output value will cycle back and forth between -1 and 1. A visual representation of the wave is shown below:
+The **Cosine** expression outputs the cosine of the value input (in radians). Most commonly, this is used to output a continuous oscillating waveform by connecting a [Time](Engine/Rendering/Materials/ExpressionReference/Constant#Time) expression to its input. The output value will cycle back and forth between -1 and 1. A visual representation of the wave is shown below:
 
 ![CosineWave.png](CosineWave.png)
 
@@ -96,7 +97,7 @@ In the example above, the color would oscillate with co-sinusoidal frequency.
 [EXCERPT:CrossProduct]
 ## CrossProduct
 
-The CrossProduct expression computes the cross product of two three-channel vector value inputs and outputs the resulting three-channel vector value. Given two vectors in space, the cross product is a vector that is perpendicular to both of the inputs.
+The **CrossProduct** expression computes the cross product of two three-channel vector value inputs and outputs the resulting three-channel vector value. Given two vectors in space, the cross product is a vector that is perpendicular to both of the inputs.
 
 | Item | Description |
 | -------- | ----------- |
@@ -114,7 +115,7 @@ The CrossProduct expression computes the cross product of two three-channel vect
 
 ## Divide
 
-The Divide expression takes two inputs and outputs the result of the first divided by the second.  The division happens per channel, meaning that the R channel of the first is divided by the second, the G channel of the first is divided by the second, and so on. Both inputs must have the same number of values unless the divisor is a single float value. _Never divide by zero._
+The **Divide** expression takes two inputs and outputs the result of the first divided by the second.  The division happens per channel, meaning that the R channel of the first is divided by the second, the G channel of the first is divided by the second, and so on. Both inputs must have the same number of values unless the divisor is a single float value. _Never divide by zero._
 
 | Item | Description |
 | -------- | ----------- |
@@ -134,7 +135,7 @@ The Divide expression takes two inputs and outputs the result of the first divid
 [EXCERPT:DotProduct]
 ## DotProduct
 
-The DotProduct expression computes the dot product, or the length of one vector projected onto the other. This calculation is used by many techniques for computing falloff.  DotProduct requires both vector inputs to have the same number of channels.
+The **DotProduct** expression computes the dot product, or the length of one vector projected onto the other. This calculation is used by many techniques for computing falloff.  DotProduct requires both vector inputs to have the same number of channels.
 
 | Item | Description |
 | -------- | ----------- |
@@ -149,7 +150,7 @@ The DotProduct expression computes the dot product, or the length of one vector
 
 ## Floor
 
-The Floor expression takes in value(s), rounds them **down** to the previous integer, and outputs the result.  See also [Ceil](#Ceil) and [Frac](#Frac).
+The **Floor** expression takes in value(s), rounds them **down** to the previous integer, and outputs the result.  See also [Ceil](#Ceil) and [Frac](#Frac).
 
 **Examples:** Floor of 0.2 is 0.0; Floor of (0.2,1.6) is (0.0, 1.0).
 
@@ -158,7 +159,7 @@ The Floor expression takes in value(s), rounds them **down** to the previous int
 
 ## FMod
 
-The FMod expression returns the floating-point remainder of the division operation of the two inputs.
+The **FMod** expression returns the floating-point remainder of the division operation of the two inputs.
 
 ![FModExample.png](FModExample.png)
 
@@ -166,7 +167,7 @@ In this example, FMod is returning a repeating value that ramps up from 0 to 1 e
 
 ## Frac
 
-The Frac expression takes in value(s) and outputs the decimal part of those values.  See also [Ceil](#Ceil) and [Floor](#Floor).
+The **Frac** expression takes in value(s) and outputs the decimal part of those values.  See also [Ceil](#Ceil) and [Floor](#Floor).
 
 **Examples:** Frac of 0.2 is 0.2; Frac of (0.0,1.6) is (0.0, 0.6).
 
@@ -177,16 +178,16 @@ In this example, the Frac node is converting time into a series of repeating 0 -
 [EXCERPT:If]
 ## If
 
-The If expression compares two inputs and then passes through one of three other input values based on the result of the comparison. Both compared inputs must be single float values.
+The **If** expression compares two inputs and then passes through one of three other input values based on the result of the comparison. Both compared inputs must be single float values.
 
 | Item | Description |
 | -------- | ----------- |
 | [REGION:tablesection]Inputs[/REGION] ||
 | **A** | Takes in a single float value. |
 | **B** | Takes in a single float value. |
-| **A &lt; B** | Takes in the value(s) to output if the value of A is greater than the value of B. |
+| **A > B** | Takes in the value(s) to output if the value of A is greater than the value of B. |
 | **A = B** | Takes in the value(s) to output if the value of A is equal to the value of B. |
-| **A &gt; B** | Takes in the value(s) to output if the value of A is less than the value of B. |
+| **A &lt; B** | Takes in the value(s) to output if the value of A is less than the value of B. |
 
 
 ![IfExample.png](IfExample.png)(w:900)
@@ -201,7 +202,7 @@ In this example, the If expression takes in a heightmap and outputs either red o
 
 ## Max
 
-The Max expression takes in two inputs and then outputs the maximum of the two. 
+The **Max** expression takes in two inputs and then outputs the maximum of the two. 
 
 This node is similar to Photoshop's Lighten. 
 
@@ -221,7 +222,7 @@ In the example above, A is "0" and B is "1"; therefore, "1" (white) is the resul
 
 ## Min
 
-The Min expression takes in two inputs then outputs the minimum of the two. 
+The **Min** expression takes in two inputs then outputs the minimum of the two. 
 
 This node is similar to Photoshop's Darken. 
 
@@ -243,7 +244,7 @@ In the example above, A is "0" and B is "1"; therefore, "0" (black) is the resul
 
 ## Multiply
 
-The Multiply expression takes two inputs, multiplies them together, and outputs the result. Similar to Photoshop's multiply layer blend.  Multiplication happens per channel, meaning that the R channel of the first is multiplied by the second; the G channel of the first is multiplied by the second, and so on. Both inputs must have the same number of values unless one of the values is a single float value.
+The **Multiply** expression takes two inputs, multiplies them together, and outputs the result. Similar to Photoshop's multiply layer blend.  Multiplication happens per channel, meaning that the R channel of the first is multiplied by the second; the G channel of the first is multiplied by the second, and so on. Both inputs must have the same number of values unless one of the values is a single float value.
 
 | Item | Description |
 | -------- | ----------- |
@@ -267,7 +268,7 @@ Do not forget that materials in UE4 are not limited to [0,1].  If colors/values
 [EXCERPT:Normalize]
 ## Normalize
 
-The Normalize expression calculates and outputs the normalized value of its input. This means each component of the input is divided by the L-2 norm (length) of the vector.
+The **Normalize** expression calculates and outputs the normalized value of its input. This means each component of the input is divided by the L-2 norm (length) of the vector.
 
 ![NormalizeExample.png](NormalizeExample.png)(w:900)
 
@@ -279,7 +280,7 @@ It is not necessary to normalize an expression that plugs into the Normal materi
 
 ## OneMinus
 
-The OneMinus expression takes an input value and outputs 'one-minus-that-value'. This operation is performed per channel.
+The **OneMinus** expression takes an input value and outputs 'one-minus-that-value'. This operation is performed per channel.
 
 **Examples:** OneMinus of 0.4 is 0.6; OneMinus of (0.2,0.5,1.0) is (0.8,0.5,0.0); OneMinus of (0.0,-0.4,1.6) is (1.0,1.4,-0.6).
 
@@ -290,7 +291,7 @@ The OneMinus expression takes an input value and outputs 'one-minus-that-value'.
 
 ## Power
 
-The Power expression takes two inputs, raises Base to the Exp power, and outputs the result; in other words, Base multiplied by itself Exp times.
+The **Power** expression takes two inputs, raises Base to the Exp power, and outputs the result; in other words, Base multiplied by itself Exp times.
 
 | Item | Description |
 | -------- | ----------- |
@@ -310,7 +311,7 @@ The Power expression takes two inputs, raises Base to the Exp power, and outputs
 
 ## Sine
 
-The Sine expression outputs the sine of the value input (in radians). Most commonly, this is used to output a continuous oscillating waveform by connecting a [Time](Engine/Rendering/Materials/ExpressionReference/Constant#Time) expression to its input. The output value will cycle back and forth between -1 and 1. The difference between this and the output of the [Cosine](#Cosine) expression is the output waveform is offset half the period. A visual representation of the wave is shown below:
+The **Sine** expression outputs the sine of the value input (in radians). Most commonly, this is used to output a continuous oscillating waveform by connecting a [Time](Engine/Rendering/Materials/ExpressionReference/Constant#Time) expression to its input. The output value will cycle back and forth between -1 and 1. The difference between this and the output of the [Cosine](#Cosine) expression is the output waveform is offset half the period. A visual representation of the wave is shown below:
 
 ![SineWave.png](SineWave.png)
 
@@ -327,14 +328,14 @@ The Sine expression outputs the sine of the value input (in radians). Most commo
 
 ## SquareRoot
 
-The SquareRoot expression outputs the square root of the input value.  SquareRoot can operate only on a single float input value.
+The **SquareRoot** expression outputs the square root of the input value.  SquareRoot can operate only on a single float input value.
 
 ![SqrtExample.png](SqrtExample.png)(w:900)
 
 
 ## Subtract
 
-The Subtract node takes in two inputs, subtracts the second input from the first, and outputs the difference.  Subtraction happens per channel, meaning that the R channel of the second input gets subtracted from the first; the G channel of the second input gets subtracted from the first, and so on. Both inputs must have the same number of channels unless the second input is a single Constant value. Constants can be subtracted from a vector with any number of inputs.
+The **Subtract** node takes in two inputs, subtracts the second input from the first, and outputs the difference.  Subtraction happens per channel, meaning that the R channel of the second input gets subtracted from the first; the G channel of the second input gets subtracted from the first, and so on. Both inputs must have the same number of channels unless the second input is a single Constant value. Constants can be subtracted from a vector with any number of inputs.
 
 | Item | Description |
 | -------- | ----------- |
diff --git a/Engine/Documentation/Source/Engine/Rendering/Materials/ExpressionReference/Utility/UtilityExpressions.INT.udn b/Engine/Documentation/Source/Engine/Rendering/Materials/ExpressionReference/Utility/UtilityExpressions.INT.udn
--- a/Engine/Documentation/Source/Engine/Rendering/Materials/ExpressionReference/Utility/UtilityExpressions.INT.udn
+++ b/Engine/Documentation/Source/Engine/Rendering/Materials/ExpressionReference/Utility/UtilityExpressions.INT.udn
@@ -202,7 +202,7 @@ The LightmassReplace expression simply passes through the Realtime input when co
 [EXCERPT:LinearInterpolate]
 ## LinearInterpolate
 
-The LinearInterpolate expression blends between two input value(s) based on a third input value used as a mask. This can be thought of as a mask to define transitions between two textures, like a layer mask in Photoshop.  The intensity of the mask Alpha determines the ratio of color to take from the two input values.  If Alpha is 0.0/black, the first input is used.  If Alpha is 1.0/white, the second input is used.  If Alpha is grey (somewhere between 0.0 and 1.0), the output is a blend between the two inputs. Keep in mind that the blend happens per channel.  So, if Alpha is an RGB color, Alpha's red channel value defines the blend between A and B's red channels **independently** of Alpha's green channel, which defines the blend between A and B's green channels.
+The **LinearInterpolate** expression blends between two input value(s) based on a third input value used as a mask. This can be thought of as a mask to define transitions between two textures, like a layer mask in Photoshop.  The intensity of the mask Alpha determines the ratio of color to take from the two input values.  If Alpha is 0.0/black, the first input is used.  If Alpha is 1.0/white, the second input is used.  If Alpha is grey (somewhere between 0.0 and 1.0), the output is a blend between the two inputs. Keep in mind that the blend happens per channel.  So, if Alpha is an RGB color, Alpha's red channel value defines the blend between A and B's red channels **independently** of Alpha's green channel, which defines the blend between A and B's green channels.
 
 | Item | Description |
 | -------- | ----------- |
diff --git a/Engine/Documentation/Source/Engine/Rendering/Materials/ExpressionReference/VectorOps/VectorOpExpressions.INT.udn b/Engine/Documentation/Source/Engine/Rendering/Materials/ExpressionReference/VectorOps/VectorOpExpressions.INT.udn
--- a/Engine/Documentation/Source/Engine/Rendering/Materials/ExpressionReference/VectorOps/VectorOpExpressions.INT.udn
+++ b/Engine/Documentation/Source/Engine/Rendering/Materials/ExpressionReference/VectorOps/VectorOpExpressions.INT.udn
@@ -10,7 +10,7 @@ Description:Expressions that perform operations on vector input values.
 [EXCERPT:AppendVector]
 ## AppendVector
 
-The AppendVector expression allows you to combine channels together to create a vector with more channels than the original.  For example, you can take two individual [Constants](Engine/Rendering/Materials/ExpressionReference/Constant#Constant) values and append them to make a two-channel [Constant2Vector](Engine/Rendering/Materials/ExpressionReference/Constant#Constant2Vector) value. This can be useful for reordering the channels within a single texture or for combining multiple grayscale textures into one RGB color texture.
+The **AppendVector** expression allows you to combine channels together to create a vector with more channels than the original.  For example, you can take two individual [Constants](Engine/Rendering/Materials/ExpressionReference/Constant#Constant) values and append them to make a two-channel [Constant2Vector](Engine/Rendering/Materials/ExpressionReference/Constant#Constant2Vector) value. This can be useful for reordering the channels within a single texture or for combining multiple grayscale textures into one RGB color texture.
 
 | Item | Description |
 | -------- | ----------- |
@@ -29,7 +29,7 @@ The AppendVector expression allows you to combine channels together to create a
 [EXCERPT:ComponentMask]
 ## ComponentMask
 
-The ComponentMask expression allows you to select a specific subset of channels (R, G, B, and/or A) from the input to pass through to the output. Attempting to pass a channel through that does not exist in the input will cause an error, unless the input is a single constant value. In that case, the single value is passed through to each channel. The current channels selected to be passed through are displayed in the title bar of the expression.
+The **ComponentMask** expression allows you to select a specific subset of channels (R, G, B, and/or A) from the input to pass through to the output. Attempting to pass a channel through that does not exist in the input will cause an error, unless the input is a single constant value. In that case, the single value is passed through to each channel. The current channels selected to be passed through are displayed in the title bar of the expression.
 
 | Item | Description |
 | -------- | ----------- |
