#UE4doc #docs - Added Engine Version Field. Bolded Terms.

diff --git a/Engine/Documentation/Source/Engine/Rendering/Materials/ExpressionReference/Parameters/ParameterExpressions.INT.udn b/Engine/Documentation/Source/Engine/Rendering/Materials/ExpressionReference/Parameters/ParameterExpressions.INT.udn
--- a/Engine/Documentation/Source/Engine/Rendering/Materials/ExpressionReference/Parameters/ParameterExpressions.INT.udn
+++ b/Engine/Documentation/Source/Engine/Rendering/Materials/ExpressionReference/Parameters/ParameterExpressions.INT.udn
@@ -2,6 +2,7 @@ Availability:Public
 Title:Parameter Expressions
 Crumbs: %ROOT%, Engine, Engine/Rendering/Materials, Engine/Rendering/Materials/ExpressionReference
 Description:Expressions that expose properties to Material Instances to be overridden in child instances or modified at runtime.
+Version: 4.9
 
 [EXCERPT:TOC]
 [TOC(start:2)]
@@ -10,7 +11,7 @@ Description:Expressions that expose properties to Material Instances to be overr
 
 ## Collection Parameters
 
-A Collection Parameter expression is used to reference a Parameter Collection asset. These are groups of parameters that can easily be reused by many different assets such as Materials, Blueprints, and much more. For more information on Parameter Collections, be sure to read the [Parameter Collections Documentation.](Engine/Rendering/Materials\ParameterCollections)
+A **Collection Parameter** expression is used to reference a Parameter Collection asset. These are groups of parameters that can easily be reused by many different assets such as Materials, Blueprints, and much more. For more information on Parameter Collections, be sure to read the [Parameter Collections Documentation.](Engine/Rendering/Materials\ParameterCollections)
 
 
 
@@ -22,7 +23,7 @@ A Collection Parameter expression is used to reference a Parameter Collection as
 ## DynamicParameter
 [EXCERPT:DynamicParameter]
 
-The DynamicParameter expression provides a conduit for particle emitters to pass up to four values to the material to be used in any manner. These values are set in Cascade via a ParameterDynamic module placed on an emitter.
+The **DynamicParameter** expression provides a conduit for particle emitters to pass up to four values to the material to be used in any manner. These values are set in **Cascade** via a **ParameterDynamic** module placed on an emitter.
 
 | Item | Description |
 | -------- | ----------- |
@@ -41,7 +42,7 @@ The DynamicParameter expression provides a conduit for particle emitters to pass
 
 ## ScalarParameter
 
-The ScalarParameter expression outputs a single float value ([Constant](Engine/Rendering/Materials/ExpressionReference/Constant#Constant)) that can be accessed and changed in an instance of the material or on the fly by code.  
+The **ScalarParameter** expression outputs a single float value ([Constant](Engine/Rendering/Materials/ExpressionReference/Constant#Constant)) that can be accessed and changed in an instance of the material or on the fly by code.  
 
 | Item | Description |
 | -------- | ----------- |
@@ -58,7 +59,7 @@ The ScalarParameter expression outputs a single float value ([Constant](Engine/R
 
 ## StaticBoolParameter
 
-The StaticBoolParameter works like StaticSwitchParameter, except that it only creates a bool parameter and does not implement the switch.
+The **StaticBoolParameter** works like StaticSwitchParameter, except that it only creates a bool parameter and does not implement the switch.
 
 [REGION:note]
 This parameter is named static because it cannot change at runtime, it can only be set in the Material Instance Editor.  Static Switches are applied at compile time, not at runtime.  This means that whatever branch of the material was dropped will never be executed, so static switches are effectively free at runtime.  On the other hand, a new version of the material must be compiled out for every **used** combination of static parameters in a material, which can lead to a shader explosion if abused.  Try to minimize the number of static parameters in the material and the number of permutations of those static parameters that are actually used.
@@ -78,7 +79,7 @@ This node is used with [MaterialFunctions](Engine/Rendering/Materials/Functions)
 
 ## StaticSwitchParameter
 
-The StaticSwitchParameter expression takes in two inputs, and outputs the first if the value of the parameter is _true_, and the second otherwise.  
+The **StaticSwitchParameter** expression takes in two inputs, and outputs the first if the value of the parameter is _true_, and the second otherwise.  
 
 [REGION:note]
 This parameter is named static because it cannot change at runtime, it can only be set in the Material Instance Editor.  Static Switches are applied at compile time, not at runtime.  This means that whatever branch of the material was dropped will never be executed, so static switches are effectively free at runtime.  On the other hand, a new version of the material must be compiled out for every **used** combination of static parameters in a material, which can lead to a shader explosion if abused.  Try to minimize the number of static parameters in the material and the number of permutations of those static parameters that are actually used.
@@ -108,7 +109,7 @@ This parameter is named static because it cannot change at runtime, it can only
 
 ## StaticComponentMaskParameter
 
-The StaticComponentMaskParameter expression behaves just like an ordinary Component Mask, except that the mask values can be set by instances.  
+The **StaticComponentMaskParameter** expression behaves just like an ordinary Component Mask, except that the mask values can be set by instances.  
 
 [REGION:note]
 This parameter is named static because it cannot change at runtime, it can only be set in the Material Instance Editor.  Static Switches are applied at compile time, not at runtime.  This means that whatever branch of the material was dropped will never be executed, so static switches are effectively free at runtime.  On the other hand, a new version of the material must be compiled out for every **used** combination of static parameters in a material, which can lead to a shader explosion if abused.  Try to minimize the number of static parameters in the material and the number of permutations of those static parameters that are actually used.
@@ -133,7 +134,7 @@ This parameter is named static because it cannot change at runtime, it can only
 
 ## VectorParameter
 
-The VectorParameter expression is identical to the [Constant4Vector](Engine/Rendering/Materials/ExpressionReference/Constant#Constant4Vector), except that it is a parameter and can be modified in instances of the material and through code. One nicety of the VectorParameter is that its value can be set using the Color picker.
+The **VectorParameter** expression is identical to the [Constant4Vector](Engine/Rendering/Materials/ExpressionReference/Constant#Constant4Vector), except that it is a parameter and can be modified in instances of the material and through code. One nicety of the VectorParameter is that its value can be set using the Color picker.
 
 | Item | Description |
 | -------- | ----------- |
@@ -159,7 +160,7 @@ VertexColor is mutually exclusive with the Transform node due to limited interpo
 [EXCERPT:TextureObjectParameter]
 ## TextureObjectParameter
 
-The TextureObjectParameter expression defines a texture parameter and outputs the texture object, used in materials that call a function with texture inputs. This node does not actually sample the texture, so it must be used in conjunction with a TextureSample node.
+The **TextureObjectParameter** expression defines a texture parameter and outputs the texture object, used in materials that call a function with texture inputs. This node does not actually sample the texture, so it must be used in conjunction with a TextureSample node.
 
 | Item | Description |
 | -------- | ----------- |
@@ -177,7 +178,7 @@ This node is used with [MaterialFunctions](Engine/Rendering/Materials/Functions)
 [EXCERPT:TextureSampleParameter2D]
 ## TextureSampleParameter2D
 
-The TextureSampleParameter2D expression is identical to the TextureSample except that it is a parameter that can be modified in instances of the material and through code.
+The **TextureSampleParameter2D** expression is identical to the TextureSample except that it is a parameter that can be modified in instances of the material and through code.
 
 | Item | Description |
 | -------- | ----------- |
@@ -203,7 +204,7 @@ The TextureSampleParameter2D expression is identical to the TextureSample except
 [EXCERPT:TextureSampleParameterCube]
 ## TextureSampleParameterCube
 
-The TextureSampleParameterCube expression is identical to the TextureSample except that it only accepts cubemaps and it is a parameter that can be modified in instances of the material and through code.
+The **TextureSampleParameterCube** expression is identical to the TextureSample except that it only accepts cubemaps and it is a parameter that can be modified in instances of the material and through code.
 
 | Item | Description |
 | -------- | ----------- |
@@ -227,7 +228,7 @@ The TextureSampleParameterCube expression is identical to the TextureSample exce
 [EXCERPT:TextureSampleParameterMovie]
 ## TextureSampleParameterMovie
 
-The TextureSampleParameterMovie expression is identical to the TextureSample except that it only accepts movie textures (Bink movies) and it is a parameter that can be modified in instances of the material and through code.
+The **TextureSampleParameterMovie** expression is identical to the TextureSample except that it only accepts movie textures (Bink movies) and it is a parameter that can be modified in instances of the material and through code.
 
 | Item | Description |
 | -------- | ----------- |
