#DocFix

Fixing errors in documentaiton reported through the forums

diff --git a/Engine/Documentation/Source/Engine/Rendering/ParticleSystems/Cascade/CascadeEditor.INT.udn b/Engine/Documentation/Source/Engine/Rendering/ParticleSystems/Cascade/CascadeEditor.INT.udn
--- a/Engine/Documentation/Source/Engine/Rendering/ParticleSystems/Cascade/CascadeEditor.INT.udn
+++ b/Engine/Documentation/Source/Engine/Rendering/ParticleSystems/Cascade/CascadeEditor.INT.udn
@@ -182,7 +182,7 @@ The View menu allows you show and hide various diagnostic and visualization feat
 | **Geometry** | Shows or hides the default test mesh. |
 | **Geometry Properties** | Brings up the geometry properties window, where you can adjust the properties of the default geometry, adjusting placement, swapping out the mesh used, and many other options. |
 | **Bounds** | Toggles visibility for the particle system bounds. These are shown as both a bounding box and bounding sphere. |
-| **Set Fixed Bounds** | This activated fixed boundaries for the particle system and automatically sets the values to the largest setting calculated thus far. |
+| **Set Fixed Bounds** | This activates fixed boundaries for the particle system and automatically sets the values to the largest setting calculated thus far. |
 
 
 ##### View Overlay Submenu
diff --git a/Engine/Documentation/Source/Engine/Rendering/ParticleSystems/Overview/KeyConcepts.INT.udn b/Engine/Documentation/Source/Engine/Rendering/ParticleSystems/Overview/KeyConcepts.INT.udn
--- a/Engine/Documentation/Source/Engine/Rendering/ParticleSystems/Overview/KeyConcepts.INT.udn
+++ b/Engine/Documentation/Source/Engine/Rendering/ParticleSystems/Overview/KeyConcepts.INT.udn
@@ -199,7 +199,7 @@ There are 5 primary distribution types:
 * Parameter - This is a type of distribution that allows a property to be parameterized so that it can be manipulated externally via Blueprints, Matinee, or other means.
 
 [REGION:note]
-	Each of these distribution types exist separately for float and vector data type values. For instance, there exist both Float Uniform Curve and Vector Uniform Curve distributions. However, whether you use a float or a vector version of each distribution is governed by the type of property being controlled; the user has no choice. As an example, you would control color (a vector value for RGB) via a Vector Uniform distribution, but Lifespan (a single float value) via a Float Uniform.
+	Each of these distribution types exist separately for float and vector data type values. For instance, there exists both Float Uniform Curve and Vector Uniform Curve distributions. However, whether you use a float or a vector version of each, distribution is governed by the type of property being controlled; the user has no choice. As an example, you would control color (a vector value for RGB) via a Vector Uniform distribution, but Lifespan (a single float value) via a Float Uniform.
 [/REGION]
 
 
diff --git a/Engine/Documentation/Source/Engine/Rendering/PostProcessEffects/SceneFringe/SceneFringe.INT.udn b/Engine/Documentation/Source/Engine/Rendering/PostProcessEffects/SceneFringe/SceneFringe.INT.udn
--- a/Engine/Documentation/Source/Engine/Rendering/PostProcessEffects/SceneFringe/SceneFringe.INT.udn
+++ b/Engine/Documentation/Source/Engine/Rendering/PostProcessEffects/SceneFringe/SceneFringe.INT.udn
@@ -18,6 +18,7 @@ Version: 4.9
 	[/PARAM]
 [/OBJECT]
 
+**Scene Fringe** settings can be located under **Scene Color** in the **Post Process Settings**.
 
 | Property | Description |
 | -------- | ----------- |
diff --git a/Engine/Documentation/Source/Engine/Rendering/PostProcessEffects/Vignette/Vignette.INT.udn b/Engine/Documentation/Source/Engine/Rendering/PostProcessEffects/Vignette/Vignette.INT.udn
--- a/Engine/Documentation/Source/Engine/Rendering/PostProcessEffects/Vignette/Vignette.INT.udn
+++ b/Engine/Documentation/Source/Engine/Rendering/PostProcessEffects/Vignette/Vignette.INT.udn
@@ -12,6 +12,8 @@ Version: 4.9
 **Vignette** is an effect that simulates the darkening in real-world camera lenses. High quality lens try to compensate for this effect. The effect is mostly noticeable near the edges of the image.
 [/EXCERPT:Vignette_Desc]
 
+**Vignette** settings can be located under **Scene Color** in the **Post Process Settings**.
+
 | Property | Description |
 | -------- | ----------- |
 | **Vignette Intensity** | Allows for darkening of the screen corners. Does not affect bloom (performance reasons but possible to change). Possible uses: simulate camera imperfections. [INCLUDE:#Intensity] |
