#UE4doc #docs - Added Engine Version Field. Bolded terms. Replaced image table with comparison slider.

diff --git a/Engine/Documentation/Source/Engine/Rendering/Materials/ExpressionReference/Particles/ParticleExpressions.INT.udn b/Engine/Documentation/Source/Engine/Rendering/Materials/ExpressionReference/Particles/ParticleExpressions.INT.udn
--- a/Engine/Documentation/Source/Engine/Rendering/Materials/ExpressionReference/Particles/ParticleExpressions.INT.udn
+++ b/Engine/Documentation/Source/Engine/Rendering/Materials/ExpressionReference/Particles/ParticleExpressions.INT.udn
@@ -2,6 +2,7 @@ Availability:Public
 Title:Particle Expressions
 Crumbs: %ROOT%, Engine, Engine/Rendering/Materials, Engine/Rendering/Materials/ExpressionReference
 Description:Expressions used for creating materials to be applied to emitters in Particle Systems.
+Version: 4.9
 
 [EXCERPT:TOC]
 [TOC(start:2)]
@@ -16,7 +17,7 @@ Description:Expressions used for creating materials to be applied to emitters in
 [EXCERPT:ParticleColor]
 ## ParticleColor
 
-The ParticleColor expression ties into the current color of a given particle based on any per-particle color data defined within Cascade. This must be plugged into the appropriate channel (Emissive Color).
+The **ParticleColor** expression ties into the current color of a given particle based on any per-particle color data defined within **Cascade**. This must be plugged into the appropriate channel (Emissive Color).
 
 | Item | Description |
 |----|----|
@@ -38,7 +39,7 @@ In this example, you can see where the ParticleColor expression is providing col
 [EXCERPT:ParticleDirection]
 ## ParticleDirection
 
-The ParticleDirection expression outputs Vector3 (RGB) data on a per-particle basis, representing the direction a given particle is currently traveling. 
+The **ParticleDirection** expression outputs Vector3 (RGB) data on a per-particle basis, representing the direction a given particle is currently traveling. 
 
 
 ![ParticleDirectionExample.png](ParticleDirectionExample.png)(w:900)
@@ -51,7 +52,7 @@ You can see in this image how the color of the particles is changing based on th
 
 ## ParticleMacroUV
 
-The ParticleMacroUV expression outputs UV texture coordinates that can be used to map any 2d texture onto the entire particle system in a continuous way, meaning the texture will be seamless across particles.  The UVs will be centered around MacroUVPosition (specified in Cascade on the ParticleSystem, under the MacroUV category) and MacroUVRadius determines the world space radius that the UVs should tile at.  The ParticleMacroUV node is useful for mapping continuous noise onto particles to break up the pattern introduced by mapping a texture onto each particle with normal texture coordinates.
+The **ParticleMacroUV** expression outputs UV texture coordinates that can be used to map any 2d texture onto the entire particle system in a continuous way, meaning the texture will be seamless across particles.  The UVs will be centered around MacroUVPosition (specified in Cascade on the ParticleSystem, under the MacroUV category) and MacroUVRadius determines the world space radius that the UVs should tile at.  The ParticleMacroUV node is useful for mapping continuous noise onto particles to break up the pattern introduced by mapping a texture onto each particle with normal texture coordinates.
 
 | Item | Description |
 | -------- | ----------- |
@@ -65,7 +66,7 @@ The ParticleMacroUV expression outputs UV texture coordinates that can be used t
 [EXCERPT:ParticleMotionBlurFade]
 ## ParticleMotionBlurFade
 
-The ParticleMotionBlurFade expression outputs a value representing the amount of fade on a particle as a result of motion blur. A value of 1 represents no blur, black represents complete blur.
+The **ParticleMotionBlurFade** expression outputs a value representing the amount of fade on a particle as a result of motion blur. A value of 1 represents no blur, black represents complete blur.
 
 ![ParticleMotionBlurFade.png](ParticleMotionBlurFade.png)(w:900)
 
@@ -76,7 +77,7 @@ The ParticleMotionBlurFade expression outputs a value representing the amount of
 [EXCERPT:ParticlePositionWS]
 ## ParticlePositionWS
 
-The ParticlePositionWS expression outputs Vector3 (RGB) data representing each individual particle's position in world space.
+The **ParticlePositionWS** expression outputs Vector3 (RGB) data representing each individual particle's position in world space.
 
 
 ![ParticlePositionWS.png](ParticlePositionWSExample.png)(w:900)
@@ -92,7 +93,7 @@ In this image, ParticlePositionWS is being fed into emissive color to visualize
 [EXCERPT:ParticleRadius]
 ## ParticleRadius
 
-The ParticleRadius expression outputs the radius in Unreal units of each particle individually. This allows, for example, for changes to be made to a material once the radius has reached a certain point.
+The **ParticleRadius** expression outputs the radius in Unreal units of each particle individually. This allows, for example, for changes to be made to a material once the radius has reached a certain point.
 
 [REGION:fullwidth]
 ![ParticleRadiusExample.png](ParticleRadiusExample.png)(w:1000 h:448)
@@ -107,7 +108,7 @@ In this image, the particles are shifting from green to red as their radius exce
 [EXCERPT:ParticleRelativetime]
 ## ParticleRelativeTime
 
-The ParticleRelativeTime expression outputs a number between 0 and 1 representing a particle's age, with 0 being the moment of birth and 1 being the moment of death.
+The **ParticleRelativeTime** expression outputs a number between 0 and 1 representing a particle's age, with 0 being the moment of birth and 1 being the moment of death.
 
 [REGION:fullwidth]
 ![ParticleRelativeTimeExample.png](ParticleRelativeTimeExample.png)(w:900)
@@ -120,7 +121,7 @@ In this example, you can see the particle relative time fed into the emissive co
 [EXCERPT:ParticleSize]
 ## ParticleSize
 
-The Particle Size expression outputs the X and Y size of a particle sprite. This can then be used to drive some aspect of a Material.
+The **Particle Size** expression outputs the X and Y size of a particle sprite. This can then be used to drive some aspect of a Material.
 
 ![](ParticleSize.png)(w:480) 
 
@@ -138,7 +139,7 @@ In the example above, Particle Size is being multiplied into Particle Color. Not
 
 ## ParticleSpeed
 
-ParticleSpeed outputs the current speed each particle is traveling, measured in Unreal units per second.
+**ParticleSpeed** outputs the current speed each particle is traveling, measured in Unreal units per second.
 
 [REGION:fullwidth]
 ![ParticleSpeedExample.png](ParticleSpeedExample.png)(w:1000 h:448)
@@ -151,7 +152,7 @@ In this example, particle speed is feeding the color of the particles, which is
 
 ## SphericalParticleOpacity
 
-The SphericalParticleOpacity expression creates a procedural opacity map to cause sprite particles to appear spherical. This can be much simpler than having to create an import a texture map for a similar effect.
+The **SphericalParticleOpacity** expression creates a procedural opacity map to cause sprite particles to appear spherical. This can be much simpler than having to create an import a texture map for a similar effect.
 
 | Item | Description |
 | -------- | ----------- |
@@ -159,11 +160,17 @@ The SphericalParticleOpacity expression creates a procedural opacity map to caus
 | **Density** | This 0-1 value controls the density of the sphere map. Lower values result in more transparent spheres. Default is 1. |
 
 
-[REGION:imagetable]
-| ![](SPO_Before.png)(w:290) | ![](SPO_After.png)(w:290) | [INCLUDE:#SPO] |
-| :---: | :---: | :---: |
-| Before | After | Network (click to enlarge) |
-[/REGION]
+[OBJECT:ComparisonSlider]
+	[PARAM:before]
+	![Before](SPO_Before.png)(w:350)
+	[/PARAM]
+	[PARAM:after]
+	![After](SPO_After.png)(w:350)
+	[/PARAM]
+[/OBJECT]
+
+[INCLUDE:#SPO]
+
 
 <!--
 [EXCERPT:SPO]
@@ -177,7 +184,7 @@ The SphericalParticleOpacity expression creates a procedural opacity map to caus
 [EXCERPT:ParticleSubUV]
 ## ParticleSubUV
 
-The ParticleSubUV expression is used to render sub-images of a texture to a particle.  ParticleSubUV is similar to a flipbook, except that ParticleSubUV allows the texture animation to be manipulated in Cascade.
+The **ParticleSubUV** expression is used to render sub-images of a texture to a particle.  ParticleSubUV is similar to a flipbook, except that ParticleSubUV allows the texture animation to be manipulated in Cascade.
 
 | Item | Description |
 | -------- | ----------- |
@@ -205,7 +212,7 @@ The ParticleSubUV expression is used to render sub-images of a texture to a part
 [EXCERPT:TextureSampleParameterSubUV]
 ## TextureSampleParameterSubUV
 
-The TextureSampleParameterSubUV expression is identical to the [ParticleSubUV](Engine/Rendering/Materials/ExpressionReference/Particles#ParticleSubUV) except that it is a parameter that can be modified in instances of the material and through code.
+The **TextureSampleParameterSubUV** expression is identical to the [ParticleSubUV](Engine/Rendering/Materials/ExpressionReference/Particles#ParticleSubUV) except that it is a parameter that can be modified in instances of the material and through code.
 
 | Item | Description |
 | -------- | ----------- |
