#UE4doc #docs - Added Engine Version field. Bolded terms. Fixed a grammatical error.

diff --git a/Engine/Documentation/Source/Engine/Rendering/Materials/ExpressionReference/Textures/TextureExpressions.INT.udn b/Engine/Documentation/Source/Engine/Rendering/Materials/ExpressionReference/Textures/TextureExpressions.INT.udn
--- a/Engine/Documentation/Source/Engine/Rendering/Materials/ExpressionReference/Textures/TextureExpressions.INT.udn
+++ b/Engine/Documentation/Source/Engine/Rendering/Materials/ExpressionReference/Textures/TextureExpressions.INT.udn
@@ -2,6 +2,7 @@ Availability:Public
 Title:Texture Expressions
 Crumbs: %ROOT%, Engine, Engine/Rendering/Materials, Engine/Rendering/Materials/ExpressionReference
 Description:Expressions that sample and output textures.
+Version: 4.9
 
 [EXCERPT:TOC]
 [TOC(start:2)]
@@ -14,7 +15,7 @@ Description:Expressions that sample and output textures.
 
 ## SceneColor
 
-The SceneColor expression outputs the existing scene color.
+The **SceneColor** expression outputs the existing scene color.
 
 | Item | Description |
 | -------- | ----------- |
@@ -63,7 +64,7 @@ Sprite instances will pass in their color as a vertex color, but it can be used
 
 ## TextureSample
 
-The TextureSample expression outputs the color value(s) from a texture. This texture can be a regular Texture2D (including normal maps), a cubemap, or a movie texture.
+The **TextureSample** expression outputs the color value(s) from a texture. This texture can be a regular Texture2D (including normal maps), a cubemap, or a movie texture.
 
 | Item | Description |
 | -------- | ----------- |
@@ -85,13 +86,13 @@ The TextureSample expression outputs the color value(s) from a texture. This tex
 
 ![TextureSampleExample.png](TextureSampleExample.png)
 
-For a better understanding of the DDX(UVs) and DDY(UVs) inputs, here is what you could do. But you need to keep in mind that this has the exact same output as with **MipValueMode** = **None**, but is slower. The DDX(UVs) and DDY(UVs) should only be used when you actually needs to overide the default hardware computation because computing those values by a different way than using the DDX/DDY material expression.
+For a better understanding of the DDX(UVs) and DDY(UVs) inputs, here is what you could do. But you need to keep in mind that this has the exact same output as with **MipValueMode** = **None**, but is slower. The DDX(UVs) and DDY(UVs) should only be used when you actually need to overide the default hardware computation because computing those values by a different way than using the DDX/DDY material expression.
 
 ![SampleTextureDerivative.png](SampleTextureDerivative.png)
 
 ## TextureProperty
 
-The TextureProperty exposes a texture's property of your choice such as the texture's size or texel size.
+The **TextureProperty** exposes a texture's property of your choice such as the texture's size or texel size.
 
 ![TextureProperty.png](TextureProperty.png)
 
