#UE4doc #docs - Added Engine version Field. Bolded terms. Replaced image table with comparison slider.

diff --git a/Engine/Documentation/Source/Engine/Rendering/Materials/ExpressionReference/Utility/UtilityExpressions.INT.udn b/Engine/Documentation/Source/Engine/Rendering/Materials/ExpressionReference/Utility/UtilityExpressions.INT.udn
--- a/Engine/Documentation/Source/Engine/Rendering/Materials/ExpressionReference/Utility/UtilityExpressions.INT.udn
+++ b/Engine/Documentation/Source/Engine/Rendering/Materials/ExpressionReference/Utility/UtilityExpressions.INT.udn
@@ -2,6 +2,7 @@ Availability:Public
 Title:Utility Expressions
 Crumbs: %ROOT%, Engine, Engine/Rendering/Materials, Engine/Rendering/Materials/ExpressionReference
 Description:Expressions that perform various utility operations on one or more inputs.
+Version: 4.9
 
 [EXCERPT:TOC]
 [TOC(start:2)]
@@ -10,7 +11,7 @@ Description:Expressions that perform various utility operations on one or more i
 
 ## BlackBody
 
-The Black Body expression simulates the effects of [black body radiation](http://en.wikipedia.org/wiki/Black-body_radiation) within your Material. The user inputs a temperature (in Kelvin) and the resulting color and intensity can be used to drive Base Color and Emissive values to get a physically accurate result. 
+The **BlackBody** expression simulates the effects of [black body radiation](http://en.wikipedia.org/wiki/Black-body_radiation) within your Material. The user inputs a temperature (in Kelvin) and the resulting color and intensity can be used to drive Base Color and Emissive values to get a physically accurate result. 
 
 [REGION:fullwidth]
 ![](BlackBody.png)
@@ -19,7 +20,7 @@ The Black Body expression simulates the effects of [black body radiation](http:/
 ## BumpOffset
 
 [EXCERPT:BumpOffset]
-BumpOffset is the Unreal Engine 4 term for what is commonly known as 'Parallax Mapping'.  The Bump Offset expression allows a material to give the illusion of depth without the need for additional geometry.  BumpOffset materials use a grayscale _heightmap_ to give depth information.  The brighter the value in the heightmap, the more 'popped out' the material will be; these areas will parallax (shift) as a camera moves across the surface.  Darker areas in the heightmap are 'further away' and will shift the least.
+**BumpOffset** is the Unreal Engine 4 term for what is commonly known as 'Parallax Mapping'.  The Bump Offset expression allows a material to give the illusion of depth without the need for additional geometry.  BumpOffset materials use a grayscale _heightmap_ to give depth information.  The brighter the value in the heightmap, the more 'popped out' the material will be; these areas will parallax (shift) as a camera moves across the surface.  Darker areas in the heightmap are 'further away' and will shift the least.
 [/EXCERPT:BumpOffset]
 
 
@@ -39,7 +40,7 @@ BumpOffset is the Unreal Engine 4 term for what is commonly known as 'Parallax M
 
 ## ConstantBiasScale
 
-The ConstantBiasScale expression takes an input value, adds a bias value to it, and then multiplies it by a scaling factor outputting the result.  So for example, to convert input data from [-1,1] to [0,1] you would use a bias of 1.0 and a scale of 0.5.
+The **ConstantBiasScale** expression takes an input value, adds a bias value to it, and then multiplies it by a scaling factor outputting the result.  So for example, to convert input data from [-1,1] to [0,1] you would use a bias of 1.0 and a scale of 0.5.
 
 | Item | Description |
 | -------- | ----------- |
@@ -53,23 +54,32 @@ The ConstantBiasScale expression takes an input value, adds a bias value to it,
 
 ## DDX
 
-The DDX expression exposes DDX derivative calculation, a GPU hardware feature used in pixel shader calculation.
+The **DDX** expression exposes DDX derivative calculation, a GPU hardware feature used in pixel shader calculation.
 
 ## DDY
 
-The DDY expression exposes DDX derivative calculation, a GPU hardware feature used in pixel shader calculation.
+The **DDY** expression exposes DDX derivative calculation, a GPU hardware feature used in pixel shader calculation.
 
 [INCLUDE:Engine/Rendering/Materials/ExpressionReference/Depth#DepthFade]
 
 
 ## DepthOfFieldFunction
 
-The Depth of Field Function expression is designed to give artists control over what happens to a Material when it is being blurred by Depth of Field. It outputs a value between 0-1 such that 0 represents "in focus" and 1 represents "completely blurred." This is useful for interpolating between sharp and blurry versions of a texture, for instance. The Depth input allows for the existing results from the scene's Depth of Field calculations to be overridden by other calculations.
+The **Depth of Field Function** expression is designed to give artists control over what happens to a Material when it is being blurred by Depth of Field. It outputs a value between 0-1 such that 0 represents "in focus" and 1 represents "completely blurred." This is useful for interpolating between sharp and blurry versions of a texture, for instance. The Depth input allows for the existing results from the scene's Depth of Field calculations to be overridden by other calculations.
+
+[OBJECT:ComparisonSlider]
+	[PARAM:before]
+	![Blended Colors](DepthOfFieldFunction_Color.png)(w:350)
+	[/PARAM]
+	[PARAM:after]
+	![Blending regular and blurred textures](DepthOfFieldFunction_Texture.png)(w:350)
+	[/PARAM]
+[/OBJECT]
 
 [REGION:imagetable]
-| ![](DepthOfFieldFunction_Color.png)(w:305) |![](DepthOfFieldFunction_Texture.png)(w:305)  | [INCLUDE:#DOFF] |
-| :--: | :--: | :--: |
-| Blended colors | Blending regular and blurred textures | Network sample (click to enlarge) |
+| [INCLUDE:#DOFF] |
+| :--: | 
+| Network sample (click to enlarge) |
 [/REGION]
 
 <!--
@@ -83,7 +93,7 @@ The Depth of Field Function expression is designed to give artists control over
 [EXCERPT:Desaturation]
 ## Desaturation
 
-The Desaturation expression desaturates its input, or converts the colors of its input into shades of gray, based a certain percentage.
+The **Desaturation** expression desaturates its input, or converts the colors of its input into shades of gray, based a certain percentage.
 
 | Item | Description |
 | -------- | ----------- |
@@ -104,7 +114,7 @@ The Desaturation expression desaturates its input, or converts the colors of its
 [EXCERPT:Distance]
 ## Distance
 
-The Distance expression computes the (Euclidian) distance between two points/colors/positions/vectors and outputs the resulting value. This works on one, two, three and four component vectors, but both inputs to the expression must have the same number of channels.
+The **Distance** expression computes the (Euclidian) distance between two points/colors/positions/vectors and outputs the resulting value. This works on one, two, three and four component vectors, but both inputs to the expression must have the same number of channels.
 
 | Item | Description |
 | -------- | ----------- |
@@ -131,7 +141,7 @@ The Distance expression computes the (Euclidian) distance between two points/col
 [EXCERPT:FeatureLevelSwitch]
 ## FeatureLevelSwitch
 
-The Feature Level Switch node allows you to make simplified materials for lower powered devices. 
+The **Feature Level Switch** node allows you to make simplified materials for lower powered devices. 
 
 **Example Usage**: You might have a material with 10 textures overlapping and complex math, but just a single static texture for mobile (feature level ES2). 
 
@@ -150,7 +160,7 @@ The Feature Level Switch node allows you to make simplified materials for lower
 
 ## Fresnel
 
-The Fresnel expression calculates a falloff based on the dot product of the surface normal and the direction to the camera. When the surface normal points directly at the camera, a value of 0 is output. When the surface normal is perpendicular to the camera, a value of 1 is output. The result is clamped to [0,1] so you do not have any negative color in the center.
+The **Fresnel** expression calculates a falloff based on the dot product of the surface normal and the direction to the camera. When the surface normal points directly at the camera, a value of 0 is output. When the surface normal is perpendicular to the camera, a value of 1 is output. The result is clamped to [0,1] so you do not have any negative color in the center.
 
 | Item | Description |
 | -------- | ----------- |
@@ -169,7 +179,7 @@ The Fresnel expression calculates a falloff based on the dot product of the surf
 [EXCERPT:GIReplace]
 ## GIReplace
 
-GIReplace allows artists to specify a different, usually simpler, expression chain when the material is being used for GI.  
+**GIReplace** allows artists to specify a different, usually simpler, expression chain when the material is being used for GI.  
 
 **Example Usage**: Lightmass static GI and LPV dynamic GI use it.
 
@@ -190,7 +200,7 @@ GIReplace allows artists to specify a different, usually simpler, expression cha
 
 ## LightmassReplace
 
-The LightmassReplace expression simply passes through the Realtime input when compiling the material for normal rendering purposes, and passes through the Lightmass input when exporting the material to Lightmass for global illumination.  This is useful to work around material expressions that the exported version cannot handle correctly, for example WorldPosition.
+The **LightmassReplace** expression simply passes through the Realtime input when compiling the material for normal rendering purposes, and passes through the Lightmass input when exporting the material to Lightmass for global illumination.  This is useful to work around material expressions that the exported version cannot handle correctly, for example WorldPosition.
 
 | Item | Description |
 | -------- | ----------- |
@@ -224,7 +234,7 @@ The **LinearInterpolate** expression blends between two input value(s) based on
 
 ## Noise
 
-The Noise expression creates a procedural noise field, giving you control over how it is generated.
+The **Noise** expression creates a procedural noise field, giving you control over how it is generated.
 
 | Item | Description |
 | -------- | ----------- |
@@ -246,7 +256,7 @@ The Noise expression creates a procedural noise field, giving you control over h
 
 ## QualitySwitch
 
-The QualitySwitch expression allows for the use of different expression networks based on the engine is switched between quality levels, such as using lower quality on lower-end devices.
+The **QualitySwitch** expression allows for the use of different expression networks based on the engine is switched between quality levels, such as using lower quality on lower-end devices.
 
 | Item | Description |
 | -------- | ----------- |
@@ -258,7 +268,7 @@ The QualitySwitch expression allows for the use of different expression networks
 [EXCERPT:RotateAboutAxis]
 ## RotateAboutAxis
 
-The RotateAboutAxis expression rotates a three-channel vector input given the rotation axis, a point on the axis, and the angle to rotate.  It is useful for animations using WorldPositionOffset that have better quality than simple shears.
+The **RotateAboutAxis** expression rotates a three-channel vector input given the rotation axis, a point on the axis, and the angle to rotate.  It is useful for animations using WorldPositionOffset that have better quality than simple shears.
 
 | Item | Description |
 | -------- | ----------- |
@@ -278,7 +288,7 @@ In the above example, the preview plane would appear to spin on its vertical axi
 
 ## SphereMask
 
-The SphereMask expression outputs a mask value based on a distance calculation. If one input is the position of a point and the other input is the center of a sphere with some radius, the mask value is 0 outside and 1 inside with some transition area. This works on one, two, three, and four component vectors
+The **SphereMask** expression outputs a mask value based on a distance calculation. If one input is the position of a point and the other input is the center of a sphere with some radius, the mask value is 0 outside and 1 inside with some transition area. This works on one, two, three, and four component vectors
 
 | Item | Description |
 | -------- | ----------- |
@@ -302,7 +312,7 @@ In this example, the preview object will smoothly fade to black as the camera ex
 
 ## AntialiasedTextureMask
 
-The AntialiasedTextureMask expression allows you to create a material using a soft (anti-aliased) transition mask. The mask can be used to blend between two complex material properties or to fade out an alpha blended material (works well with SoftMasked). Simple specify a texture with the mask specified in one channel (red, green, blue, or alpha), set the used channel in the expression and specify the comparison value. Assuming the channel stores a grayscale value in the range 0=black to 1=white the comparison function defines if the resulting mask should be 0 or 1. This expression is a parameter, allowing the **Texture** property to be overridden by child MaterialInstances.
+The **AntialiasedTextureMask** expression allows you to create a material using a soft (anti-aliased) transition mask. The mask can be used to blend between two complex material properties or to fade out an alpha blended material (works well with SoftMasked). Simple specify a texture with the mask specified in one channel (red, green, blue, or alpha), set the used channel in the expression and specify the comparison value. Assuming the channel stores a grayscale value in the range 0=black to 1=white the comparison function defines if the resulting mask should be 0 or 1. This expression is a parameter, allowing the **Texture** property to be overridden by child MaterialInstances.
 
 | Item | Description |
 | -------- | ----------- |
