#UE4doc #docs - Added Engine Version Field. Bolded terms.

diff --git a/Engine/Documentation/Source/Engine/Rendering/Materials/ExpressionReference/Vector/VectorExpressions.INT.udn b/Engine/Documentation/Source/Engine/Rendering/Materials/ExpressionReference/Vector/VectorExpressions.INT.udn
--- a/Engine/Documentation/Source/Engine/Rendering/Materials/ExpressionReference/Vector/VectorExpressions.INT.udn
+++ b/Engine/Documentation/Source/Engine/Rendering/Materials/ExpressionReference/Vector/VectorExpressions.INT.udn
@@ -2,6 +2,7 @@ Availability:Public
 Title:Vector Expressions
 Crumbs: %ROOT%, Engine, Engine/Rendering/Materials, Engine/Rendering/Materials/ExpressionReference
 Description:Expressions that output a vector value, such as a position or normal.
+Version: 4.9
 
 [EXCERPT:TOC]
 [TOC(start:2)]
@@ -27,7 +28,7 @@ The **CameraWorldPosition** expression outputs a three-channel vector value repr
 
 The preview sphere changes color as the camera rotates.
 
-![CameraPositionWSExample.png](CameraPositionWSExample.png)(w:476) ![](CameraPositionWSNetwork.png)(w:450)
+![CameraPositionWSExample.png](CameraPositionWSExample.png)(w:450) ![](CameraPositionWSNetwork.png)(h:333)
 [/EXCERPT:CameraPositionWS]
 
 
@@ -94,7 +95,7 @@ The **ObjectOrientation** expression outputs the world-space up vector of the ob
 The **ObjectPositionWS** expression outputs the world-space center position of the object's bounds.  For example, this is useful for creating spherical lighting for foliage.
 
 
-![](WorldPositionWSScene.png)(w:460) ![](WorldPositionNetwork.png)(w:470)
+![](WorldPositionWSScene.png)(w:450) ![](WorldPositionNetwork.png)(h:349)
 
 [/EXCERPT:ObjectPositionWS]
 
@@ -105,7 +106,7 @@ The **ObjectPositionWS** expression outputs the world-space center position of t
 [EXCERPT:PixelNormalWS]
 ## PixelNormalWS
 
-The PixelNormalWS expression outputs vector data representing the direction that pixels are facing based on the current normal.
+The **PixelNormalWS** expression outputs vector data representing the direction that pixels are facing based on the current normal.
 
 ![PixelNormalWSExample.png](PixelNormalWSExample.png)
 
@@ -116,7 +117,7 @@ In this example, PixelNormalWS is fed into Base Color. Notice how the normal map
 
 ## ReflectionVectorWS
 
-The ReflectionVectorWS expression is similar in spirit to [CameraVectorWS](#CameraVectorWS), but it outputs a three-channel vector value representing the camera direction reflected across the surface normal.
+The **ReflectionVectorWS** expression is similar in spirit to [CameraVectorWS](#CameraVectorWS), but it outputs a three-channel vector value representing the camera direction reflected across the surface normal.
 
 **Example Usage:** ReflectionVector is commonly used in environment maps, where the x/y components of the reflection vector are used as UVs into a cubemap texture.
 
diff --git a/Engine/Documentation/Source/Engine/Rendering/Materials/ExpressionReference/VectorOps/VectorOpExpressions.INT.udn b/Engine/Documentation/Source/Engine/Rendering/Materials/ExpressionReference/VectorOps/VectorOpExpressions.INT.udn
--- a/Engine/Documentation/Source/Engine/Rendering/Materials/ExpressionReference/VectorOps/VectorOpExpressions.INT.udn
+++ b/Engine/Documentation/Source/Engine/Rendering/Materials/ExpressionReference/VectorOps/VectorOpExpressions.INT.udn
@@ -2,6 +2,7 @@ Availability:Public
 Title:Vector Operation Expressions
 Crumbs: %ROOT%, Engine, Engine/Rendering/Materials, Engine/Rendering/Materials/ExpressionReference
 Description:Expressions that perform operations on vector input values.
+Version: 4.9
 
 [EXCERPT:TOC]
 [TOC(start:2)]
@@ -54,7 +55,7 @@ The **ComponentMask** expression allows you to select a specific subset of chann
 
 ## DeriveNormalZ
 
-The DeriveNormalZ expression derives the Z component of a tangent space normal given the X and Y components and outputs the resulting three-channel tangent space normal. Z is calculated as Z = sqrt(1 - (x * x + y * y));
+The **DeriveNormalZ** expression derives the Z component of a tangent space normal given the X and Y components and outputs the resulting three-channel tangent space normal. Z is calculated as Z = sqrt(1 - (x * x + y * y));
 
 | Item | Description |
 | -------- | ----------- |
@@ -73,7 +74,7 @@ The DeriveNormalZ expression derives the Z component of a tangent space normal g
 
 ## Transform
 
-The Transform expression converts a three-channel vector value from one reference coordinate system to another.
+The **Transform** expression converts a three-channel vector value from one reference coordinate system to another.
 
 By default, all shader calculations in a material are done in tangent space.  The vector constants, camera vector, light vector, etc are all transformed to tangent space before being used in a material. The Transform expression allows these vectors to be transformed from tangent space to world-space, local-space, or view-space coordinate systems. In addition, it allows world-space and local-space vectors to be transformed to any of the other reference coordinate systems.
 
@@ -122,5 +123,5 @@ Here is an example of this spheremap material applied to a blob-like rock mesh:
 This node is deprecated due to major precision problems when used to derive world pos away from the origin!  Use the WorldPosition node instead.
 [/REGION]
 
-The TransformPosition expression can transform any position from screen space to the destination space specified by the expression's TransformType variable.  Currently only transforming into world space is supported.  This expression can be used to get world space coordinates in the material.  To visualize world position, you can plug it straight into emissive:
+The **TransformPosition** expression can transform any position from screen space to the destination space specified by the expression's TransformType variable.  Currently only transforming into world space is supported.  This expression can be used to get world space coordinates in the material.  To visualize world position, you can plug it straight into emissive:
 
