#UE4doc #docs - Added Engine Version Field. Bolded term.

diff --git a/Engine/Documentation/Source/Engine/Rendering/Materials/Functions/Functions.INT.udn b/Engine/Documentation/Source/Engine/Rendering/Materials/Functions/Functions.INT.udn
--- a/Engine/Documentation/Source/Engine/Rendering/Materials/Functions/Functions.INT.udn
+++ b/Engine/Documentation/Source/Engine/Rendering/Materials/Functions/Functions.INT.udn
@@ -2,12 +2,13 @@ Availability: Public
 Title:Material Functions
 Crumbs: %ROOT%, Engine, Engine/Rendering/Materials
 Description:Snippets of material graphs that can be saved in packages and reused across multiple materials.
+Version: 4.9
 
 [REGION:full_feature]
 [/REGION]
 
 [EXCERPT:FunctionsIntro]
-Material Functions are little snippets of Material graphs that can be saved in packages and reused across 
+**Material Functions** are little snippets of Material graphs that can be saved in packages and reused across 
 multiple Materials. Material Functions allow complex Material Graph networks to be saved out and quickly reused in other Materials, while also allowing 
 more complex networks to be abstracted into single nodes, making material creation easier for artists.
 
