#UE4doc #docs - Added Engine Version Field. Added Skill Level Field. Bolded Terms.

diff --git a/Engine/Documentation/Source/Engine/Rendering/Materials/Functions/Creating/CreatingFunctions.INT.udn b/Engine/Documentation/Source/Engine/Rendering/Materials/Functions/Creating/CreatingFunctions.INT.udn
--- a/Engine/Documentation/Source/Engine/Rendering/Materials/Functions/Creating/CreatingFunctions.INT.udn
+++ b/Engine/Documentation/Source/Engine/Rendering/Materials/Functions/Creating/CreatingFunctions.INT.udn
@@ -2,6 +2,8 @@ Availability: Public
 Title:Creating Material Functions
 Crumbs: %ROOT%, Engine, Engine/Rendering/Materials
 Description:A guide to the creation process for material functions.
+Version: 4.9
+SkillLevel: Advanced
 
 [TOC(start:2)]
 
@@ -16,7 +18,7 @@ By default, your new function will only have a single output node labeled _Resul
 
 ![NewResult.png](NewResult.png)
 
-If you require an input into your function, you will need to create a _FunctionInput_ node. Do this by **right-clicking** within the Material Editor and choosing _Functions > New Function Input_.
+If you require an input into your function, you will need to create a _FunctionInput_ node. Do this by **right-clicking** within the Material Editor and choosing **Functions** > **New Function Input**.
 
 ![AddInput.png](AddInput.png)
 
@@ -29,9 +31,9 @@ Functions define their interface to any materials they are used in by function i
 
 ## Publishing Your New Function
 
-In order to make the most out of your function, you will need to make sure that it shows up within the Material Function Library, available in the Material Editor. In order to do this, you will need to make sure that the _Expose to Library_ property is set to true. This is done like so:
+In order to make the most out of your function, you will need to make sure that it shows up within the **Material Function Library**, available in the Material Editor. In order to do this, you will need to make sure that the _Expose to Library_ property is set to true. This is done like so:
 
-1. Deselect all nodes within your function. This will bring up the function's base properties within the Details Window.
+1. Deselect all nodes within your function. This will bring up the function's base properties within the **Details** panel.
 
 	![ExposeToLibrary.png](ExposeToLibrary.png)
 
