#UE4doc #docs - Added Engine Version Field. Bolded terms.

diff --git a/Engine/Documentation/Source/Engine/Rendering/Materials/Functions/Overview/FunctionsOverview.INT.udn b/Engine/Documentation/Source/Engine/Rendering/Materials/Functions/Overview/FunctionsOverview.INT.udn
--- a/Engine/Documentation/Source/Engine/Rendering/Materials/Functions/Overview/FunctionsOverview.INT.udn
+++ b/Engine/Documentation/Source/Engine/Rendering/Materials/Functions/Overview/FunctionsOverview.INT.udn
@@ -2,6 +2,7 @@ Availability: Public
 Title:Material Functions Overview
 Crumbs: %ROOT%, Engine, Engine/Rendering/Materials
 Description:An overview of how material functions work and some critical concepts for their use.
+Version: 4.9
 
 [REGION:banner]
 ![](matfunc_overview_banner_NEW.jpg)
@@ -11,7 +12,7 @@ Description:An overview of how material functions work and some critical concept
 
 
 
-Material Functions are little snippets of material graphs that can be saved in packages and reused across multiple materials.  Their purpose is to streamline the process of material creation by giving instant access to commonly used networks of material nodes. For example, if you find yourself often creating a network just to handle chaotic texture panning, then it will be much faster to save that portion of your network out as a material function and just use that whenever you need such behavior.
+**Material Functions** are little snippets of material graphs that can be saved in packages and reused across multiple materials.  Their purpose is to streamline the process of material creation by giving instant access to commonly used networks of material nodes. For example, if you find yourself often creating a network just to handle chaotic texture panning, then it will be much faster to save that portion of your network out as a material function and just use that whenever you need such behavior.
 
 Functions are edited in the material editor like a normal material, but with some restrictions on what nodes can be used.  When used properly, they can reduce material redundancy which in turn reduces artist maintenance effort keeping those duplicate expressions in sync, and the inevitable bugs that arise when one duplicate is missed during a modification. 
 
@@ -28,7 +29,7 @@ As you can see above, what happens between the input and the output is entirely
 
 ## Material Function Library
 
-When your material function is finished, it should be published to the material function library for easy access during material creation. The material function library is a window in the Material Editor that contains a categorized and filterable list of available material functions.  This list is populated from any loaded functions, but also from any functions which were not loaded but were found through the **Content Browser Database** which is used by the **Content Browser**.  
+When your material function is finished, it should be published to the material function library for easy access during material creation. The material function library is a window in the Material Editor that contains a categorized and filterable list of available material functions.  This list is populated from any loaded functions, but also from any functions which were not loaded but were found through the **Content Browser Database** which is used by the Content Browser.  
 
 ![FunctionLibrary.png](FunctionLibrary.png)
 
@@ -97,12 +98,12 @@ They have the following properties and data pins:
 ### FunctionOutput Nodes
 
 
-FunctionOutput nodes provide the means by which processed data will exit the final function for further use in a material. As with FunctionInput nodes, a function may have any number of these nodes, leading to any number of potential outputs.
+**FunctionOutput** nodes provide the means by which processed data will exit the final function for further use in a material. As with FunctionInput nodes, a function may have any number of these nodes, leading to any number of potential outputs.
 
 
 ![OutputNode.png](OutputNode.png)
 
-FunctionOutput nodes have the following properties:
+**FunctionOutput** nodes have the following properties:
 
 | Item | Description | 
 |---|---|
@@ -159,7 +160,7 @@ Most of the time you will want to preview a function output, so those will be pr
 
 ![PreviewingOutput.png](PreviewingOutput.png)
 
-Function input nodes have some options for specifying preview values, since they do not know the actual values that will be used with in a material.  Each input has a built in PreviewValue that can be used to show a constant for float input types.  Function inputs also have a 'Preview' connector which allows you to override that built in value with any value that matches the input's type.  In this example, a texture sample has been used to provide the preview for a float 3 input.
+Function input nodes have some options for specifying preview values, since they do not know the actual values that will be used with in a material.  Each input has a built in _Preview Value_ that can be used to show a constant for float input types.  Function inputs also have a 'Preview' connector which allows you to override that built in value with any value that matches the input's type.  In this example, a texture sample has been used to provide the preview for a float 3 input.
 
 ![PreviewTexture.png](PreviewTexture.png)
 
@@ -179,15 +180,15 @@ To use a texture parameter with a function, make a texture input and connect it
 
 ![TextureParameterFunction.png](TextureParameterFunction.png)(w:720)
 
-Then, in the material that uses the function, place a TextureObjectParameter node and connect it to the texture input:
+Then, in the material that uses the function, place a **TextureObjectParameter** node and connect it to the texture input:
 
 ![TextureParameterMaterial.png](TextureParameterMaterial.png)(w:720)
 
-Similarly for static switch parameters, make a static bool input and connect it to a StaticSwitch node:
+Similarly for static switch parameters, make a **Static Bool** input and connect it to a StaticSwitch node:
 
 ![StaticSwitchFunction.png](StaticSwitchFunction.png)(w:720)
 
-Then, in the material that uses the function, place a StaticBoolParameter node and connect it to the static bool input:
+Then, in the material that uses the function, place a **StaticBoolParameter** node and connect it to the static bool input:
 
 ![StaticSwitchMaterial.png](StaticSwitchMaterial.png)(w:720)
 
