#UE4doc #docs - Added Engine Version Field.

diff --git a/Engine/Documentation/Source/Engine/Rendering/Materials/Functions/Reference/FunctionReference.INT.udn b/Engine/Documentation/Source/Engine/Rendering/Materials/Functions/Reference/FunctionReference.INT.udn
--- a/Engine/Documentation/Source/Engine/Rendering/Materials/Functions/Reference/FunctionReference.INT.udn
+++ b/Engine/Documentation/Source/Engine/Rendering/Materials/Functions/Reference/FunctionReference.INT.udn
@@ -2,6 +2,7 @@ Availability: Public
 title:Material Function Reference
 Crumbs: %ROOT%, Engine, Engine/Rendering/Materials
 Description:Reference page for each of the material functions, separated by category.
+Version: 4.9
 
 
 Below is breakdown of the various material functions included within Unreal Editor by default, sorted by category.
