#UE4doc #docs - Added Engine Version field. Updated Project Browser image.

diff --git a/Engine/Documentation/Source/Engine/Rendering/Materials/HowTo/MaterialHowTo.INT.udn b/Engine/Documentation/Source/Engine/Rendering/Materials/HowTo/MaterialHowTo.INT.udn
--- a/Engine/Documentation/Source/Engine/Rendering/Materials/HowTo/MaterialHowTo.INT.udn
+++ b/Engine/Documentation/Source/Engine/Rendering/Materials/HowTo/MaterialHowTo.INT.udn
@@ -7,6 +7,7 @@ Related: Engine/Rendering/Materials/MaterialInputs
 Related: Engine/Rendering/Materials/Editor
 Related: Engine/Rendering/Materials/IntroductionToMaterials
 Related: Engine/Rendering/Materials/MaterialProperties
+Version: 4.9
 
 [VAR:TopicCompact]
 [OBJECT:TopicCompact]
@@ -48,7 +49,7 @@ If you already have a project that also includes the Starter Content, please fee
 
 1. From the **Unreal Project Browser** window, create a **New Project** and use the **Blueprint Third Person** template. 
 
-1. Make sure that **Include Starter Content** is enabled. 
+1. Make sure that **With Starter Content** is enabled. 
 
 1. Give the project a location to be saved to. 
 
