#UE4doc #docs:
Added engine version field
Removed Overview header
Fixed capitalization for a few terms
Changed font settings for many terms
Replaced a few images due to outdated UI elements

diff --git a/Engine/Documentation/Source/Engine/Blueprints/UserGuide/Events/Custom/CustomEvents.INT.udn b/Engine/Documentation/Source/Engine/Blueprints/UserGuide/Events/Custom/CustomEvents.INT.udn
--- a/Engine/Documentation/Source/Engine/Blueprints/UserGuide/Events/Custom/CustomEvents.INT.udn
+++ b/Engine/Documentation/Source/Engine/Blueprints/UserGuide/Events/Custom/CustomEvents.INT.udn
@@ -73,7 +73,7 @@ They have a fairly straightforward workflow:
 
 ### Creating Custom Events
 
-1. Create a Custom Event node by **right-clicking** and choosing **Add Custom Event...** from the context menu. 
+1. Create a Custom Event node by **Right-clicking** and choosing **Add Custom Event...** from the context menu. 
 
 	![](add_custom_event.png)
 
@@ -109,7 +109,7 @@ This Custom Event example prints a String to the screen.
 
 ### Calling Custom Events
 
-Your Custom Event and its associated network have been created, but unlike regular Events, there are no preset conditions for the Custom Event to fire. To call your Custom Event, **right-click** and choose **Call Function > [Custom Event Name]** from the  context menu. 
+Your Custom Event and its associated network have been created, but unlike regular Events, there are no preset conditions for the Custom Event to fire. To call your Custom Event, **Right-click** and choose **Call Function > [Custom Event Name]** from the  context menu. 
 
 ![](call_ce_context_menu.png)
 
