#UE4doc #docs:
Added engine version field
Replaced a couple images due to outdated UI
Fixed capitalization for a few terms
Changed font settings for multiple terms

diff --git a/Engine/Documentation/Source/Engine/Blueprints/UserGuide/FlowControl/FlowControl.INT.udn b/Engine/Documentation/Source/Engine/Blueprints/UserGuide/FlowControl/FlowControl.INT.udn
--- a/Engine/Documentation/Source/Engine/Blueprints/UserGuide/FlowControl/FlowControl.INT.udn
+++ b/Engine/Documentation/Source/Engine/Blueprints/UserGuide/FlowControl/FlowControl.INT.udn
@@ -46,7 +46,7 @@ Description:Nodes that allow for controlling the flow of execution based on cond
 
 [EXCERPT:Intro]
 
-[INCLUDE:Shared/Glossary/F/#Glossary_Flow_Control]
+[INCLUDE:Shared/Glossary/F/#GlossaryFlow_Control]
 
 [/EXCERPT:Intro]
 
@@ -61,8 +61,8 @@ In general, switches have an execution input, and a data input for the type of d
 
 ### Editing Switch Nodes
 
-When an **Int**, **String**, or **Name** switch node is added to a _Blueprint_, the only output execution pin available is the **Default** pin. The **Default** output execution pin will fire if the 
-input fails to match any of the other specified output pins.  It can be removed by either **right-clicking** on the pin and selecting **Remove Execution Pin** or by unchecking the **Has Default Pin** option in the 
+When an **Int**, **String**, or **Name** switch node is added to a Blueprint, the only output execution pin available is the **Default** pin. The **Default** output execution pin will fire if the 
+input fails to match any of the other specified output pins.  It can be removed by either **Right-clicking** on the pin and selecting **Remove Execution Pin** or by unchecking the **Has Default Pin** option in the 
 **Details** tab for the switch node.
 
 ![](NewSwitchOnInt.png)
@@ -79,7 +79,7 @@ input fails to match any of the other specified output pins.  It can be removed
 
 	![](SwitchOnInt_AddPin.png)
 
-	Clicking on **Add Pin** further will add more pins, with the value incremented by 1 each time.  To delete an execution pin, **right-click** on the pin and select **Remove Execution Pin**. Note that removing an output execution pin on a 
+	Clicking on **Add Pin** further will add more pins, with the value incremented by 1 each time.  To delete an execution pin, **Right-click** on the pin and select **Remove Execution Pin**. Note that removing an output execution pin on a 
 	**Switch on Int** node causes any higher-valued pins to decrease their value by 1 to fill the gap.
 
 
@@ -97,7 +97,7 @@ input fails to match any of the other specified output pins.  It can be removed
 	![](SwitchOnString_WithPin.png)
 
 
-Repeat the process for any other pins you would like to add.  Pins can also be duplicated, inserted, or deleted by clicking on the ![Pin Options](PinOptionsButton.png) icon by the desired pin in the **Details** tab.   To delete an execution pin, you can also **right-click** on the pin and select **Remove Execution Pin**.
+Repeat the process for any other pins you would like to add.  Pins can also be duplicated, inserted, or deleted by clicking on the ![Pin Options](PinOptionsButton.png) icon by the desired pin in the **Details** tab.   To delete an execution pin, you can also **Right-click** on the pin and select **Remove Execution Pin**.
 
 ## Standard Flow Control Nodes
 
@@ -110,12 +110,12 @@ These nodes provide a variety of means to control the flow of execution.
 ![Branch_Example.png](Branch_Example.png)
 
 
-The **Branch** node serves as a simple way to create decision-based flow from a single _true_/_false_ condition. 
+The **Branch** node serves as a simple way to create decision-based flow from a single true/false condition. 
 Once executed, the Branch node looks at the incoming value of the attached Boolean, and outputs an execution 
 pulse down the appropriate output.
 
-In this simple example, the branch is checking the current state of a boolean variable. If _true_, it sets 
-the color of a light to be red. If _false_, it sets it to blue.
+In this simple example, the branch is checking the current state of a boolean variable. If true, it sets 
+the color of a light to be red. If false, it sets it to blue.
 
 **Click for larger view**
 [REGION:lightbox]
@@ -291,7 +291,7 @@ A **Gate** node is used as a way to open and close a stream of execution. Simply
 execution pulses, and the current state of the gate (open or closed) determines whether those pulses pass 
 out of the Exit output or not. 
 
-In this simple example, a timeline with no tracks, **with both auto-play and looping activated**, updates 
+In this simple example, a timeline with no tracks, with both auto-play and looping activated, updates 
 the Enter input pin of a gate. In the level are two triggers. One trigger opens the gate, the other closes it. 
 If the gate is open, then execution pulses leave the Exit pin and a Print String is called that logs a message to the screen. Once the player touches the Close trigger, the gate closes and the message stops displaying. If they then touch the Open trigger, the message starts to appear again.
 
@@ -340,7 +340,7 @@ which, when played in order, reveal a special message.
 | **Start Index** | Takes in an Int to represent the output index the MultiGate should use first. A value of -1 is the same as not specifying a start point. |
 | [REGION:tablesection]Output Pins[/REGION] ||
 | **Out #** | Each output represents a possible output pin that the MultiGate may use to send out a routed execution pulse. |
-| **Add pin** | Though not really an output pin, this button allows you to add as many outputs as you like. Outputs can be removed by **right-clicking** and choosing Remove Output Pin. |
+| **Add pin** | Though not really an output pin, this button allows you to add as many outputs as you like. Outputs can be removed by **Right-clicking** and choosing Remove Output Pin. |
 
 
 ### Sequence
@@ -367,18 +367,18 @@ time as one another.
 | **(Unlabeled)** | This is the primary execution input that takes in any pulses that need to be routed by this Sequence. |
 | [REGION:tablesection]Output Pins[/REGION] ||
 | **Out #** | Each output represents a possible output pin that the Sequence may use to send out a routed execution pulse. |
-| **Add pin** | Though not really an output pin, this button allows you to add as many outputs as you like. Outputs can be removed by **right-clicking** and choosing Remove Output Pin. |
+| **Add pin** | Though not really an output pin, this button allows you to add as many outputs as you like. Outputs can be removed by **Right-clicking** and choosing Remove Output Pin. |
 [/EXCERPT:Sequence]
 
 ### WhileLoop
 
 ![WhileLoop_Example.png](WhileLoop_Example.png)
 
-The **WhileLoop** node will output a result so long as a specific condition is _true_. During each iteration of the 
-loop, it checks to see the current status of its input boolean value. As soon as it reads _false_, the loop breaks.
+The **WhileLoop** node will output a result so long as a specific condition is true. During each iteration of the 
+loop, it checks to see the current status of its input boolean value. As soon as it reads false, the loop breaks.
 
 In this example, a WhileLoop is used alongside an integer iteration network to count off the number of times 
-the loop has run. Once it reaches its 10th iteration, the test boolean is set to _false_, ending the loop behavior.
+the loop has run. Once it reaches its 10th iteration, the test boolean is set to false, ending the loop behavior.
 
 [REGION:warning]
 	As with While loops in programming languages, extra care must be taken to prevent infinite loops from occurring.
