#UE4doc #docs:
Resubmitting as not all changes were saved

diff --git a/Engine/Documentation/Source/Engine/Blueprints/UserGuide/FlowControl/FlowControl.INT.udn b/Engine/Documentation/Source/Engine/Blueprints/UserGuide/FlowControl/FlowControl.INT.udn
--- a/Engine/Documentation/Source/Engine/Blueprints/UserGuide/FlowControl/FlowControl.INT.udn
+++ b/Engine/Documentation/Source/Engine/Blueprints/UserGuide/FlowControl/FlowControl.INT.udn
@@ -2,6 +2,7 @@ Availability:Public
 Title:Flow Control
 Crumbs: %ROOT%, Engine, Engine/Blueprints, Engine/Blueprints/UserGuide
 Description:Nodes that allow for controlling the flow of execution based on conditions.
+version: 4.9
 
 [VAR:Topic]
 [OBJECT:Topic]
