#UE4doc #docs:
Added engine version field
Changed font settings for multiple terms

diff --git a/Engine/Documentation/Source/Engine/Blueprints/UserGuide/Functions/BlueprintFunctions.INT.udn b/Engine/Documentation/Source/Engine/Blueprints/UserGuide/Functions/BlueprintFunctions.INT.udn
--- a/Engine/Documentation/Source/Engine/Blueprints/UserGuide/Functions/BlueprintFunctions.INT.udn
+++ b/Engine/Documentation/Source/Engine/Blueprints/UserGuide/Functions/BlueprintFunctions.INT.udn
@@ -2,6 +2,7 @@ Availability:Public
 Title:Functions
 Crumbs: %ROOT%, Engine, Engine/Blueprints, Engine/Blueprints/UserGuide
 Description:Node graphs that can be executed, or called, from another graph.
+version: 4.9
 
 [VAR:Topic]
 [OBJECT:Topic]
@@ -44,7 +45,7 @@ Description:Node graphs that can be executed, or called, from another graph.
 
 
 [EXCERPT:Intro]
-**Functions** are node graphs belonging to a particular **Blueprint** that can be executed, or _called_, from another 
+**Functions** are node graphs belonging to a particular **Blueprint** that can be executed, or called, from another 
 graph within the Blueprint. Functions have a single entry point designated by a node with the name of the Function 
 containing a single exec output pin. When the Function is called from another graph, the output exec pin is activated 
 causing the connected network to execute.
@@ -67,9 +68,9 @@ The property contains the following settings:
 
 | **Setting** | **Description** |
 | --- | --- |
-| **Public** | A setting of _Public_ means that any other object can call this Function. This is the default setting. |
-| **Protected** | A setting of _Protected_ means that the Function can only be called by the current Blueprint, as well as any Blueprints that derive from the current Blueprint. |
-| **Private** | A setting of _Private_ means that the Function can only be called from within the current Blueprint. |
+| **Public** | A setting of Public means that any other object can call this Function. This is the default setting. |
+| **Protected** | A setting of Protected means that the Function can only be called by the current Blueprint, as well as any Blueprints that derive from the current Blueprint. |
+| **Private** | A setting of Private means that the Function can only be called from within the current Blueprint. |
 [/EXCERPT:AccessSpecifiers]
 
 (#PurevsImpure)
@@ -77,20 +78,20 @@ The property contains the following settings:
 
 [EXCERPT:PureVsImpure]
 Functions can either be **Pure** or **Impure**. The main difference is that Pure Functions promise not to modify state or the 
-members of the class in any way, while Impure Functions are free to modify state. Pure Functions are generally used for _getter_ 
+members of the class in any way, while Impure Functions are free to modify state. Pure Functions are generally used for getter 
 Functions or operators that just output a data value.
 
 Impure Functions must be explicitly executed by connecting exec wires to [Function Call nodes](Engine/Blueprints/UserGuide/Functions/FunctionCalls) in an [EventGraph](Engine/Blueprints/UserGuide/EventGraph). 
 Impure Functions are designated using one of these methods:
 
-* Specifying the `BlueprintCallable` keyword in the Function declaration for Functions defined in code. 
+* Specifying the **BlueprintCallable** keyword in the UFUNCTION declaration for UFUNCTIONs defined in code. 
 * Leaving the **Pure** checkbox unchecked for Functions added through the **Blueprint Editor**.
 
 
 Pure Functions are wired to Data Pins and are automatically executed by the compiler when the data relying on them is required. This 
-means that a Pure Function will be **called one time for each node it is connected to**. Pure Functions are designated using one of these methods:
+means that a Pure Function will be called one time for each node it is connected to. Pure Functions are designated using one of these methods:
 
-* Specifying the `BlueprintPure` keyword in the Function declaration for Functions defined in code.
+* Specifying the **BlueprintPure** keyword in the Function declaration for Functions defined in code.
 * Checking the **Pure** checkbox for Functions added through the **Blueprint Editor**.
 [/EXCERPT:PureVsImpure]
 
@@ -230,7 +231,7 @@ The network below takes two vector variables, and on every tick, calculates the
 
 ![](EventGraph.png)(w:700)
 
-In this example, we have created 2 public vector variables. In the settings for each one, we set **Show 3D Widget** to _True_. Now, after adding the Blueprint to the level, 3D widgets will appear at the positions defined by Point A and Point B.
+In this example, we have created 2 public vector variables. In the settings for each one, we set **Show 3D Widget** to True. Now, after adding the Blueprint to the level, 3D widgets will appear at the positions defined by Point A and Point B.
 The 3D widgets allow us to easily change the values of **Point A** and **Point B** by moving them in the viewport.
 
 ![](Points1And2Viewport.png)(w:700)
@@ -253,7 +254,7 @@ Inside another Blueprint, which is a projectile that gets spawned when the playe
 ![](DefaultProjectile.png)
 [/REGION:fullwidth]
 
-The script, which is the default **MyProjectile** Blueprint that is included with the _Blueprint First Person Template Project_, already does something when the projectile hits something that is simulating physics and adds impulse at the location it was hit. Say we wanted to check if we hit the player and if we did, to call the **Take Damage** Function. 
+The script, which is the default **MyProjectile** Blueprint that is included with the Blueprint First Person Template Project, already does something when the projectile hits something that is simulating physics and adds impulse at the location it was hit. Say we wanted to check if we hit the player and if we did, to call the **Take Damage** Function. 
 
 We can do so by dragging off the **Other** pin of the **Event Hit** and **Casting To** our **MyCharacter** Blueprint. 
 
