#UE4doc #docs:
Added engine version field

diff --git a/Engine/Documentation/Source/Engine/Blueprints/UserGuide/Functions/FunctionCalls/FunctionCalls.INT.udn b/Engine/Documentation/Source/Engine/Blueprints/UserGuide/Functions/FunctionCalls/FunctionCalls.INT.udn
--- a/Engine/Documentation/Source/Engine/Blueprints/UserGuide/Functions/FunctionCalls/FunctionCalls.INT.udn
+++ b/Engine/Documentation/Source/Engine/Blueprints/UserGuide/Functions/FunctionCalls/FunctionCalls.INT.udn
@@ -3,6 +3,7 @@ Title:Function Calls
 Crumbs: %ROOT%, Engine, Engine/Blueprints, Engine/Blueprints/UserGuide
 Description: Nodes that execute both code-defined and user-defined functions of a target Actor or Object.
 Related: Engine\Blueprints\UserGuide\Functions
+version: 4.9
 
 [VAR:Topic]
 [OBJECT:Topic]
