Setting these two to public so that they will publish correctly.

diff --git a/Engine/Documentation/Source/Platforms/SteamVR/QuickStart/2/SteamVRQuickStart_2.INT.udn b/Engine/Documentation/Source/Platforms/SteamVR/QuickStart/2/SteamVRQuickStart_2.INT.udn
--- a/Engine/Documentation/Source/Platforms/SteamVR/QuickStart/2/SteamVRQuickStart_2.INT.udn
+++ b/Engine/Documentation/Source/Platforms/SteamVR/QuickStart/2/SteamVRQuickStart_2.INT.udn
@@ -1,7 +1,7 @@
 Title: 2. Setting up UE4 to work with SteamVR
 Crumbs:%ROOT%, Platforms
 Description: Setting up SteamVR and UE4 to work together
-Availability:Doc
+Availability:Public
 Navigation:platform
 Platform:SteamVR
 
diff --git a/Engine/Documentation/Source/Platforms/SteamVR/QuickStart/SteamVRQuickStart.INT.udn b/Engine/Documentation/Source/Platforms/SteamVR/QuickStart/SteamVRQuickStart.INT.udn
--- a/Engine/Documentation/Source/Platforms/SteamVR/QuickStart/SteamVRQuickStart.INT.udn
+++ b/Engine/Documentation/Source/Platforms/SteamVR/QuickStart/SteamVRQuickStart.INT.udn
@@ -1,7 +1,7 @@
 Title: SteamVR Quick Start
 Crumbs:%ROOT%, Platforms
 Description: A How To guide for setting up SteamVR to work with UE4
-Availability:Doc
+Availability:public
 Navigation:platform
 Platform:SteamVR
 
