#UE4doc #docs:
Added engine version field
Changed font settings for multiple terms
Rehauled the sections that used "IconSection" as they did not follow documentation guidelines

diff --git a/Engine/Documentation/Source/Engine/Blueprints/UserGuide/Graphs/Graphs.INT.udn b/Engine/Documentation/Source/Engine/Blueprints/UserGuide/Graphs/Graphs.INT.udn
--- a/Engine/Documentation/Source/Engine/Blueprints/UserGuide/Graphs/Graphs.INT.udn
+++ b/Engine/Documentation/Source/Engine/Blueprints/UserGuide/Graphs/Graphs.INT.udn
@@ -2,6 +2,7 @@ Availability:Public
 Title:Graphs
 Crumbs: %ROOT%, Engine, Engine/Blueprints, Engine/Blueprints/UserGuide
 Description:Node graph that uses events and function calls to perform actions in response to events associated with the Blueprint.
+version: 4.9
 
 [VAR:Topic]
 [OBJECT:Topic]
@@ -52,80 +53,47 @@ properties, but the key principles of adding variable references, adding and con
 [/EXCERPT:GraphIntro]
 
 
-## Overview of Graph Types
+##Graph Types
 
-[OBJECT:IconSection]
-	[PARAM:icon]
-		[![%Engine/Blueprints/UserGuide/EventGraph:title%](Engine/Blueprints/UserGuide/EventGraph/eventgraph_topic.png)(w:300)](Engine/Blueprints\UserGuide\EventGraph)
-	[/PARAM]
-	[PARAM:title]
-		[%Engine/Blueprints/UserGuide/EventGraph:title%](Engine/Blueprints\UserGuide\EventGraph)
-	[/PARAM]
-	[PARAM:description]
-		**EventGraphs** are the most general type of _Blueprint_ graph.  Each new _Blueprint Class_ will contain one **EventGraph** when it is created, although more can be added.  These additional **EventGraphs** can
-help to organize your _Blueprint_ networks.  **EventGraphs** generally contain the networks for the gameplay behavior of the _Blueprint_, with [**Events**](Engine/Blueprints\UserGuide\Events), 
-[**CustomEvents**](Engine/Blueprints\UserGuide\Events\Custom), and **Input** nodes serving to initiate the flow of execution through the **EventGraph**.
-	[/PARAM]
-	[PARAM:path]
-		[RELATIVE:Engine/Blueprints/UserGuide/EventGraph]
-	[/PARAM]
-[/OBJECT]
-[OBJECT:IconSection]
-	[PARAM:icon]
-		[![%Engine/Blueprints/UserGuide/UserConstructionScript:title%](Engine/Blueprints/UserGuide/UserConstructionScript/ucs_topic.png)(w:300)](Engine/Blueprints\UserGuide\UserConstructionScript)
-	[/PARAM]
-	[PARAM:title]
-		[%Engine/Blueprints/UserGuide/UserConstructionScript:title%](Engine/Blueprints\UserGuide\UserConstructionScript)
-	[/PARAM]
-	[PARAM:description]
-		**ConstructionScripts** are unique to _Blueprint Classes_, with only one **ConstructionScript** graph present in each _Blueprint Class_.  **ConstructionScripts** are useful for _Blueprint Class_ initialization,
-as they run right after the **Components** list is set up for the _Blueprint Class_.  
-	[/PARAM]
-	[PARAM:path]
-		[RELATIVE:Engine/Blueprints/UserGuide/UserConstructionScript]
-	[/PARAM]
-[/OBJECT]
-[OBJECT:IconSection]
-	[PARAM:icon]
-		[![%Engine/Blueprints/UserGuide/Functions:title%](Engine/Blueprints/UserGuide/Functions/function_topic.png)(w:300)](Engine/Blueprints\UserGuide\Functions)
-	[/PARAM]
-	[PARAM:title]
-		[%Engine/Blueprints/UserGuide/Functions:title%](Engine/Blueprints\UserGuide\Functions)
-	[/PARAM]
-	[PARAM:description]
-		[INCLUDE:Engine/Blueprints/UserGuide/Functions#Intro]
-	[/PARAM]
-	[PARAM:path]
-		[RELATIVE:Engine/Blueprints/UserGuide/Functions]
-	[/PARAM]
-[/OBJECT]
-[OBJECT:IconSection]
-	[PARAM:icon]
-		[![](Engine/Blueprints/UserGuide/macro_topic_2.png)(w:300)](Engine/Blueprints\UserGuide\Macros)
-	[/PARAM]
-	[PARAM:title]
-		[%Engine/Blueprints/UserGuide/Macros:title%](Engine/Blueprints\UserGuide\Macros)
-	[/PARAM]
-	[PARAM:description]
-		[INCLUDE:Engine/Blueprints/UserGuide/Macros#Intro]
-	[/PARAM]
-	[PARAM:path]
-		[RELATIVE:Engine/Blueprints/UserGuide/Macros]
-	[/PARAM]
-[/OBJECT]
+###Event Graphs
+
+![](Engine/Blueprints/UserGuide/EventGraph/eventgraph_topic.png)(w:300)
+
+**Event Graphs** are the most general type of Blueprint graph.  Each new Blueprint Class will contain one Event Graph when it is created, although more can be added.  These additional Event Graphs can
+help to organize your Blueprint networks.  Event Graphs generally contain the networks for the gameplay behavior of the Blueprint, with [Events](Engine/Blueprints\UserGuide\Events), 
+[CustomEvents](Engine/Blueprints\UserGuide\Events\Custom), and **Input** nodes serving to initiate the flow of execution through the Event Graph.
+
+For more information about this section, see [%Engine/Blueprints/UserGuide/EventGraph:title%](Engine/Blueprints\UserGuide\EventGraph)
+
+###Construction Script
+
+![%Engine/Blueprints/UserGuide/UserConstructionScript:title%](Engine/Blueprints/UserGuide/UserConstructionScript/ucs_topic.png)(w:300)
+
+**Construction Scripts** are unique to Blueprint Classes, with only one ConstructionScript graph present in each Blueprint Class.  Construction Scripts are useful for Blueprint Class initialization,
+as they run right after the **Components** list is set up for the Blueprint Class.
+
+For more information about this section, see [%Engine/Blueprints/UserGuide/UserConstructionScript:title%](Engine/Blueprints\UserGuide\UserConstructionScript)
 
+###Functions
 
+![%Engine/Blueprints/UserGuide/Functions:title%](Engine/Blueprints/UserGuide/Functions/function_topic.png)(w:300)
 
+[INCLUDE:Engine/Blueprints/UserGuide/Functions#Intro]
 
+For more information about this section, see [%Engine/Blueprints/UserGuide/Functions:title%](Engine/Blueprints\UserGuide\Functions)
 
+###Macros
 
+![](Engine/Blueprints/UserGuide/macro_topic_2.png)(w:300)
 
+[INCLUDE:Engine/Blueprints/UserGuide/Macros#Intro]
 
+For more information about this section, see [%Engine/Blueprints/UserGuide/Macros:title%](Engine/Blueprints\UserGuide\Macros)
 
 ## Working with Graphs
 
-Whether your _Graph_ is a _Construction Script_, _EventGraph_, _Function_, or _Macro_, you will edit it in the [Graph](Engine/Blueprints\Editor\Modes\GraphPanel) mode of the [](Engine/Blueprints\Editor). Fundamentally,
-all _Graphs_ contain networks of _Nodes_ connected by _Wires_.
+Whether your Graph is a Construction Script, EventGraph, Function, or Macro, you will edit it in the [Graph](Engine/Blueprints\Editor\Modes\GraphPanel) mode of the [](Engine/Blueprints\Editor). Fundamentally,
+all Graphs contain networks of Nodes connected by Wires.
 
 
 
