#UE4doc #docs:
Added engine version field
Replaced a few images due to outdated UI
Fixed capitalization in a few places
Changed font settings for many terms

diff --git a/Engine/Documentation/Source/Engine/Blueprints/UserGuide/Macros/BlueprintMacros.INT.udn b/Engine/Documentation/Source/Engine/Blueprints/UserGuide/Macros/BlueprintMacros.INT.udn
--- a/Engine/Documentation/Source/Engine/Blueprints/UserGuide/Macros/BlueprintMacros.INT.udn
+++ b/Engine/Documentation/Source/Engine/Blueprints/UserGuide/Macros/BlueprintMacros.INT.udn
@@ -2,6 +2,7 @@ Availability:Public
 Title:Macros
 Crumbs: %ROOT%, Engine, Engine/Blueprints, Engine/Blueprints/UserGuide
 Description:Collapsed networks of nodes that can be executed, or called, from another graph.
+version: 4.9
 
 [VAR:Topic]
 [OBJECT:Topic]
@@ -46,7 +47,7 @@ Description:Collapsed networks of nodes that can be executed, or called, from an
 [EXCERPT:Intro]
 **Blueprint Macros**, or **Macros**, are essentially the same as collapsed graphs of nodes. They have an entry point and exit point designated by 
 tunnel nodes. Each tunnel can have any number of execution or data pins which are visible on the macro node when used in other 
-_Blueprints_ and graphs.
+Blueprints and graphs.
 
 [/EXCERPT:Intro]
 
@@ -56,21 +57,21 @@ _Blueprints_ and graphs.
 ## Creating a Blueprint Macro
 
 [EXCERPT:MacroCreationLib]
-_Blueprint Macros_ can be created within a _Blueprint Class_ or _Level Blueprint_, just like _Blueprint Functions_. They can also be organized into _Blueprint Macro Libraries_.
+Blueprint Macros can be created within a Blueprint Class or Level Blueprint, just like Blueprint Functions. They can also be organized into Blueprint Macro Libraries.
 
-To create a _Blueprint Macro_ inside a _Blueprint Class_, _Level Blueprint_, or _Blueprint Macro Library_:
+To create a Blueprint Macro inside a Blueprint Class, Level Blueprint, or Blueprint Macro Library:
 
 1. In the **My Blueprint** tab create a new macro, by clicking on the **Add Button** ![](Engine\Blueprints\Editor\UIComponents\MyBlueprint\plus_Button.png) on the macros list header ![](Engine\Blueprints\Editor\UIComponents\MyBlueprint\myblueprint_macro.png).
 
-1. Enter a name for your _Blueprint Macro_.
+1. Enter a name for your Blueprint Macro.
 
 	[REGION:raw]
 	![](name_macro_blueprint.png)
 	[/REGION]
 	
-Your _Blueprint Macro_ will open in a new tab in the **Graph** tab of the Blueprint Editor.
+Your Blueprint Macro will open in a new tab in the **Graph** tab of the Blueprint Editor.
 
-You can also create a _Blueprint Macro_ by **right-clicking** in the **My Blueprint** tab and selecting **Macro**.
+You can also create a Blueprint Macro by **Right-clicking** in the **My Blueprint** tab and selecting **Macro**.
 
 ![](add_macro_rtclick.png)
 
@@ -78,11 +79,11 @@ You can also create a _Blueprint Macro_ by **right-clicking** in the **My Bluepr
 
 ## Building a Blueprint Macro
 [EXCERPT:BuildingMacros]
-When you first create a _Blueprint Macro_, a new graph will open containing an **Inputs** tunnel node and an **Outputs** tunnel node.
+When you first create a Blueprint Macro, a new graph will open containing an **Inputs** tunnel node and an **Outputs** tunnel node.
 
 ![](input_output_macro.png)
 
-In the **Details** pane for your _Blueprint Macro_, you can add input and output execution and data pins. You can also set the **Description**, **Category**, and **Instance Color** of your _Blueprint Macro_. 
+In the **Details** pane for your Blueprint Macro, you can add input and output execution and data pins. You can also set the **Description**, **Category**, and **Instance Color** of your Blueprint Macro. 
 
 [REGION:raw]
 ![](new_macro_details.png)
@@ -92,13 +93,13 @@ To add input or output parameters:
 
 1. Click on the **New** button in either the **Inputs** or **Outputs** section of the **Details** pane.
 
-1. Name the new parameter and set its type using the dropdown menu.  In this example, there is an _Integer_ data input parameter named **Score**, an input execution pin named **Test**, and two output execution pins named **Win** and **Lose**. 
+1. Name the new parameter and set its type using the dropdown menu.  In this example, there is an Integer data input parameter named **Score**, an input execution pin named **Test**, and two output execution pins named **Win** and **Lose**. 
 	
 	[REGION:raw]
 	![](macro_details.png)
 	[/REGION]
 
-	The tunnel nodes in the _Blueprint Macro_ graph will automatically update with the correct pins.
+	The tunnel nodes in the Blueprint Macro graph will automatically update with the correct pins.
 
 	![](macro_tunnels_with_pins.png)
 
@@ -110,18 +111,18 @@ To add input or output parameters:
 
 To change the location of the pin for this parameter on the edge of the node, use the up and down arrows in the expanded **Details** pane entry.
 
-To give your _Blueprint Macro_ some functionality, connect data and execution wires to the pins of your **Inputs** and **Outputs** tunnel nodes and create a network between them.
+To give your Blueprint Macro some functionality, connect data and execution wires to the pins of your **Inputs** and **Outputs** tunnel nodes and create a network between them.
 
 ![](score_comparison_example_macro.png)
 
-This example _Blueprint Macro_ checks if the score that is input into the _Macro_ is greater than the necessary score for winning, and then triggers different output execution flows depending on 
-the result of the comparison. Note here that the **Test** and **Score** pins were flipped using the up and down arrows in the **Details** pane, to avoid wires crossing in the _Blueprint Macro_ graph. 
+This example Blueprint Macro checks if the score that is input into the Macro is greater than the necessary score for winning, and then triggers different output execution flows depending on 
+the result of the comparison. Note here that the **Test** and **Score** pins were flipped using the up and down arrows in the **Details** pane, to avoid wires crossing in the Blueprint Macro graph. 
 	 
 ![](move_pin_arrows.png)
 
 
-Unlike _Functions_, _Macros_ can have more than one output execution pin, so you can have execution flow like this where different output execution pins are activated
-depending on the results of graph logic. Also, you could have a _Macro_ with no execution pins that simply manipulated data, as long as the nodes within the _Macro_ were not execution nodes.
+Unlike Functions, Macros can have more than one output execution pin, so you can have execution flow like this where different output execution pins are activated
+depending on the results of graph logic. Also, you could have a Macro with no execution pins that only manipulates data, as long as the nodes within the Macro are not execution nodes.
 
 [REGION:fullwidth]
 ![](data_macro.png)
@@ -131,27 +132,27 @@ depending on the results of graph logic. Also, you could have a _Macro_ with no
 ### Implementing Macros
 
 [EXCERPT:UniversalImplement]
-There are several ways to add a _Macro_ node to another _Blueprint_ graph. Like _Function_ nodes and _Custom Event_ Call nodes, it is possible to add multiple copies of _Macro_ nodes
-throughout the graphs in your _Blueprint_.
+There are several ways to add a Macro node to another Blueprint graph. Like Function nodes and Custom Event Call nodes, it is possible to add multiple copies of Macro nodes
+throughout the graphs in your Blueprint.
 
-To add a _Macro_, **right-click** in the graph and select the _Macro_ in the context menu that appears.
+To add a Macro, **Right-click** in the graph and select the Macro in the context menu that appears.
 
 ![](right_click_add_macro.png)
 
-You can also drag and drop off the pin of another node, and the _Macro_ will appear in the context menu if it has a parameter pin of the corresponding variable type and flow direction. 
+You can also drag and drop off the pin of another node, and the Macro will appear in the context menu if it has a parameter pin of the corresponding variable type and flow direction. 
 
 [REGION:raw]
 ![](pin_macro_summon.png)
 [/REGION]
 
 [/EXCERPT:UniversalImplement]
-Alternatively, you can drag the _Macro_ from the **My Blueprint** tab and drop it in the graph.
+Alternatively, you can drag the Macro from the **My Blueprint** tab and drop it in the graph.
 
 [REGION:raw]
 ![](macro_drag_1.png)
 [/REGION]
 
-Once the _Macro_ node has been added to the graph, it behaves like any other node and the input and output pins can be wired accordingly. **Double-clicking** on a _Macro_ node will open the _Macro_'s graph.
+Once the Macro node has been added to the graph, it behaves like any other node and the input and output pins can be wired accordingly. **Double-clicking** on a Macro node will open the Macro's graph.
 
 
 
