#UE4doc #docs:
Added engine version field
Changed font settings for multiple terms
Fixed capitalization in multiple areas

diff --git a/Engine/Documentation/Source/Engine/Blueprints/UserGuide/MathNode/MathNode.INT.udn b/Engine/Documentation/Source/Engine/Blueprints/UserGuide/MathNode/MathNode.INT.udn
--- a/Engine/Documentation/Source/Engine/Blueprints/UserGuide/MathNode/MathNode.INT.udn
+++ b/Engine/Documentation/Source/Engine/Blueprints/UserGuide/MathNode/MathNode.INT.udn
@@ -2,6 +2,7 @@ Availability:Public
 Title:Math Expression Node
 Crumbs: %ROOT%, Engine/Blueprints, Engine/Blueprints/UserGuide
 Description:The math expression node allows you to type in a math expression and builds the appropriate sub-graph to create that expression.
+version: 4.9
 
 [VAR:Topic]
 [OBJECT:Topic]
@@ -44,16 +45,15 @@ Description:The math expression node allows you to type in a math expression and
 
 [TOC (start:2 end:2)]
 
-## Overview
 
-To create a Math Expression node, **right-click** in the graph and select **Add Math Expression...** from the context menu. 
+To create a Math Expression node, **Right-click** in the graph and select **Add Math Expression...** from the context menu. 
 
 ![](AddMathExpression.png)
 
-The Math Expression node acts like a collapsed graph. It is a single node that you can **double-click** to open the sub-graph that makes up its functionality. 
+The Math Expression node acts like a collapsed graph. It is a single node that you can **Double-click** to open the sub-graph that makes up its functionality. 
 Initially, the name/expression is blank. Whenever you rename the node, then the new expression is parsed and a new sub-graph is generated. 
 
-To rename the node and update the sub-graph, **right-click** on the node and select **Rename**.
+To rename the node and update the sub-graph, **Right-click** on the node and select **Rename**.
 
 For example, you could rename the Math Expression node with the following expression:
 
@@ -63,7 +63,7 @@ This would update the Math Expression node to have two float inputs, **X** and *
  
 ![](mathnodeexample.png)
 
-If you **double-click** the node with this expression, you would see the below sub-graph:
+If you **Double-click** the node with this expression, you would see the below sub-graph:
 
 [REGION:fullwidth]
 ![](mathnodeexpand.png)
@@ -116,7 +116,7 @@ Certain functions have multiple names that are commonly used. To account for tha
 * Trig arc functions (asin/arcsin, acos/arccos, etc.) 
 
 [REGION:note]
-Because you are typing in functions rather than connecting pins, make sure to enter the correct number and type of parameters. The Math Expression node will display an error if 
+Since you are typing in functions rather than connecting pins, make sure to enter the correct number and type of parameters. The Math Expression node will display an error if 
 there is a mismatch.
 [/REGION]
 
@@ -133,9 +133,9 @@ There are some basic structures that we use quite often in math as well: vectors
 
 Just like there are aliases for certain functions, there are vector and rotator aliases.
 * Vector aliases 
-		(vector, vec, vect, etc.)
+		vector, vec, vect
 * Rotator aliases 
-		(rotator, rot, etc.)
+		rotator, rot
 
 ## Supported Operators
 
