#UE4doc #docs:
Added engine version field
Changed font settings for multiple terms
Fixed capitalization in multiple spots

diff --git a/Engine/Documentation/Source/Engine/Blueprints/UserGuide/Nodes/Nodes.INT.udn b/Engine/Documentation/Source/Engine/Blueprints/UserGuide/Nodes/Nodes.INT.udn
--- a/Engine/Documentation/Source/Engine/Blueprints/UserGuide/Nodes/Nodes.INT.udn
+++ b/Engine/Documentation/Source/Engine/Blueprints/UserGuide/Nodes/Nodes.INT.udn
@@ -2,6 +2,7 @@ Availability:Public
 Title:Nodes
 Crumbs: %ROOT%, Engine/Blueprints, Engine/Blueprints/UserGuide
 Description:Node graph that uses events and function calls to perform actions in response to gameplay events associated with the Blueprint.
+version: 4.9
 
 
 **Nodes** are objects such as events, function calls, flow control operations, variables, etc. 
@@ -69,7 +70,7 @@ There are two main types of pins: execution pins and data pins.
 
 #### Auto-Casting
 
-Connections can be made between certain pins of different data types by way of the auto-casting feature in _Blueprints_. Compatible types can be identified by the tooltip displayed when attempting to create a connection between two pins.
+Connections can be made between certain pins of different data types by way of the auto-casting feature in Blueprints. Compatible types can be identified by the tooltip displayed when attempting to create a connection between two pins.
 
 ![Blueprint - Compatible Types Message](k2_autocast_message.png)
 
@@ -81,13 +82,13 @@ Dragging a wire from a pin of one type to a pin of a different, but compatible,
 
 #### Promote to Variable
 
-The value represented by a data pin can be converted to a variable within the _Blueprint_ using the **Promote to Variable** command. This command creates a new variable in the _Blueprint_ and wires it to the data pin being promoted. In the case of an output data pin, a Set node is used to set the value of the new variable. Essentially, this is just a shortcut to manually adding a new variable, adding it to the graph and wiring it to the data pin.
+The value represented by a data pin can be converted to a variable within the Blueprint using the **Promote to Variable** command. This command creates a new variable in the Blueprint and wires it to the data pin being promoted. In the case of an output data pin, a Set node is used to set the value of the new variable. Essentially, this is just a shortcut to manually adding a new variable, adding it to the graph and wiring it to the data pin.
 
 [INCLUDE:Engine/Blueprints\UserGuide\Variables#PromoteVariable]
 
 ### Wires
 
-Connections between pins are called **wires**. Wires can represent either the flow of execution or the flow of data.
+Connections between pins are called **Wires**. Wires can represent either the flow of execution or the flow of data.
 
 #### Execution Wires
 
@@ -148,7 +149,7 @@ Collections of nodes in the graph can be collapsed into sub-graphs for organizat
 [![](k2_collapsed_node.png)](k2_collapsed_node.png "Blueprint Collapse Node")
 [/REGION]
 
-Unlike [](Engine/Blueprints\UserGuide\Macros), a set of collapsed nodes is not shared, even within a single _Level Blueprint_ or _Blueprint Class_. If you copy the collapsed node, it duplicates the internal graph. This can be handy if you want to make several variants of the same approximate behavior, but it also means that any bug fixes would have to be applied to each copy. The feature is really more intended to 'tidy up' a graph, hiding complexity inside, rather than any sort of sharing or reuse.
+Unlike [](Engine/Blueprints\UserGuide\Macros), a set of collapsed nodes is not shared, even within a single Level Blueprint or Blueprint Class. If you copy the collapsed node, it duplicates the internal graph. This can be handy if you want to make several variants of the same approximate behavior, but it also means that any bug fixes would have to be applied to each copy. The feature is really more intended to 'tidy up' a graph, hiding complexity inside, rather than any sort of sharing or reuse.
 
 #### Tunnels
 
@@ -192,7 +193,7 @@ These pins are automatically generated when the nodes are collapsed. Any executi
 
 	![](collapsed_graph_myblueprint.png)
 
-1.	To edit the collapsed nodes, **double-click** the collapsed graph node or select the sub-graph in the **My Blueprint** tab.
+1.	To edit the collapsed nodes, **Double-click** the collapsed graph node or select the sub-graph in the **My Blueprint** tab.
 	
 	[REGION:lightbox fullwidth]
 	[![](k2_collapse_edit.png)](k2_collapse_edit.png "Blueprint Collapse Nodes - Resulting Graph")
