#UE4doc #docs:
Added engine version field
Changed font settings for multiple terms

diff --git a/Engine/Documentation/Source/Engine/Blueprints/UserGuide/RandomStreams/KismetRandomStreams.INT.udn b/Engine/Documentation/Source/Engine/Blueprints/UserGuide/RandomStreams/KismetRandomStreams.INT.udn
--- a/Engine/Documentation/Source/Engine/Blueprints/UserGuide/RandomStreams/KismetRandomStreams.INT.udn
+++ b/Engine/Documentation/Source/Engine/Blueprints/UserGuide/RandomStreams/KismetRandomStreams.INT.udn
@@ -2,6 +2,7 @@ Availability:Public
 Title:Random Streams
 Crumbs: %ROOT%, Engine, Engine/Blueprints, Engine/Blueprints/UserGuide
 Description:Explanation of random streams and how to use them in Blueprints.
+version: 4.9
 
 [VAR:Topic]
 [OBJECT:Topic]
@@ -46,57 +47,57 @@ Description:Explanation of random streams and how to use them in Blueprints.
 [TOC (start:2)]
 
 [EXCERPT:Intro]
-**_RandomStreams_** allow repeatable, random numbers to be generated and used inside of _Blueprints_ and _Level Blueprints_, as well as _AnimBlueprints_ for animations. This can be extremely useful when 
+**RandomStreams** allow repeatable, random numbers to be generated and used inside of Blueprints and Level Blueprints, as well as AnimBlueprints for animations. This can be extremely useful when 
 setting up effects, such as scattering objects or building procedural environments, where you want a random effect, but also a consistent distribution each time the graph is evaluated. Previously, 
-using random values would result in a different distribution each time the graph was evaluated, meaning moving the _Blueprint_ or performing any action which caused the graph to be re-evaluated 
-resulting in a completely different effect. With _RandomStreams_, you can tweak the effect based on a seed value to get the desired result, and then perform any other modifications you want while 
+using random values would result in a different distribution each time the graph was evaluated, meaning moving the Blueprint or performing any action which caused the graph to be re-evaluated 
+resulting in a completely different effect. With RandomStreams, you can tweak the effect based on a seed value to get the desired result, and then perform any other modifications you want while 
 maintaining the overall effect.
 [/EXCERPT:Intro]
 
 ## RandomStream Variables
 
-_RandomStreams_ are available as a special type of struct variable in _Blueprints_. They can be added like any other variable by creating one through the **My Blueprint** tab in **Graph** mode.
+RandomStreams are available as a special type of struct variable in Blueprints. They can be added like any other variable by creating one through the **My Blueprint** tab in **Graph** mode.
 
 ![RandomStream Variable](variable.png)
 
-See [Creating Variables](Engine/Blueprints/UserGuide\Variables#CreatingVariables) for complete instructions on how to add new variables inside _Blueprint Classes_ or _Level Blueprints_.
+See [Creating Variables](Engine/Blueprints/UserGuide\Variables#CreatingVariables) for complete instructions on how to add new variables inside Blueprint Classes or Level Blueprints.
 
 ## RandomStream Functions
 
-In order to make use of a _RandomStream_ variable, there are a set of functions which take a _RandomStream_ as input and output a random value of different types, depending on the function.
+In order to make use of a RandomStream variable, there are a set of functions which take a RandomStream as input and output a random value of different types, depending on the function.
 
 ![Random Stream Context Menu](context_functions.png)
 
 | Function | Description |
 | -------- | ----------- |
-| `Random Bool from Stream` | Outputs either 0 or 1 at random. |
-| `Random Float from Stream` | Outputs a random `Float` value in the range (0.0, 1.0). |
-| `Random Float in Range from Stream` | Outputs a random `Float` value in the range (Min, Max). |
-| `Random Integer From Stream` | Outputs a uniformly distributed random `Integer` value in the range (0, Max - 1). |
-| `Random Integer In Range From Stream` | Outputs a random `Integer` value in the range (Min, Max). |
-| `Random Rotator From Stream` | Outputs a random `Rotator` value.|
-| `Random Unit Vector From Stream` | Outputs a random unit-length `Vector` value. |
+| **Random Bool from Stream** | Outputs either 0 or 1 at random. |
+| **Random Float from Stream** | Outputs a random Float value in the range (0.0, 1.0). |
+| **Random Float in Range from Stream** | Outputs a random Float value in the range (Min, Max). |
+| **Random Integer From Stream** | Outputs a uniformly distributed random Integer value in the range (0, Max - 1). |
+| **Random Integer In Range From Stream** | Outputs a random Integer value in the range (Min, Max). |
+| **Random Rotator From Stream** | Outputs a random Rotator value.|
+| **Random Unit Vector From Stream** | Outputs a random unit-length Vector value. |
 
-Choosing one of the functions shown above will place it in the graph wired to the _RandomStream_ variable.
+Choosing one of the functions shown above will place it in the graph wired to the RandomStream variable.
 
 ![Random Stream Function](function.png)
 
-An example of using _RandomStream_ functions to randomly place and rotate a component belonging to the _Blueprint_ is shown below:
+An example of using RandomStream functions to randomly place and rotate a component belonging to the Blueprint is shown below:
 
 [REGION:fullwidth]
 ![Random Stream Functions Example](functions_example.png)
 [/REGION]
 
-As you can see, random _Booleans_, _Floats_, _Integers_, _Vectors_, and _Rotators_ can all be pulled from the same stream.
+As you can see, random Booleans, Floats, Integers, Vectors, and Rotators can all be pulled from the same stream.
 
 ## Initial Seed
 
 The **Initial Seed** property is used to calculate the stream of random values. The sequence of random values produced from a single seed value will always be the same each time 
 it is calculated, giving us the consistency mentioned previously. Different seeds generate different sequences of values.
 
-So, modifying the **Initial Seed** of a _RandomStream_ will cause the values produced to vary. This can be used to tweak a random effect until you get the sequence or distribution that 
-suits your needs. A new **Initial Seed** is generated per-instance of the _Blueprint_. This means each time an instance of the _Blueprint_ is placed in the world or a duplicate of an existing 
-instance is made, each _RandomStream_ variable belonging to the _Blueprint_ is given a new **Initial Seed** value. So, each instance will be different and can be tweaked to create the exact desired result.
+So, modifying the Initial Seed of a RandomStream will cause the values produced to vary. This can be used to tweak a random effect until you get the sequence or distribution that 
+suits your needs. A new Initial Seed is generated per-instance of the Blueprint. This means each time an instance of the Blueprint is placed in the world or a duplicate of an existing 
+instance is made, each RandomStream variable belonging to the Blueprint is given a new Initial Seed value. So, each instance will be different and can be tweaked to create the exact desired result.
 
 <!--
 Same Blueprint, Different Playthroughs
@@ -128,14 +129,14 @@ LogBlueprintUserMessages: HTHTHTTHTT
 ### Modifying the Initial Seed
 
 
-To set the **Initial Seed** property directly on the variable, the _RandomStream_ variable must be exposed for editing. Once exposed, the **Initial Seed** property is visible in the **Details** tab when an instance of the _Blueprint_ is selected. 
+To set the Initial Seed property directly on the variable, the RandomStream variable must be exposed for editing. Once exposed, the **Initial Seed** property is visible in the **Details** tab when an instance of the Blueprint is selected. 
 
 ![Exposed Initial Seed Property](seed_exposed.png)
 
-Within a _Blueprint_ graph, you can also set the **Initial Seed** to a specific value or a new random value.
+Within a Blueprint graph, you can also set the Initial Seed to a specific value or a new random value.
 
-The **Set Random Stream Seed** node can be useful when you have one _Blueprint_ that places many objects, like blades of grass or stones, but you want them to always be placed the same way for testing or other purposes.
-It will override the **Initial Seed** being randomly set when the _Blueprint_ is placed in a level. 
+The **Set Random Stream Seed** node can be useful when you have one Blueprint that places many objects, like blades of grass or stones, but you want them to always be placed the same way for testing or other purposes.
+It will override the Initial Seed being randomly set when the Blueprint is placed in a level. 
 
 
 <!--### Resetting a Random Stream-->
