#UE4doc #docs:
Fixed capitalization in the timeline node excerpt

diff --git a/Engine/Documentation/Source/Shared/Glossary/T/Glossary_T.INT.udn b/Engine/Documentation/Source/Shared/Glossary/T/Glossary_T.INT.udn
--- a/Engine/Documentation/Source/Shared/Glossary/T/Glossary_T.INT.udn
+++ b/Engine/Documentation/Source/Shared/Glossary/T/Glossary_T.INT.udn
@@ -74,7 +74,7 @@ Timeline node
 [EXCERPT:Glossary_Timeline_node]
 **Timeline nodes** are special nodes within Blueprints that allow for simple time-based animation to be quickly designed and played back based on in-game events. Timelines are somewhat like simple Matinee sequences 
 in that they allow simple values to be animated and events to be fired off over time. These can be edited 
-directly inside the Blueprint editor by **double-clicking** on the Timeline in the Graph tab or in the My Blueprint tab. They are specifically built for handling simple, non-cinematic tasks such as opening doors, altering 
+directly inside the Blueprint editor by **Double-clicking** on the Timeline in the Graph tab or in the My Blueprint tab. They are specifically built for handling simple, non-cinematic tasks such as opening doors, altering 
 lights, or performing other time-centric manipulations to Actors within a scene.
 [/EXCERPT:Glossary_Timeline_node]
 
