#UE4 Docs: Jira UEDOC-1980 (updates to Using Interfaces page).

diff --git a/Engine/Documentation/Source/Engine/Blueprints/UserGuide/Types/Interface/InterfaceBlueprint.INT.udn b/Engine/Documentation/Source/Engine/Blueprints/UserGuide/Types/Interface/InterfaceBlueprint.INT.udn
--- a/Engine/Documentation/Source/Engine/Blueprints/UserGuide/Types/Interface/InterfaceBlueprint.INT.udn
+++ b/Engine/Documentation/Source/Engine/Blueprints/UserGuide/Types/Interface/InterfaceBlueprint.INT.udn
@@ -1,23 +1,24 @@
 Availability:Public
-Title:Blueprint Interface
-Crumbs: %ROOT%, Engine/Blueprints, Engine/Blueprints/UserGuide/Types
-Description:Blueprints that declare functions to define an interface between Blueprints.
+Title:Implementing Blueprint Interfaces
+Crumbs: %ROOT%, Engine/Blueprints, Engine/Blueprints/UserGuide/Types/Interface
+Description:Blueprint Interfaces allow Blueprints to call functions in different types of targets, if they also implement the interface.
+Version: 4.9
 
 [VAR:Topic]
 [OBJECT:Topic]
 	[PARAM:image]
-		![%Engine/Blueprints\UserGuide\Types\Interface:title%](Engine/Blueprints/UserGuide/Types/blueprint_interface.png)
+		![%Engine\Blueprints\UserGuide\Types\Interface\UsingInterfaces:title%](Engine\Blueprints\UserGuide\using_interface_topic.png)
 	[/PARAM]
 	[PARAM:icon]
 	[/PARAM]
 	[PARAM:title]
-		%Engine/Blueprints\UserGuide\Types\Interface:title%
+		%Engine\Blueprints\UserGuide\Types\Interface\UsingInterfaces:title%
 	[/PARAM]
 	[PARAM:description]
-		%Engine/Blueprints\UserGuide\Types\Interface:description%
+		%Engine\Blueprints\UserGuide\Types\Interface\UsingInterfaces:description%
 	[/PARAM]
 	[PARAM:path]
-		[RELATIVE:Engine/Blueprints\UserGuide\Types\Interface]
+		[RELATIVE:Engine\Blueprints\UserGuide\Types\Interface\UsingInterfaces]
 	[/PARAM]
 [/OBJECT]
 [/VAR]
@@ -25,124 +26,149 @@ Description:Blueprints that declare functions to define an interface between Blu
 [VAR:TopicCompact]
 [OBJECT:TopicCompact]
 	[PARAM:image]
-		![%Engine/Blueprints\UserGuide\Types\Interface:title%](Engine/Blueprints/UserGuide/Types/blueprint_interface.png)
+		![%Engine\Blueprints\UserGuide\Types\Interface\UsingInterfaces:title%](Engine\Blueprints\UserGuide\using_interface_topic.png)
 	[/PARAM]
 	[PARAM:icon]
 	[/PARAM]
 	[PARAM:title]
-		%Engine/Blueprints\UserGuide\Types\Interface:title%
+		%Engine\Blueprints\UserGuide\Types\Interface\UsingInterfaces:title%
 	[/PARAM]
 	[PARAM:description]
-		%Engine/Blueprints\UserGuide\Types\Interface:description%
+		%Engine\Blueprints\UserGuide\Types\Interface\UsingInterfaces:description%
 	[/PARAM]
 	[PARAM:path]
-		[RELATIVE:Engine/Blueprints\UserGuide\Types\Interface]
+		[RELATIVE:Engine\Blueprints\UserGuide\Types\Interface\UsingInterfaces]
 	[/PARAM]
 [/OBJECT]
 [/VAR]
 
 [TOC (start:2)]
 
-[EXCERPT:Intro]
-A **Blueprint Interface** is a collection of one or more functions - name only, no implementation - that can be added to 
-other Blueprints. Any Blueprint that has the Interface added is guaranteed to have those functions. The functions 
-of the Interface can be given functionality in each of the Blueprints that added it. This is essentially like the 
-concept of an interface in general programming, which allows multiple different types of Objects to all share and be accessed 
-through a common interface. Put simply, Blueprint Interfaces allow different Blueprints to share with and send data to one another.
+## Using Interfaces in Blueprints
 
-Blueprint Interfaces can be made by content creators through the editor in a similar fashion to other Blueprints, but they 
-come with certain limitations in that they cannot:
+[EXCERPT:intro]
+The use of **Blueprint Interfaces** allows for a common method of interacting with multiple different types of objects that all share 
+some specific functionality. This means you can have completely different types of objects, such as a car and a tree, that 
+share one specific thing, like they can both be shot by weapon fire and take damage. By creating a Blueprint Interface that 
+contains an `OnTakeWeaponFire` function, and having both the car and the tree implement that Blueprint Interface, you can treat the 
+car and the tree as the same type and only call the `OnTakeWeaponFire` function when either of them is shot.
+[/EXCERPT:intro]
 
-* Add new variables
-* Edit graphs
-* Add Components
 
+### Adding Interfaces
 
-[/EXCERPT:Intro]
+Blueprint Interfaces as well as Interfaces defined in native code (C++) can be added to any type of [Blueprint](Engine/Blueprints).
 
-The use of _Blueprint Interfaces_ allows for a common method of interacting with multiple disparate types of Objects that all share 
-some specific functionality. This means you can have completely different types of Objects, such as a car and a tree, that 
-share one specific thing like they can both be shot by weapon fire and take damage. By creating a _Blueprint Interface_ that 
-contains an `OnTakeWeaponFire` function, and having both the car and the tree implement that _Blueprint Interface_, you can treat the 
-car and the tree as the same type and simply call the `OnTakeWeaponFire` function when either of them is shot. Read about how to implement
-_Blueprint Interfaces_ on the [](Engine/Blueprints\UserGuide\Types\Interface\UsingInterfaces) page.
+1. From the Main Level Editor window, open the **Level Blueprint** inside your project. 
 
-## Creating Blueprint Interfaces
+	![](OpenLevelBP.png)(w:640)
 
-_Blueprint Interface_ creation within the **Content Browser** is as simple as **right-clicking** and choosing _Blueprint Interface_ from the context menu, then naming your new Interface.
+1. Click the **Class Settings** button from the [Toolbar](Engine/Blueprints/Editor/UIComponents/Toolbar) across the top.
 
-![](CreateInterface.gif)
+	![](ClassSettings.png)
 
-![](CreatedBlueprintInterface.gif)
+1. In the **Details** panel on the right, locate the **Interfaces** category. 
 
-Once created, you need to edit the Interface's functions.
+1. Click the **Add Interface** button to see a list of the available Interfaces and select the Inteface to use. 
 
-## Editing Blueprint Interfaces
+	![](InterfacesCategory.png)
 
-_Blueprint Interfaces_ are edited the **Blueprint Editor**. Since you cannot create your own variables, graphs, or components, the process of Interface editing differs greatly from the process of editing standard _Blueprint Classes_. 
+	![](InterfacesListPopulated.png)
 
-When you first open up a new Interface, the editor looks like this:
+[REGION:tip]
+	Interfaces can be removed by clicking the ![](button_X_Remove.png) button located next to each one on the _Implemented Interfaces_ list. You can also migrate from Blueprint Interfaces to normal Blueprint functions. For more information, see [](Engine/Blueprints\UserGuide\BlueprintComms\InterfaceMigration).
+[/REGION]
 
-![](InterfaceEditor.png)
+[REGION:note]
+	The _Inherited Interfaces_ section within the Interfaces category automatically populates with any Interfaces that were inherited by this object (i.e. implemented in the Blueprint from which the current Blueprint was derived).
+[/REGION]
 
-If you have just created your Interface, you will see that the Editor has created a new blank Function for you named `NewFunction_0` and it will be highlighted for you to rename it.
 
-### Adding Functions
 
-Functions are the primary component of an Interface. Interface functions have no implementation. They exist simply as a definition of inputs and outputs. These can be used to send data through the interface, or can be overridden within any Blueprint that implements that Interface.
+### Implementing Interface Members
 
-To add a new function:
+When an Interface is added to a Blueprint, that Blueprint must **implement** the Interface's functions in order to have 
+them perform any action when called. This is done by placing **Event** nodes for the Interface functions in the **EventGraph**. 
+The nodes connected to the Event define the functionality, or implementation, for the function.
 
-1. In the **My Blueprint** tab create a new function, by clicking on the **Add Button** ![](Engine\Blueprints\Editor\UIComponents\MyBlueprint\plus_Button.png) on the functions list header ![](Engine\Blueprints\Editor\UIComponents\MyBlueprint\myblueprint_function.png).
+1. **Right-click** in the **EventGraph** of the Blueprint and choose **Add Event > Event ...** 
 
-1. In the **My Blueprint** pane, enter a name for the new function.
+	![Implement Interface - Add Event](k2_interface_event_menu.png) ![Implement Interface - Event Node](k2_interface_event_graph.png)
 
-	![](RenameMyFunction.png)
+	Select the type of Event you would like to add and the Event will be added to the graph. 
 
-1. The new Graph area will appear with the new function. Note that the function has neither inputs nor outputs.
+1. Create a network of nodes connected to the Event's executable pin to give it functionality.
 
-	![](NewFunctionCreated.png)
+	![Implement Interface - Event Implementation](k2_interface_event_implementation.png)
 
+	
+### Calling Interface Functions
 
-### Editing Function Signatures
+The functions of the Interface are executed through the use of **Interface Message** nodes, These look and act like regular 
+[](Engine/Blueprints/UserGuide/Functions/FunctionCalls), except that they can fail silently if the target does not implement 
+the Interface.
 
-Since an interface function has no implementation, all you can do as a designer is designate a series of typed inputs and outputs. 
+![Implement Interface - Call Function](k2_interface_event_call.png)
 
-To edit a function's signature:
 
-1. In the **Details** tab, scroll to the _Inputs_ category and click the ![](button_New.png) button to create a new Input Parameter.
+### Calling Interface Functions on Other Blueprints
 
-	![](Details_Signature.png)
+Interfaces are useful for allowing communication between different Blueprints. Functions from a single interface can be called on other Blueprints. This can be done when 2 criteria are met:
 
-1. Set the Input Name and Type as desired. You may also expand the input using the ![](button_DropdownArrow.png) button next to the name, and thereby set a Default Value.
+* Both Blueprints implement the same interface with the required function.
+* The calling Blueprint knows the name of the in-game instance of the other Blueprint.
 
-	![](FloatInput.png) ![](MyFunctionGraph.png)
 
-1. In the same manner, outputs can also be added. Note how the graph automatically updates to show them.
+When creating an Interface Function node within the Event Graph, you will often see two versions of the function listed in the context menu. 
 
-	![](OutputBool.png) ![](InputWithOutput.png)
+![](MyFunctionInMenu.png)
 
+It is important to note that the version listed beneath **Call Function** is a local version, intended only to be called within this Blueprint. 
 
+The version listed under **Interface Messages** is designed to call the function via the interface on other Blueprints implementing the same interface.
 
+There is some labeling beneath the node title in place to help differentiate:
 
-### Other Considerations
+[REGION:imagetable]
+| ![](Local_MyFunction.png)| ![](Interface_MyFunction.png) |
+| ------ | ------ |
+| Local version. | Interface Message version. |
+[/REGION]
 
-* On any input or output parameter, you can click the ![](button_X_Remove.png) button to remove that parameter.
+Calling an Interface Message Function requires that you supply the instance of the appropriate Blueprint on which you are calling the function.
 
-* Input parameters can be given a default value using the small text field immediately under of the parameter name field, though you must expand the property entry in the **Details** tab to see it.
+Note in this image we are sending another Blueprint as the Target input. This would call the function on the supplied Blueprint. If that Blueprint did not implement the interface, then nothing would happen.
 
-* You can change where the input and output parameters appear on the node using the ![](button_moveUp.png) and ![](button_moveDown.png) found in the **Details** tab when the parameter properties are expanded.
+![](InterfaceMessageFunctionCalled.png)
 
-* You can use the Replicates checkbox for any Interfaces containing functions that need to replicate across the server. This is found within the **Details** tab by first clicking the ![](BlueprintProps.png) button.
+[REGION:tip]
+	To specify a Blueprint that exists within the level, it is best to make a Public variable of type 'Actor', and assign the appropriate Blueprint in the value for that variable using the **Details** tab of the Level Editor.
+[/REGION]
 
-## Reference Links
+### Calling Interface Functions on Level Blueprints
 
-Refer to the following links for additional documentation regarding **Blueprint Communication**: 
+Interfaces can also allow for communication between sub-levels using the **Get Streaming Level** node. If the level is loaded and valid, the **Interface Message** will call the function with the level's Blueprint Script as the target.
 
-* For information on **Blueprint Communication Guidelines**, see: [**Blueprint Communication Usage**](Engine\Blueprints\UserGuide\BlueprintCommsUsage) 
+![](InterfaceGetLevelStreaming.png)
 
-* For information on **How to Setup Different Blueprint Communications**, see: [**How to Use Blueprint Communications**](Engine\Blueprints\UserGuide\BlueprintCommsUsage\BPComHowTo) 
+### Overriding an Interface Function
 
-* For information on **Event Dispatchers**, see: [**Event Dispatchers**](Engine\Blueprints\UserGuide\EventDispatcher) 
+One of the most important aspects of Interface functions is that, since they have no implementation of their own, their functionality can be independently defined for each Blueprint that implements the Interface.
 
-* For information on **Direct Blueprint Communication**, see: [**Direct Blueprint Communications**](Engine\Blueprints\UserGuide\BlueprintComms)
+Once an Interface has been implemented by a Blueprint, the functions of that interface will be visible in the **My Blueprint** panel.
+
+![](iFunctionInExplorer.png)
+
+If you **Double-click** that function, it will open up within the Blueprint Editor. This shows the local implementation for this function. Whatever functionality is defined here will be unique to this Blueprint. In this way, you can define custom functionality for each Blueprint that implements the interface.
+
+For example, in this image we see a basic Interface function as opened within a Blueprint. Right now, the function does nothing.
+
+![](DefaultFunction.png)(w:600)
+
+We could create a simple network to do a number comparison, and output the value, like so.
+
+![](ComparisonFunction.png)(w:600)
+
+Now, when the function is called locally, it will take in a float, compare it against the number 10, and then output the boolean result to the screen.
+
+![](FunctionLogging.png)
