#UE4doc #docs:
Added engine version field
Removed Italics from multiple terms

diff --git a/Engine/Documentation/Source/Engine/Blueprints/UserGuide/Timelines/Timelines.INT.udn b/Engine/Documentation/Source/Engine/Blueprints/UserGuide/Timelines/Timelines.INT.udn
--- a/Engine/Documentation/Source/Engine/Blueprints/UserGuide/Timelines/Timelines.INT.udn
+++ b/Engine/Documentation/Source/Engine/Blueprints/UserGuide/Timelines/Timelines.INT.udn
@@ -2,6 +2,7 @@ Availability:Public
 Title: Timelines
 Crumbs: %ROOT%, Engine, Engine/Blueprints, Engine/Blueprints/UserGuide
 Description: This document contains an overview of Timelines in Unreal Engine 4 Blueprints.
+version: 4.9
 
 ![k2_timeline_demo.png](k2_timeline_demo.png)
 
@@ -76,25 +77,25 @@ Description: This document contains an overview of Timelines in Unreal Engine 4
 
 ![k2_timeline_node.png](k2_timeline_node.png)
 
-_Timelines_ come with the following input and output pins:
+Timelines come with the following input and output pins:
 
 |  Item | Description |
 | -------- | ----------- |
 | [REGION:tablesection]Input Pins[/REGION] ||
-| **Play** | Causes the _Timeline_ to play forward from its current time. |
-| **Play from Start** | Causes the _Timeline_ to play forward from the beginning. |
-| **Stop** | Freezes the playback of the _Timeline_ at the current time. |
-| **Reverse** | Plays the _Timeline_ backwards from its current time. |
-| **Reverse from End** | Plays the _Timeline_ backwards from the beginning. |
+| **Play** | Causes the Timeline to play forward from its current time. |
+| **Play from Start** | Causes the Timeline to play forward from the beginning. |
+| **Stop** | Freezes the playback of the Timeline at the current time. |
+| **Reverse** | Plays the Timeline backwards from its current time. |
+| **Reverse from End** | Plays the Timeline backwards from the beginning. |
 | **Set New Time** | Sets the current time to the value set (or input) in the New Time input. |
-| **New Time** | This data pin takes a float value representing time in seconds, to which the _Timeline_ can jump when the Set New Time input is called. |
+| **New Time** | This data pin takes a float value representing time in seconds, to which the Timeline can jump when the Set New Time input is called. |
 | [REGION:tablesection]Output Pins[/REGION] ||
-| **Update** | Outputs an execution signal as soon as the _Timeline_ is called.  |
+| **Update** | Outputs an execution signal as soon as the Timeline is called.  |
 | **Finished** | Outputs an execution signal when playback ends. This is not triggered by the Stop function.  |
-| **Direction** | Outputs enum data indicating the direction the _Timeline_ is playing. |
+| **Direction** | Outputs enum data indicating the direction the Timeline is playing. |
 
 [REGION:note]
-	_Timelines_ may have any number of additional output data pins reflecting the types of tracks created within them. These include Float, Vector, and Event tracks.
+	Timelines may have any number of additional output data pins reflecting the types of tracks created within them. These include Float, Vector, and Event tracks.
 [/REGION]
 
 
