#UE4doc #docs:
Resubmitting as not all changes were saved

diff --git a/Engine/Documentation/Source/Engine/Blueprints/UserGuide/Timelines/Creation/CreatingTimelines.INT.udn b/Engine/Documentation/Source/Engine/Blueprints/UserGuide/Timelines/Creation/CreatingTimelines.INT.udn
--- a/Engine/Documentation/Source/Engine/Blueprints/UserGuide/Timelines/Creation/CreatingTimelines.INT.udn
+++ b/Engine/Documentation/Source/Engine/Blueprints/UserGuide/Timelines/Creation/CreatingTimelines.INT.udn
@@ -1,40 +1,41 @@
 Availability:Public
 Title: Creating Timelines
 Crumbs: %ROOT%, Engine, Engine/Blueprints, Engine/Blueprints/UserGuide, Engine/Blueprints/UserGuide/Timelines
-Description: This document contains an overview of how to create Timeline nodes in Blueprints.
+Description: This document contains an overview of how to create **Timeline** nodes in Blueprints.
+version: 4.9
 
 %Description%
 
 ## Creating Timelines
 
-1. Creating _Timelines_ in _Blueprints_ is extremely simple. In the **Graph** tab, **right-click** and choose **Add Timeline** from the context menu.
+1. Creating Timelines in Blueprints is extremely simple. In the **Graph** tab, **Right-click** and choose **Add Timeline** from the context menu.
 
 	![](Add_Timeline.png)
 
-1. A new _Timeline_ node is added to the graph.
+1. A new Timeline node is added to the graph.
 
 	![](Timeline_Created.png)
 
 
-Once a _Timeline_ has been added, you will be able to see it listed in the **My Blueprint** tab. 
+Once a Timeline has been added, you will be able to see it listed in the **My Blueprint** tab. 
 
 ![](Timeline_In_Explorer.png)
 
 ## Timeline Variables
 
-Once you have created a Timeline, you will notice it become available as a variable in the **My Blueprint** tab. This is done to allow for reference to the _Timeline_ component, and is particularly useful when having to utilize the _Timeline_ nodes. These nodes exist to allow for access to certain features of a Timeline, such as its play rate.
+Once you have created a Timeline, you will notice it become available as a variable in the **My Blueprint** tab. This is done to allow for reference to the Timeline component, and is particularly useful when having to utilize the Timeline nodes. These nodes exist to allow for access to certain features of a Timeline, such as its play rate.
 
 ![](Timeline_Variable.png)
 
-In this example, we are using the Get _Timeline_ 0 variable to query the current Play Rate value of that timeline. 
+In this example, we are using the Get Timeline 0 variable to query the current Play Rate value of that timeline. 
 
 ![](Timeline_Get_PlayRate.png)
 
-A wide variety of such nodes exist, allowing you to query values and control behavior of your _Timeline_ nodes. For a complete list of the available _Timeline_ nodes and their functionality, see the [Timeline nodes page](Engine/Blueprints/UserGuide\Timelines\TimelineNodes).
+A wide variety of such nodes exist, allowing you to query values and control behavior of your Timeline nodes. For a complete list of the available Timeline nodes and their functionality, see the [Timeline nodes page](Engine/Blueprints/UserGuide\Timelines\TimelineNodes).
 
 ## Renaming Timelines
 
-When using multiple _Timelines_ in your _Blueprints_, it is important to name each _Timeline_ accordingly. By default, all _Timelines_ are named "Timeline_X" with a serialized number at the end. You can rename a _Timeline_ by **right-clicking** it within the **Graph** tab or in the **My Blueprint** tab and choosing **Rename**.
+When using multiple Timelines in your Blueprints, it is important to name each Timeline accordingly. By default, all Timelines are named "Timeline_X" with a serialized number at the end. You can rename a Timeline by **Right-clicking** it within the **Graph** tab or in the **My Blueprint** tab and choosing **Rename**.
 
 ![](Rename_Timeline.png)
 
@@ -44,6 +45,6 @@ When using multiple _Timelines_ in your _Blueprints_, it is important to name ea
 
 ## Opening the Timeline Editor 
 
-Of course, creating a _Timeline_ is not useful until you have set it up to do something. To open the **Timeline Editor**, simply **double-click** the _Timeline_ node, either in the **Graph** tab or via its listing in the **My Blueprint** tab.
+Of course, creating a Timeline is not useful until you have set it up to do something. To open the **Timeline Editor**, **Double-click** the Timeline node, either in the **Graph** tab or via its listing in the **My Blueprint** tab.
 
-For more information on editing _Timelines_, please see the [Timeline Editing page](Engine/Blueprints/UserGuide\Timelines\Editor).
+For more information on editing Timelines, please see the [Timeline Editing page](Engine/Blueprints/UserGuide\Timelines\Editor).
