#UE4doc #docs:
Added engine version field
Changed font settings for multiple terms
Fixed capitalization for a term
Replaced an image as the UI was outdated

diff --git a/Engine/Documentation/Source/Engine/Blueprints/UserGuide/Timelines/Editor/EditingTimelines.INT.udn b/Engine/Documentation/Source/Engine/Blueprints/UserGuide/Timelines/Editor/EditingTimelines.INT.udn
--- a/Engine/Documentation/Source/Engine/Blueprints/UserGuide/Timelines/Editor/EditingTimelines.INT.udn
+++ b/Engine/Documentation/Source/Engine/Blueprints/UserGuide/Timelines/Editor/EditingTimelines.INT.udn
@@ -2,8 +2,9 @@ Availability:Public
 Title: Editing Timelines
 Crumbs: %ROOT%, Engine, Engine/Blueprints, Engine/Blueprints/UserGuide, Engine/Blueprints/UserGuide/Timelines
 Description: This document contains an overview of how to edit Timeline nodes in Blueprints.
+version: 4.9
 
-_Timelines_ can be edited by **double-clicking** on the _Timeline_ node in the **Graph** tab or on the _Timeline_ in the **My Blueprint** tab. This opens the **Timeline Editor** in a new tab.
+**Timelines** can be edited by **Double-clicking** on the Timeline node in the **Graph** tab or on the Timeline in the **My Blueprint** tab. This opens the **Timeline Editor** in a new tab.
 
 ## Timeline Editor
 
@@ -11,20 +12,20 @@ _Timelines_ can be edited by **double-clicking** on the _Timeline_ node in the *
 
 | Button/Checkbox | Description |
 | ----  | --|
-|![](add_float_track_button.png) | Adds a new float track to the _Timeline_, for animating scalar float values.|
-|![](add_vector_track_button.png)  | Adds a new vector track to the _Timeline_, for animating float vector values such as rotation or translation.|
+|![](add_float_track_button.png) | Adds a new float track to the Timeline, for animating scalar float values.|
+|![](add_vector_track_button.png)  | Adds a new vector track to the Timeline, for animating float vector values such as rotation or translation.|
 |![](add_event_track_button.png) | This adds an event track, which provides another execution output pin that will be triggered at keyframed times on the track.|
-|![](add_color_track_button.png)  | Adds a new linear color track to the _Timeline_, for animating colors. |
-| ![](add_external_curve_button.png)| Add an external curve to the _Timeline_. This button is only active if an external curve is selected in the **Content Browser**.  | 
-| ![](timeline_length_button.png) | Allows you to set the playback length for this _Timeline_. |
+|![](add_color_track_button.png)  | Adds a new linear color track to the Timeline, for animating colors. |
+| ![](add_external_curve_button.png)| Add an external curve to the Timeline. This button is only active if an external curve is selected in the **Content Browser**.  | 
+| ![](timeline_length_button.png) | Allows you to set the playback length for this Timeline. |
 | ![](last_keyframe_button.png) | If this is not active, the last keyframe of a sequence is ignored. This can help prevent skipping when an animation is looping.|
-|![](autoplay_button.png) | If active, this _Timeline_ node does not require an execution input to begin, and will start playing as soon as the level begins. |
-| ![](loop_button.png)| If active, the _Timeline_ animation will repeat indefinitely unless stopped via the Stop input pin.|
-| ![](replicated_button.png) | If active, the _Timeline_ animation will be replicated across clients.|
+|![](autoplay_button.png) | If active, this Timeline node does not require an execution input to begin, and will start playing as soon as the level begins. |
+| ![](loop_button.png)| If active, the Timeline animation will repeat indefinitely unless stopped via the Stop input pin.|
+| ![](replicated_button.png) | If active, the Timeline animation will be replicated across clients.|
 
 ## Adding Tracks
 
-_Timelines_ use **tracks** to define the animation of a single piece of data. This can be a float value, vector value, color value, or events. Tracks can be added to the _Timeline_ by clicking one of the **Add Track** buttons. For instance, click the ![Blueprint Timeline - Add Float Track Button](add_float_track_button.png)  button and enter a name for the new track.
+Timelines use **Tracks** to define the animation of a single piece of data. This can be a float value, vector value, color value, or events. Tracks can be added to the Timeline by clicking one of the **Add Track** buttons. For instance, click the ![Blueprint Timeline - Add Float Track Button](add_float_track_button.png)  button and enter a name for the new track.
 Press **Enter** to save the name for your new float track.
 
 ![Blueprint Timeline - Add Float Track](k2_timeline_track_float.png)
@@ -38,7 +39,7 @@ Press **Enter** to save the name for your new float track.
 
 Once you have tracks in place, you can then start adding keys to define your animation. 
 
-For more information on working with keys and curves in _Timelines_, please see the [Keys and Curves page](Engine/Blueprints/UserGuide\Timelines\KeysAndCurves).
+For more information on working with keys and curves in Timelines, please see the [Keys and Curves page](Engine/Blueprints/UserGuide\Timelines\KeysAndCurves).
 
 [![](Engine/Blueprints/UserGuide\Timelines\KeysAndCurves\k2_timeline_track_key_add.png)](Engine/Blueprints/UserGuide\Timelines\KeysAndCurves)
 
