#UE4doc #docs:
-WIP-
Changed font settings for multiple terms
Added a light box region for a picture that needed to be clicked on.

diff --git a/Engine/Documentation/Source/Engine/Blueprints/UserGuide/Timelines/Examples/LightFlickering/LightFlickering.INT.udn b/Engine/Documentation/Source/Engine/Blueprints/UserGuide/Timelines/Examples/LightFlickering/LightFlickering.INT.udn
--- a/Engine/Documentation/Source/Engine/Blueprints/UserGuide/Timelines/Examples/LightFlickering/LightFlickering.INT.udn
+++ b/Engine/Documentation/Source/Engine/Blueprints/UserGuide/Timelines/Examples/LightFlickering/LightFlickering.INT.udn
@@ -20,15 +20,15 @@ The light shown will flicker its brightness and change color.
 
 The timeline serves 3 primary functions, with a separate track for each:
 
-1. A **float track** (named *Brightness*) handles animating the brightness. This curve was added from the **Content Browser** using the *Flicker Curve* asset.
-2. A **vector track** (named *Light Color*) animates the light's color. This was handled such that the light turns orange at the beginning of the sequence and then returns to white at the end. The curve interpolation was set to clamp, allowing the value to remain constant until the next key.
-3. An **event track** (named *Trigger Toggle*) disables our trigger throughout the duration of the timeline's sequence. Because this is an event track, the values of the keyframes are irrelevant. Wherever the timeline encounters a keyframe, an execution pulse is sent out of the track's associated output pin.
+1. A **float track** (named Brightness) handles animating the brightness. This curve was added from the **Content Browser** using the Flicker Curve asset.
+2. A **vector track** (named Light Color) animates the light's color. This was handled such that the light turns orange at the beginning of the sequence and then returns to white at the end. The curve interpolation was set to clamp, allowing the value to remain constant until the next key.
+3. An **event track** (named Trigger Toggle) disables our trigger throughout the duration of the timeline's sequence. Because this is an event track, the values of the keyframes are irrelevant. Wherever the timeline encounters a keyframe, an execution pulse is sent out of the track's associated output pin.
 
 
 
-The logic of the network is fairly straightforward. Because we always want the flicker to play from the beginning, our event feeds into *Play from Start*. However, since we do not want the timeline to be re-triggered during playback, we check the value of a boolean variable whose only purpose is to verify if the timeline is playing.
+The logic of the network is fairly straightforward. Because we always want the flicker to play from the beginning, our event feeds into Play from Start. However, since we do not want the timeline to be re-triggered during playback, we check the value of a boolean variable whose only purpose is to verify if the timeline is playing.
 
-That variable has a default value of *false*, meaning we will always play the first time we touch the trigger. So execution from the timeline runs as follows:
+That variable has a default value of false, meaning we will always play the first time we touch the trigger. So execution from the timeline runs as follows:
 
 * The result of the Brightness track is multiplied by *20* to give the light more intensity. 
 * The color of the light has also changed as the next action.
@@ -38,6 +38,8 @@ That variable has a default value of *false*, meaning we will always play the fi
 
 **Click for full size**
 
+[REGION:lightbox]
 [![FlickeringLightNetwork.png](FlickeringLightNetwork.png)(w:600)](FlickeringLightNetwork.png)
+[/REGION]
 
 
