UEDOC-2356

Updating all images on this page as they are out of date.

diff --git a/Engine/Documentation/Source/Engine/Blueprints/UserGuide/Nodes/Nodes.INT.udn b/Engine/Documentation/Source/Engine/Blueprints/UserGuide/Nodes/Nodes.INT.udn
--- a/Engine/Documentation/Source/Engine/Blueprints/UserGuide/Nodes/Nodes.INT.udn
+++ b/Engine/Documentation/Source/Engine/Blueprints/UserGuide/Nodes/Nodes.INT.udn
@@ -34,15 +34,15 @@ New nodes are added to a graph by selecting the type of node from the **Context
 
 Nodes are selected by clicking on the node in the **Graph Editor** tab.
 
-![Blueprint Node Selected](k2_select_node.png)
+![](SelectNode.jpg)(w:820)
 
 A node can be added to, or removed from, the current selection by holding **Ctrl** and clicking on the node.
 
-![Blueprint Multiple Nodes Selected](k2_select_multiple.png)
+![](VarMessage.jpg)(w:820)
 
 Multiple nodes can be selected at once by clicking and dragging to create a marquee selection. Holding **Ctrl** + clicking and dragging creates a marquee selection that toggles selection. Holding **Shift** + clicking and dragging creates a marquee selection that adds to the current selection.
 
-![Blueprint Marquee Selection](k2_select_marquee.png)(w:820)
+![](DotBoxSpawnEmitter.jpg)(w:820)
 
 To deselect all nodes, simply click in an empty space of the **Graph Editor** tab.
 
@@ -94,11 +94,11 @@ Connections between pins are called **Wires**. Wires can represent either the fl
 
 Wires between exec pins represent the flow of execution. Exec wires are displayed as white arrows spanning from an output exec pin to an input exec pin. The direction of the arrow indicates the flow of execution.
 
-![Blueprint Exec Wire](k2_select_node.png)
+![](SelectNode.jpg)(w:820)
 
 Exec wires produce a visual indicator when being executed. During play, as one node completes execution and the next is activated, the wire between their execution pins highlights the fact that execution is being transferred from one node to another.
 
-![Blueprint Exec Wire - Attack State](k2_flow_exec.png)
+![](k2_flow_exec.jpg)(w:820)
 
 The visual indicator for an exec wire fades away over time.
 
@@ -118,17 +118,17 @@ Wires are created in the **Graph Editor** tab using one of the methods below:
 
 * Click on one pin and drag and release on another pin of the same type to create a direct connection.
 
-	![Blueprint Wire Creation - Compatible Pin](k2_connection_direct.png)
+	![](HoverCheck.jpg)(w:820)
 
 	The connection can only be made between two compatible types of pins. If you drag a connection on to a pin that is not compatible, an error will be displayed informing you the connection cannot be made.
 
-	![Blueprint Wire Creation - Incompatible Pin](k2_connection_bad.png)
+	![](HoverCheckNo.jpg)(w:820)
 
 * Drag a connection from a pin and release in an empty space to summon a context-sensitive menu that lists all the nodes that are compatible with the type of pin the connection originated from. Selecting a node from this list creates a new instance of that node and makes a connection to a compatible pin on the new node.
 
 	![Blueprint Wire Creation - New Node](k2_connection_new.png)
 
-	![Blueprint Wire Creation - New Node Created](k2_connection_new2.png)
+	![](ConnectedNodes.jpg)(w:820)
 
 	A wire between two pins can be broken using one of the methods below:
 
@@ -136,7 +136,7 @@ Wires are created in the **Graph Editor** tab using one of the methods below:
 
 * **Right-click** on one of the connected pins and choose Break Link(s).
 
-	![Blueprint Break Wire](k2_connection_break.png)
+	![](BreakLinkNode.jpg)(w:820)
 
 
 ### Collapsed Graphs
