TypoCorrection

diff --git a/Engine/Documentation/Source/Platforms/Android/GettingStarted/5/AndroidQS_5.INT.udn b/Engine/Documentation/Source/Platforms/Android/GettingStarted/5/AndroidQS_5.INT.udn
--- a/Engine/Documentation/Source/Platforms/Android/GettingStarted/5/AndroidQS_5.INT.udn
+++ b/Engine/Documentation/Source/Platforms/Android/GettingStarted/5/AndroidQS_5.INT.udn
@@ -1,4 +1,4 @@
-Availability:Public
+﻿Availability:Public
 Title:5. Package Your Game
 Crumbs:%ROOT%, Platforms, Platforms/Android, Platforms/Android/GettingStarted
 Description:Package your game for installation.
@@ -72,7 +72,7 @@ We are currently packaging in **Development** because we are testing this game,
 	[/REGION]
 
 	[REGION:note]
-	Different devices support different Texture formats depending on the hardware they run on. See, [Android Texture Formats](https://docs.unrealengine.com/latest/INT/Platforms/Android/Reference/index.html) for a detailed breakdown of the various formats.
+	Different devices support different Texture formats depending on the hardware they run on. See, [](Platforms/Android/Reference) for a detailed breakdown of the various formats.
 	[/REGION]
 
 1.	You will then be prompted for a location to save your packaged game. For now create a new folder on the desk top and select that as the location to save your project to.
@@ -138,7 +138,7 @@ We are currently packaging in **Development** because we are testing this game,
 
 	![](MAC_Install_BAT.png)
 	[region:note]
-	You must have your Android device plugged into your development PC via USB to deploy the project to your device. If you do not have your device plugged in running the .BAT file will do nothing.
+	You must have your Android device plugged into your development PC via USB to deploy the project to your device. If you do not have your device plugged in running the .command file will do nothing.
 	[/region]
 
 1. After the .command file has finished you should now see your project ready to be used on your Android device.
