#UE4doc #docs - Added Skill Level field. Added Engine Version field.

diff --git a/Engine/Documentation/Source/Engine/Rendering/Materials/HowTo/AnimatingUVCoords/AnimatingUVCoordinates.INT.udn b/Engine/Documentation/Source/Engine/Rendering/Materials/HowTo/AnimatingUVCoords/AnimatingUVCoordinates.INT.udn
--- a/Engine/Documentation/Source/Engine/Rendering/Materials/HowTo/AnimatingUVCoords/AnimatingUVCoordinates.INT.udn
+++ b/Engine/Documentation/Source/Engine/Rendering/Materials/HowTo/AnimatingUVCoords/AnimatingUVCoordinates.INT.udn
@@ -5,6 +5,8 @@ Description:A How To Guide for animating UV Coordinates
 Related: Engine/Rendering/Materials/ExpressionReference/Coordinates
 Related: Engine/Rendering/Materials/CustomizedUVs
 Related: Engine/Rendering/LightingAndShadows/LightFunctions
+Version: 4.9
+SkillLevel: Intermediate
 
 
 [TOC(start:2 end:3)]
