#UE4doc #docs - Added Skill Level field. Added Engine Version field.

diff --git a/Engine/Documentation/Source/Engine/Rendering/Materials/HowTo/BumpOffset/BumpOffset.INT.udn b/Engine/Documentation/Source/Engine/Rendering/Materials/HowTo/BumpOffset/BumpOffset.INT.udn
--- a/Engine/Documentation/Source/Engine/Rendering/Materials/HowTo/BumpOffset/BumpOffset.INT.udn
+++ b/Engine/Documentation/Source/Engine/Rendering/Materials/HowTo/BumpOffset/BumpOffset.INT.udn
@@ -5,6 +5,8 @@ Description: A How To guide for using the Bump Offset node in your Material.
 Related: Resources/Showcases/RealisticRendering
 Related: Engine/Rendering/Materials/ExpressionReference/Utility
 Related: Engine/Rendering/Materials
+Version: 4.9
+SkillLevel: Intermediate
 
 [TOC(start:2)]
 
@@ -98,7 +100,7 @@ Once named, **Double-Click** on the Material in the **Content Browser** to open
 	When you first connect from the Bump Offset Material Expression to the UV input on any of the three textures, you will notice a red warning placed at the bottom of the Bump Offset Material Expression. The Bump Offset Material Expression needs a Height Map in order to work correctly and currently one has not been provided. Once an input to the Height has been entered, this error will go away.   
 	[/Region]
 
-1. Once the Bump Offset Material Expression has been connected to the UV input of all the Textures, it now needs to be setup so that it can work. The first thing that needs to be done is that we need to give it a Height input. To do this, grab the Texture we are using for the Base Color input, T_Brick_Clay_New_D, and duplicate it by first selecting it with the mouse and then pressing **CRTL + W** on the keyboard. Once the node is duplicated, move it so that it is in front of the Bump Offset Material Expression and the connect the Alpha output of the Texture into the Height input of the Bump Offset Material Expression. When completed, you should have something that looks like this. 
+1. Once the Bump Offset Material Expression has been connected to the UV input of all the Textures, it now needs to be setup so that it can work. The first thing that needs to be done is that we need to give it a Height input. To do this, grab the Texture we are using for the Base Color input, **T_Brick_Clay_New_D**, and duplicate it by first selecting it with the mouse and then pressing **CRTL + W** on the keyboard. Once the node is duplicated, move it so that it is in front of the Bump Offset Material Expression and the connect the Alpha output of the Texture into the Height input of the Bump Offset Material Expression. When completed, you should have something that looks like this. 
 	[OBJECT:EmbeddedVideo]
 	[PARAMLITERAL:width]
 	640
