#UE4doc #docs - Added Skill Level field. Added Engine Version field.

diff --git a/Engine/Documentation/Source/Engine/Rendering/Materials/HowTo/CreatingLayeredMaterials/CreatingLayeredMaterials.INT.udn b/Engine/Documentation/Source/Engine/Rendering/Materials/HowTo/CreatingLayeredMaterials/CreatingLayeredMaterials.INT.udn
--- a/Engine/Documentation/Source/Engine/Rendering/Materials/HowTo/CreatingLayeredMaterials/CreatingLayeredMaterials.INT.udn
+++ b/Engine/Documentation/Source/Engine/Rendering/Materials/HowTo/CreatingLayeredMaterials/CreatingLayeredMaterials.INT.udn
@@ -2,6 +2,8 @@ Availability:Public
 Title:Creating Layered Materials
 Crumbs: %ROOT%, Engine, Engine/Rendering/Materials, Engine/Rendering/Materials/LayeredMaterials
 Description:An introductory document on the Layered Materials technique in UE4.
+Version: 4.9
+SkillLevel: Intermediate
 
 [VAR:Topic]
 [OBJECT:Topic]
@@ -48,7 +50,7 @@ Description:An introductory document on the Layered Materials technique in UE4.
 [/REGION]
 
 
-In this brief tutorial, we will overview the process of creating a simple Layered Material with two Material Layers: chrome and snow. Our final Layered Material will automatically place snow on the top surfaces of the object, effectively switching between the two materials. The blend between the materials will always check for the top surface, meaning that even as you rotate the objects, the snow will remain on top.
+In this brief tutorial, we will overview the process of creating a simple **Layered Material** with two Material Layers: chrome and snow. Our final Layered Material will automatically place snow on the top surfaces of the object, effectively switching between the two materials. The blend between the materials will always check for the top surface, meaning that even as you rotate the objects, the snow will remain on top.
 
 
 
@@ -195,7 +197,7 @@ We can now create our Layered Material using the layers we have generated so far
 
 	![](enableMatAttrib.png)
 
-1. From the Functions panel, drag in a MatLayerBlend_Simple Function, as well as a World_Aligned_Blend Function. We will use the MatLayerBlend_Simple to handle the transition from chrome to snow, and the World_Aligned_Blend to power the Layer Blend based on the direction the surface is pointed. 
+1. From the Functions panel, drag in a **MatLayerBlend_Simple** Function, as well as a **World_Aligned_Blend** Function. We will use the MatLayerBlend_Simple to handle the transition from chrome to snow, and the World_Aligned_Blend to power the Layer Blend based on the direction the surface is pointed. 
 
 
 ### Layered Material Network
