#UE4 Docs: Jira UEDOC-2113 (Accessing Class Defaults at Runtime).

diff --git a/Engine/Documentation/Source/Engine/Blueprints/Editor/Modes/Defaults/DefaultsTab/DefaultsTab.INT.udn b/Engine/Documentation/Source/Engine/Blueprints/Editor/Modes/Defaults/DefaultsTab/DefaultsTab.INT.udn
--- a/Engine/Documentation/Source/Engine/Blueprints/Editor/Modes/Defaults/DefaultsTab/DefaultsTab.INT.udn
+++ b/Engine/Documentation/Source/Engine/Blueprints/Editor/Modes/Defaults/DefaultsTab/DefaultsTab.INT.udn
@@ -26,5 +26,46 @@ The **Class Defaults** tab contains the following sections:
 If you create custom categories for variables, the **Class Defaults** tab will display those categories too, after the Blueprint is compiled.
 [/REGION]
 
+## Get Class Defaults
+
+The properties defined inside the Class Defaults can also be accessed at runtime through the use of the **Get Class Defaults** Blueprint node. Targeted use case is for data-only Blueprint Class types and other cases where you may not necessarily need/want to create an instance in order to access default property values. 
+
+**To Access the Get Class Defaults Node**
+
+1. **Right-click** inside your Blueprint Class or Function Graph and enter **GetClassDefaults** from the context menu. 
+
+	![](GetClassDefaults.png)
+
+1. Choose the class to access via the class picker drop-down menu. 
+
+	![](SelectClass.png)
+
+	You can also wire in a class type result from a function call or variable instead as shown in the example below. 
+
+	![](FromVariable.png)
+
+1. With the node selected, you can show/hide pins via the **Details** panel. 
+
+	![](DetailsPanel.png)
+
+**Usage Notes**
+
+* Output pins (accessible class properties) are shown by default and can be optionally hidden (similar to Struct Break/Set Nodes). 
+	* Parent class properties are currently hidden by default to reduce the number of output pins.
+	* Properties tagged as _Advanced View_ are also hidden by default in the **Advanced** node section.
+* Object reference property default values are **not** currently exposed. 
+	* This means you currently cannot, for instance, access the default property values of component templates or other archetypes within a Blueprint Class. 
+	* This is because we do not currently enforce read-only reference types in a Blueprint, and we'd need to do that in order to ensure that the internal state of the referenced object cannot be changed.\
+* Only **BlueprintVisible** or **BlueprintReadOnly** properties are currently exposed for defaults access. 
+	* This matches the behavior of "normal" variable Get nodes exposed through the context menu. 
+* Currently the _Class_ input drop-down lists **all** UOBject-derived class types (including engine classes). 
+	* This may be limited in the future to exclude engine-specific class types. 
+* Adding a **GetClassDefaults** node to an Actor-based Blueprint class (for example) will currently default the input class pin's value to match the current Blueprint Class type when placed.
+
+	 
+
+
+
+
 
 
diff --git a/Engine/Documentation/Source/Engine/Blueprints/UserGuide/Types/ClassBlueprint/ClassBlueprint.INT.udn b/Engine/Documentation/Source/Engine/Blueprints/UserGuide/Types/ClassBlueprint/ClassBlueprint.INT.udn
--- a/Engine/Documentation/Source/Engine/Blueprints/UserGuide/Types/ClassBlueprint/ClassBlueprint.INT.udn
+++ b/Engine/Documentation/Source/Engine/Blueprints/UserGuide/Types/ClassBlueprint/ClassBlueprint.INT.udn
@@ -137,6 +137,10 @@ Data-Only Blueprint are edited in a compact property editor, but can also be "co
 variables, or components using the full **Blueprint Editor**.
 [/EXCERPT:DataBlueprint]
 
+[REGION:note]
+Class Defaults can be accessed at runtime for data-only Blueprints through the **GetClassDefaults** Blueprint node. Please see the [Get Class Defaults](Engine\Blueprints\Editor\Modes\Defaults\DefaultsTab) for more information. 
+[/REGION]
+
 
 
 
