#UE4doc #docs - Added Engine Version Field. Added Skill Level Field. bolded terms. Updated images.

diff --git a/Engine/Documentation/Source/Engine/Rendering/Materials/HowTo/DetailTexturing/DetailTexturing.INT.udn b/Engine/Documentation/Source/Engine/Rendering/Materials/HowTo/DetailTexturing/DetailTexturing.INT.udn
--- a/Engine/Documentation/Source/Engine/Rendering/Materials/HowTo/DetailTexturing/DetailTexturing.INT.udn
+++ b/Engine/Documentation/Source/Engine/Rendering/Materials/HowTo/DetailTexturing/DetailTexturing.INT.udn
@@ -7,6 +7,8 @@ Related: Engine/Rendering/Materials/CustomizedUVs
 Related: Engine/Rendering/Materials/Functions/Reference/Texturing
 Related: Platforms/Oculus
 Related: Engine/Rendering/Materials/ExpressionReference/Textures
+Version:4.9
+SkillLevel: Intermediate
 
 [TOC(start:2 end:3)]
 
@@ -18,7 +20,7 @@ To address this issue, you can use what is known as a Detail Texture to help hid
 
 ## Detail Texturing
 
-Detail texturing allows you to give the illusion of more detail in a Texture by bringing in a highly repeated Diffuse and Normal map combination which layers over the original Diffuse and Normal for an object. 
+**Detail Texturing** allows you to give the illusion of more detail in a Texture by bringing in a highly repeated Diffuse and Normal map combination which layers over the original Diffuse and Normal for an object. 
 This gives the illusion of greater detail at close range than normally would be possible. 
 Here is an example of Detail Texturing in action. 
 
@@ -31,7 +33,7 @@ Notice how in the image on the left, 1, there looks to be more detail from the N
 
 ## Detail Texturing Node Breakdown
 
-If you do a search in the **Palette** for **Detail Texturing**, you will find a Detail Texturing Material Function that you can use. 
+If you do a search in the **Palette** for "Detail Texturing", you will find a **DetailTexturing** Material Function that you can use. 
 While this is not the only way that you can apply Detail Textures to Materials, it is the quickest as all you have to do is just input Textures and settings you want to use. 
 
 ![](DT_Node_Break_Down.png)
@@ -111,7 +113,7 @@ Once selected, a new **Material** will be created. Name the Material **Detail_Ma
 	Remember that the **T_Ground_Gravel** Diffuse and **T_Ground_Moss** Normal map Texture Samplers need to be converted into **Texture Objects** or the Material will not work. For more information on how to do this, please refer to the **Converting Texture Sample to Texture Objects** section above. 
 	[/region]
 
-1. Now that the Textures have been hooked up, it is time to add some Constant Material Expressions so that things like the scale and intensity of the Detail Textures can be tweaked. To find the Constant Material Expression, search for Constant in the Material Palette and then drag the Constant from the **Palette** to the Material Graph. Once you have added one Constant Material Expression, duplicate it by first selecting it and then press CTRL + W on the keyboard two times. In the end, you want to have three Constant Material Expressions. 
+1. Now that the Textures have been hooked up, it is time to add some Constant Material Expressions so that things like the scale and intensity of the Detail Textures can be tweaked. To find the Constant Material Expression, search for Constant in the Material Palette and then drag the Constant from the **Palette** to the Material Graph. Once you have added one Constant Material Expression, duplicate it by first selecting it and then press **CTRL + W** on the keyboard two times. In the end, you want to have three Constant Material Expressions. 
 
 	![](DT_Adding_Constant_Nodes.png)
 
