#UE4doc #docs - Added Skill Level field. Added Engine Version field. Bolded terms.

diff --git a/Engine/Documentation/Source/Engine/Rendering/Materials/HowTo/EmissiveGlow/EmissiveGlow.INT.udn b/Engine/Documentation/Source/Engine/Rendering/Materials/HowTo/EmissiveGlow/EmissiveGlow.INT.udn
--- a/Engine/Documentation/Source/Engine/Rendering/Materials/HowTo/EmissiveGlow/EmissiveGlow.INT.udn
+++ b/Engine/Documentation/Source/Engine/Rendering/Materials/HowTo/EmissiveGlow/EmissiveGlow.INT.udn
@@ -6,6 +6,8 @@ Related: Engine/Rendering/Materials/MaterialInputs
 Related: Engine/Rendering/Materials/MaterialProperties/LightingModels
 Related: Engine/Rendering/PostProcessEffects/Bloom
 Related: Resources/ContentExamples/MaterialNodes/1_5
+Version: 4.9
+SkillLevel: Intermediate
 
 
 [TOC(start:2 end:3)]
@@ -14,7 +16,7 @@ Related: Resources/ContentExamples/MaterialNodes/1_5
 
 During the course of project development, the need might arise for a Self-Illuminated Material or glowing Material. 
 Inside of Unreal Engine 4(UE4), Self-Illuminated Materials are known as Emissive Materials. 
-In the following How - To guide, we will be covering how you setup your Materials to use the Emissive input as well as how you can use the Emissive input to actually light the world that you are creating. 
+In the following How-to guide, we will be covering how you setup your Materials to use the Emissive input as well as how you can use the Emissive input to actually light the world that you are creating. 
 
 ## Emissive Glow
 
@@ -60,7 +62,7 @@ All of the content that is used in the following sections can be found if you in
 	![](EM_Adding_Nodes_&_Textures_To_Material.png)
 
 	[region:tip]
-	If you cannot find the **Tech_Panel** textures that are used in this How - To, that means that you might have forgotten to include the **Starter Content** with your project. 
+	If you cannot find the **Tech_Panel** textures that are used in this How-to, that means that you might have forgotten to include the **Starter Content** with your project. 
 	To find this content, you can create another project, make sure to include the Start Content, and then use the [Migrate Tool](Engine/Content/Browser/UserGuide/Migrate) 
 	to move the textures from that project into your project. 
 	[/region]
@@ -134,7 +136,7 @@ Texture exactly which parts should be illuminated. In the example above, we used
 	When completed, your room should look something like this. 
 	![](EM_Testing_Level.png)
 
-1. When the room is completed, make sure to save your work by pressing **CTRL + S** on the keyboard TO SAVE YOUR WORK.
+1. When the room is completed, make sure to save your work by pressing **CTRL + S** on the keyboard to save your work.
 
 ## Emissive Materials in Action
 
