#UE4doc #docs - Added Skill Level field. Added Engine Version field. Resized two images.

diff --git a/Engine/Documentation/Source/Engine/Rendering/Materials/IntroductionToMaterials/IntroductionToMaterials.INT.udn b/Engine/Documentation/Source/Engine/Rendering/Materials/IntroductionToMaterials/IntroductionToMaterials.INT.udn
--- a/Engine/Documentation/Source/Engine/Rendering/Materials/IntroductionToMaterials/IntroductionToMaterials.INT.udn
+++ b/Engine/Documentation/Source/Engine/Rendering/Materials/IntroductionToMaterials/IntroductionToMaterials.INT.udn
@@ -2,6 +2,8 @@ Availability: Public
 Title:Essential Material Concepts
 Crumbs: %ROOT%, Engine, Engine/Rendering/Materials
 Description:An introductory document on UE4 Materials and how they work.
+Version: 4.9
+SkillLevel: Beginner
 
 [VAR:playlistnav]
 [OBJECT:PlaylistNavigation]
@@ -178,11 +180,11 @@ Applying a Material to a Brush surface is a simple process:
 
 1. Click on the Brush surface to which you wish to assign the Material.
 
-	![](BoxPanelSelected.png)
+	![](BoxPanelSelected.png)(w:700)
 
 1. In the **Details** panel, click the ![](ArrowButton.png) button. Pressing this button will apply the Material to the Brush surface.
 
-	![](AddingMaterialToBox.png)![](RedOnBox.png)
+	![](AddingMaterialToBox.png)![](RedOnBox.png)(w:700)
 
 
 %playlistnav%
