#UE4doc #docs - Added Engine Version field.

diff --git a/Engine/Documentation/Source/Engine/Rendering/Materials/LightingModels/LightingModels.INT.udn b/Engine/Documentation/Source/Engine/Rendering/Materials/LightingModels/LightingModels.INT.udn
--- a/Engine/Documentation/Source/Engine/Rendering/Materials/LightingModels/LightingModels.INT.udn
+++ b/Engine/Documentation/Source/Engine/Rendering/Materials/LightingModels/LightingModels.INT.udn
@@ -2,6 +2,7 @@ Availability:Public
 Title:Material Shading Models
 Crumbs: %ROOT%, Engine, Engine/Rendering/Materials
 Description:Calculation models used by materials to determine how surfaces interact with light.
+Version: 4.9
 
 
 [OBJECT:TopicButtonList]
