#UE4doc #docs - Added Engine Version Field.

diff --git a/Engine/Documentation/Source/Engine/Rendering/Materials/LightingModels/SubSurface/SSSLightingModel.INT.udn b/Engine/Documentation/Source/Engine/Rendering/Materials/LightingModels/SubSurface/SSSLightingModel.INT.udn
--- a/Engine/Documentation/Source/Engine/Rendering/Materials/LightingModels/SubSurface/SSSLightingModel.INT.udn
+++ b/Engine/Documentation/Source/Engine/Rendering/Materials/LightingModels/SubSurface/SSSLightingModel.INT.udn
@@ -2,6 +2,7 @@ Availability:Public
 Title:Subsurface Shading Model
 Crumbs: %ROOT%, Engine, Engine/Rendering/Materials
 Description:Description and technical details of the subsurface shading model available in materials.
+Version: 4.9
 
 Materials have a new subsurface shading model (**_MLM_Subsurface_**) which is intended for materials like wax or jade which appear opaque, but lighting scatters inside them allowing a portion of the light from the opposite side of the surface to show through.
 
