#UE4doc #docs - Added Engine Version Field. Replaced Image Table with Comparison Slider.

diff --git a/Engine/Documentation/Source/Engine/Rendering/Overview/RenderingOverview.INT.udn b/Engine/Documentation/Source/Engine/Rendering/Overview/RenderingOverview.INT.udn
--- a/Engine/Documentation/Source/Engine/Rendering/Overview/RenderingOverview.INT.udn
+++ b/Engine/Documentation/Source/Engine/Rendering/Overview/RenderingOverview.INT.udn
@@ -2,6 +2,7 @@ Availability:Public
 Title: Rendering Overview
 Crumbs: %ROOT%, Engine
 Description:Overview of the main features of the rendering subsystem.
+Version: 4.9
 
 [VAR:Topic]
 [OBJECT:Topic]
@@ -76,7 +77,7 @@ Translucency can cast shadows onto the opaque world and onto itself and other li
 [REGION:imagetable]
 ![Translucent Particle Self-Shadowing](self_shadow_1.png)(w:465) | ![Translucent Particle Self-Shadowing](self_shadow_2.png)(w:465) |
 | ------ | ------ |
-| | |
+| Translucent Particle Self-Shadowing (Side View) | Translucent Particle Self-Shadowing (Front View) |
 [/REGION]
 
 See [Lit translucency](Engine/Rendering/LightingAndShadows/LitTranslucency) for more information.
@@ -95,11 +96,14 @@ Materials have a new shading model **MLM_Subsurface** which is intended for mate
 
 UE4 supports simulating particles on the GPU. Traditional CPU systems allow for thousands of particles in a frame. GPU simulation allows for hundreds of thousands of particles to be simulated and rendered efficiently. 
 
-[REGION:imagetable]
-![](gpu_spawn_500.png)(w:465) | ![](gpu_spawn_20000.png)(w:465) |
-| ------ | ------ |
-| Spawn Rate 500 | Spawn Rate 20000 |
-[/REGION]
+[OBJECT:ComparisonSlider]
+	[PARAM:before]
+	![Spawn Rate - 500](gpu_spawn_500.png)(w:500)
+	[/PARAM]
+	[PARAM:after]
+	![Spawn Rate - 20000](gpu_spawn_20000.png)(w:500)
+	[/PARAM]
+[/OBJECT]
 
 ### Vector Fields
 
@@ -155,11 +159,14 @@ The **Lens Flare** effect is an image-based technique that automatically creates
 
 **Tone Mapping** allows the colors of the rendered scene to be shifted or modified to alter the end result. This can be used to create effects such as a sepia filter, hit effects (i.e., a red flash), and many more.
 
-[REGION:imagetable]
-![](scene_tint_full.png)(w:465) | ![](scene_tint_quarter.png)(w:465) |
-| ------ | ------ |
-| Tone Mapping - Off | Tone Mapping - On |
-[/REGION]
+[OBJECT:ComparisonSlider]
+	[PARAM:before]
+	![Tone Mapping - Off](scene_tint_full.png)(w:500)
+	[/PARAM]
+	[PARAM:after]
+	![Tone Mapping - On](scene_tint_quarter.png)(w:500)
+	[/PARAM]
+[/OBJECT]
 
 ### Vignette
 
