#UE4doc #docs - Added Engine Version Field.

diff --git a/Engine/Documentation/Source/Engine/Rendering/ParticleSystems/ParticleSystems.INT.udn b/Engine/Documentation/Source/Engine/Rendering/ParticleSystems/ParticleSystems.INT.udn
--- a/Engine/Documentation/Source/Engine/Rendering/ParticleSystems/ParticleSystems.INT.udn
+++ b/Engine/Documentation/Source/Engine/Rendering/ParticleSystems/ParticleSystems.INT.udn
@@ -2,6 +2,8 @@ Availability:Public
 Title: Particle Systems
 Crumbs: %ROOT%, Engine, Engine/Rendering 
 Description:Assets the contain any number of particle emitters used to create effects in the world.
+Version: 4.9
+SkillLevel: Beginner
 
 [VAR:Topic]
 [OBJECT:Topic]
@@ -51,7 +53,7 @@ Description:Assets the contain any number of particle emitters used to create ef
 Unreal Engine contains an extremely powerful and robust particle system, allowing artists to create mind-blowing visual effects ranging from smoke, sparks, and fire to far more intricate and otherworldly examples. 
 [/EXCERPT:Intro]
 
-Unreal's particle systems are edited via Cascade, a fully integrated and modular particle effects editor. Cascade offers real-time feedback and modular effects editing, allowing fast and easy creation of even the most complex effects.
+Unreal's particle systems are edited via **Cascade**, a fully integrated and modular particle effects editor. Cascade offers real-time feedback and modular effects editing, allowing fast and easy creation of even the most complex effects.
 
 Particle systems are also very closely related to the various materials and textures applied to each particle. The primary job of the particle system itself is to control the behavior of the particles, while the specific look and feel of the particle system as a whole is often controlled by way of materials. As such, we have included a shortcut to the material documentation below.
 
diff --git a/Engine/Documentation/Source/Engine/UI/LevelEditor/Details/Properties/Array/ArrayProperty.INT.udn b/Engine/Documentation/Source/Engine/UI/LevelEditor/Details/Properties/Array/ArrayProperty.INT.udn
--- a/Engine/Documentation/Source/Engine/UI/LevelEditor/Details/Properties/Array/ArrayProperty.INT.udn
+++ b/Engine/Documentation/Source/Engine/UI/LevelEditor/Details/Properties/Array/ArrayProperty.INT.udn
@@ -2,6 +2,7 @@ Availability:Public
 Title: Array Control
 Crumbs: %ROOT%, Engine, Engine/UI, Engine/UI/LevelEditor, Engine/UI/LevelEditor/Details
 Description:Dynamic lists of other property types.
+Version:4.9
 
 `Arrays` are dynamic lists of some other type of property, such as Object references. The list of 
 elements is displayed vertically with each individual element using the standard controls for that 
