#UE4doc #docs - Added Engine Version Field.

diff --git a/Engine/Documentation/Source/Engine/UI/LevelEditor/Details/Properties/Basic/BasicProperty.INT.udn b/Engine/Documentation/Source/Engine/UI/LevelEditor/Details/Properties/Basic/BasicProperty.INT.udn
--- a/Engine/Documentation/Source/Engine/UI/LevelEditor/Details/Properties/Basic/BasicProperty.INT.udn
+++ b/Engine/Documentation/Source/Engine/UI/LevelEditor/Details/Properties/Basic/BasicProperty.INT.udn
@@ -2,6 +2,7 @@ Availability:Public
 Title: Basic Controls
 Crumbs: %ROOT%, Engine, Engine/UI, Engine/UI/LevelEditor, Engine/UI/LevelEditor/Details
 Description:Simple property types such as numeric boxes and spinners, text boxes, checkboxes, radio buttons, and comboboxes.
+Version: 4.9
 
 Unreal Editor has many simple control types you are familiar with for editing numeric values, text, etc.
 
