#UE4doc #docs - Added Engine Version Field.

diff --git a/Engine/Documentation/Source/Engine/UI/LevelEditor/Details/Properties/Color/ColorProperty.INT.udn b/Engine/Documentation/Source/Engine/UI/LevelEditor/Details/Properties/Color/ColorProperty.INT.udn
--- a/Engine/Documentation/Source/Engine/UI/LevelEditor/Details/Properties/Color/ColorProperty.INT.udn
+++ b/Engine/Documentation/Source/Engine/UI/LevelEditor/Details/Properties/Color/ColorProperty.INT.udn
@@ -2,6 +2,7 @@ Availability:Public
 Title: Color Control
 Crumbs: %ROOT%, Engine, Engine/UI, Engine/UI/LevelEditor, Engine/UI/LevelEditor/Details
 Description:Properties that hold byte and linear color values.
+Version: 4.9
 
 Color and LinearColor properties provide a visual display of the color in the form of a _swatch_ along with numerical values 
 for each of the individual components that make up the color.
