#UE4doc #docs - Added Engine Version Field.

diff --git a/Engine/Documentation/Source/Engine/UI/LevelEditor/Details/Properties/ObjectReference/ObjectReferenceProperty.INT.udn b/Engine/Documentation/Source/Engine/UI/LevelEditor/Details/Properties/ObjectReference/ObjectReferenceProperty.INT.udn
--- a/Engine/Documentation/Source/Engine/UI/LevelEditor/Details/Properties/ObjectReference/ObjectReferenceProperty.INT.udn
+++ b/Engine/Documentation/Source/Engine/UI/LevelEditor/Details/Properties/ObjectReference/ObjectReferenceProperty.INT.udn
@@ -2,6 +2,7 @@ Availability:Public
 Title: Object Reference Control
 Crumbs: %ROOT%, Engine, Engine/UI, Engine/UI/LevelEditor, Engine/UI/LevelEditor/Details
 Description:properties that hold references to Objects and Actors.
+Version: 4.9
 
  `Object` reference properties are properties that hold references to an Actor in the level. In 
 the **Details** panel, an `Object` reference property is displayed as a drop-down box along with a set of controls for modifying the reference. 
