#UE4doc #docs - Added Engine Version Field.

diff --git a/Engine/Documentation/Source/Engine/UI/LevelEditor/Details/Properties/VectorAndRotator/VectorAndRotatorProperty.INT.udn b/Engine/Documentation/Source/Engine/UI/LevelEditor/Details/Properties/VectorAndRotator/VectorAndRotatorProperty.INT.udn
--- a/Engine/Documentation/Source/Engine/UI/LevelEditor/Details/Properties/VectorAndRotator/VectorAndRotatorProperty.INT.udn
+++ b/Engine/Documentation/Source/Engine/UI/LevelEditor/Details/Properties/VectorAndRotator/VectorAndRotatorProperty.INT.udn
@@ -2,6 +2,7 @@ Availability:Public
 Title: Vector / Rotator Controls
 Crumbs: %ROOT%, Engine, Engine/UI, Engine/UI/LevelEditor, Engine/UI/LevelEditor/Details
 Description:Vector or Rotator properties that contain multiple Float values.
+Version: 4.9
 
  `Vectors` and `Rotators` are multi-value properties that display numeric inputs for each of their components 
 laid out horizontally.
