#UE4doc #docs - Added Skill Level field. Added Engine Version Field. Added tool tip to first section that was unclear.

diff --git a/Engine/Documentation/Source/Engine/UI/LevelEditor/InEditorTesting/KeepSimulationChanges/KeepSimulationChanges.INT.udn b/Engine/Documentation/Source/Engine/UI/LevelEditor/InEditorTesting/KeepSimulationChanges/KeepSimulationChanges.INT.udn
--- a/Engine/Documentation/Source/Engine/UI/LevelEditor/InEditorTesting/KeepSimulationChanges/KeepSimulationChanges.INT.udn
+++ b/Engine/Documentation/Source/Engine/UI/LevelEditor/InEditorTesting/KeepSimulationChanges/KeepSimulationChanges.INT.udn
@@ -2,13 +2,15 @@ Availability:Public
 Title:Keep Simulation Changes
 Crumbs: %ROOT%, Engine, Engine/UI, Engine/UI/LevelEditor, Engine/UI/LevelEditor/InEditorTesting
 Description: Updates Actor properties from the simulated or in-game state.
+Version: 4.9
+SkillLevel: Intermediate
 
 While using _Simulate In Editor_ or _Play In Editor_, you can edit the properties of Actors and see the effect of those changes on your level and the gameplay in it.  Using
 **Keep Simulation Changes**, you can take those changes and make them permanent. Changes that you can keep include properties seen in the **Details** tab, such as the Actor's
 transform in the level. After using Keep Simulation Changes, the selected Actors in your level will be updated from the runtime copy, so your changes will be reflected in the editor's copy. 
 
 Whether you alter an Actor's properties with _Blueprints_ or C++ code as a result of gameplay, or you begin simulating physics on an Actor and like where it ends up in the map,
-**Keep Simulation Changes** will allow you to save those changes. However, it  only works for Actors which were already placed in your level before you began simulating or playing the
+Keep Simulation Changes will allow you to save those changes. However, it  only works for Actors which were already placed in your level before you began simulating or playing the
  game.  If Actors are spawned from _Blueprints_ or C++ code, you will not be able to save their changes. 
 
 
@@ -18,6 +20,10 @@ Whether you alter an Actor's properties with _Blueprints_ or C++ code as a resul
 
 	![](in_simulate_keepsimchanges.png)(w:500)
 
+[Region:note]
+In order to see this option in the context menu your need to be Playing or Simulating in the Editor.
+[/Region]
+
 **To keep changes from Play In Editor:**
 	
 1. Press **Shift+F1** to regain mouse control.
