#UE4doc #docs - Added Engine Version field. Added HLODs to the Window sections.

diff --git a/Engine/Documentation/Source/Engine/UI/LevelEditor/MenuBar/MainMenuBar.INT.udn b/Engine/Documentation/Source/Engine/UI/LevelEditor/MenuBar/MainMenuBar.INT.udn
--- a/Engine/Documentation/Source/Engine/UI/LevelEditor/MenuBar/MainMenuBar.INT.udn
+++ b/Engine/Documentation/Source/Engine/UI/LevelEditor/MenuBar/MainMenuBar.INT.udn
@@ -2,6 +2,7 @@ Availability:Public
 Title: Main Menu Bar
 Crumbs:%ROOT%, Engine, Engine/UI
 Description:Menu containing general tools and commands that are used when working in the editor.
+Version: 4.9
 
 [EXCERPT:Intro]
 The **Menu Bar** in editor should be familiar to anyone who has used Windows applications previously. It provides access to general tools and commands that are used when working with levels in the editor.
@@ -53,6 +54,7 @@ The **Menu Bar** in editor should be familiar to anyone who has used Windows app
 | [REGION:tablesection]**Configuration**[/REGION] ||
 |**Editor Preferences**| Configure the behavior and features of this editor. See [](Engine/UI/EditorSettings) for details. |
 |**Project Settings**| Change the settings of the currently loaded project. See [](Engine/UI/ProjectSettings) for details. |
+|**Plugins**| This will open the Plugins browser where these can be enabled or disabled for the current project. |
 [/VAR]
 
 
@@ -73,7 +75,7 @@ The **Menu Bar** in editor should be familiar to anyone who has used Windows app
 | ------- | ----------- |
 |**Details**| Shows the **Details** panel, which provides access to various properties for selected objects. You may show up to 4 **Details** panels simultaneously.|
 |**Viewports**| Shows the Viewport, which is your view into the 3D world of your game. You may show up to 4 Viewports simultaneously.|
-|**Build and Submit**| Opens the **Build and Submit** panel, which allows you to build lighting for your level, save the level, then submit it to source control.|
+|**Hierarchical LOD Outliner**| This will display the HLOD settings and allow generation of clusters.|
 | **Layers** | Opens the **Layers** panel, where you can divide your level into a series of visibility layers.|
 | **Levels** | Opens the **Levels** panel which displays the persistent and all streaming levels. |
 | **Modes** | Shows the **Modes** panel, which accesses special level editing tools such as Landscape and Foliage. |
