#UE4doc #docs:
Resubmit as not all changes were saved

diff --git a/Engine/Documentation/Source/Engine/Blueprints/UserGuide/Timelines/Examples/LightFlickering/LightFlickering.INT.udn b/Engine/Documentation/Source/Engine/Blueprints/UserGuide/Timelines/Examples/LightFlickering/LightFlickering.INT.udn
--- a/Engine/Documentation/Source/Engine/Blueprints/UserGuide/Timelines/Examples/LightFlickering/LightFlickering.INT.udn
+++ b/Engine/Documentation/Source/Engine/Blueprints/UserGuide/Timelines/Examples/LightFlickering/LightFlickering.INT.udn
@@ -2,6 +2,7 @@ Availability:Public
 Title:Flickering Lights
 Crumbs: %ROOT%, Engine, Engine/Blueprints, Engine/Blueprints/UserGuide, Engine/Blueprints/UserGuide/Timelines
 Description: An example use of a timeline in which we set up a flickering light that also changes color.
+version: 4.9
 
 %Description%
 
@@ -20,9 +21,9 @@ The light shown will flicker its brightness and change color.
 
 The timeline serves 3 primary functions, with a separate track for each:
 
-1. A **float track** (named Brightness) handles animating the brightness. This curve was added from the **Content Browser** using the Flicker Curve asset.
-2. A **vector track** (named Light Color) animates the light's color. This was handled such that the light turns orange at the beginning of the sequence and then returns to white at the end. The curve interpolation was set to clamp, allowing the value to remain constant until the next key.
-3. An **event track** (named Trigger Toggle) disables our trigger throughout the duration of the timeline's sequence. Because this is an event track, the values of the keyframes are irrelevant. Wherever the timeline encounters a keyframe, an execution pulse is sent out of the track's associated output pin.
+1. A **Float track** (named Brightness) handles animating the brightness. This curve was added from the **Content Browser** using the Flicker Curve asset.
+2. A **Vector track** (named Light Color) animates the light's color. This was handled such that the light turns orange at the beginning of the sequence and then returns to white at the end. The curve interpolation was set to clamp, allowing the value to remain constant until the next key.
+3. An **Event track** (named Trigger Toggle) disables our trigger throughout the duration of the timeline's sequence. Because this is an event track, the values of the keyframes are irrelevant. Wherever the timeline encounters a keyframe, an execution pulse is sent out of the track's associated output pin.
 
 
 
@@ -30,9 +31,9 @@ The logic of the network is fairly straightforward. Because we always want the f
 
 That variable has a default value of false, meaning we will always play the first time we touch the trigger. So execution from the timeline runs as follows:
 
-* The result of the Brightness track is multiplied by *20* to give the light more intensity. 
+* The result of the Brightness track is multiplied by 20 to give the light more intensity. 
 * The color of the light has also changed as the next action.
-* Using the event track, we then trigger a FlipFlop that sets the *Is Timeline Playing* variable to *true*. Since this variable is checked whenever someone touches the trigger, we are deactivating the trigger as soon as playback begins.
+* Using the event track, we then trigger a FlipFlop that sets the Is Timeline Playing variable to true. Since this variable is checked whenever someone touches the trigger, we are deactivating the trigger as soon as playback begins.
 * As playback ends, the event track sends another pulse out of the Toggle Trigger event track, which hits the FlipFlop again, turning our boolean back off and thereby reactivating the trigger.
 
 
