#UE4doc #docs:
Added engine version field
Replaced an image due to out of date UI
Added a lightbox region to a picture as it was linking to a blank page when clicked

diff --git a/Engine/Documentation/Source/Engine/Blueprints/UserGuide/Timelines/Examples/OpeningDoors/OpeningDoors.INT.udn b/Engine/Documentation/Source/Engine/Blueprints/UserGuide/Timelines/Examples/OpeningDoors/OpeningDoors.INT.udn
--- a/Engine/Documentation/Source/Engine/Blueprints/UserGuide/Timelines/Examples/OpeningDoors/OpeningDoors.INT.udn
+++ b/Engine/Documentation/Source/Engine/Blueprints/UserGuide/Timelines/Examples/OpeningDoors/OpeningDoors.INT.udn
@@ -2,6 +2,7 @@ Availability:Public
 Title:Opening Doors
 Crumbs: %ROOT%, Engine, Engine/Blueprints, Engine/Blueprints/UserGuide, Engine/Blueprints/UserGuide/Timelines
 Description: An example use of a timeline in which we set up a classic proximity-based opening door.
+version: 4.9
 
 %Description%
 
@@ -22,9 +23,9 @@ Once we have done that, we can set up the timeline itself. In this case, the tim
 When the player exits the volume, we wait 2 seconds and then play the timeline in reverse, bringing the door back down. 
 
 **Click for full size**
-
+[REGION:lightbox]
 [![DoorNetwork.png](DoorNetwork.png)(w:800)](DoorNetwork.png)
-
+[/REGION]
 Note that sound effects are not considered in this setup.
 
 ## Timeline Setup
