#UE4doc #docs:
Added engine version field
Removed Overview header
Replaced a few images due to out of date UI
Removed section about Curve Interpolation as it no longer exists
Fixed capitalization in a few places

diff --git a/Engine/Documentation/Source/Engine/Blueprints/UserGuide/Timelines/KeysAndCurves/KeysAndCurves.INT.udn b/Engine/Documentation/Source/Engine/Blueprints/UserGuide/Timelines/KeysAndCurves/KeysAndCurves.INT.udn
--- a/Engine/Documentation/Source/Engine/Blueprints/UserGuide/Timelines/KeysAndCurves/KeysAndCurves.INT.udn
+++ b/Engine/Documentation/Source/Engine/Blueprints/UserGuide/Timelines/KeysAndCurves/KeysAndCurves.INT.udn
@@ -2,8 +2,8 @@ Availability:Public
 Title: Keys and Curves
 Crumbs: %ROOT%, Engine, Engine/Blueprints, Engine/Blueprints/UserGuide, Engine/Blueprints/UserGuide/Timelines
 Description: This document contains an overview of how to work with keys and curves within the Timeline editor in Blueprints.
+version: 4.9
 
-## Overview
 
 %Description%
 
@@ -19,7 +19,7 @@ Keys are added by holding the **Shift** key and clicking in the gray bar.
 
 ### Editing Keys
 
-The time and value of the key can be set by clicking on the key and entering the values into the time and value fields near the top of the track, which display next to the index number of the key.
+The time and value of the key can be set by clicking on the key and entering the values into the time and value fields near the top of the track.
 
 ![Blueprint Timeline - Key Values](k2_timeline_track_key_values.png)
 
@@ -33,18 +33,12 @@ To move a key along the timeline, select the key and then click and drag it. Mul
 
 ### Key Interpolation
 
-By **right-clicking** on a key, you can choose the interpolation type for a given key.  The same interpolation types are available for key interpolation as for curve interpolation, but key interpolation only affects the 
+By **Right-clicking** on a key, you can choose the interpolation type for a given key.  Key interpolation only affects the 
 curve between the key you interpolate for and the next key after that. For example, with all other keys set to **Linear**, and the center key (shown in orange) set to **Cubic-Auto**, this is how the track looks.
 
 ![](key_interpolation_example.png)
 
-## Curves 
-
-This section covers a few important concepts for working with interpolation curves in Timelines.
-
-### Curve Interpolation
-
-By **right-clicking** on a curve (not a key), you can choose the interpolation type for a given curve. The available interpolation types are as follows:
+The available interpolation types are as follows:
 
 * **Auto**
 
@@ -67,11 +61,7 @@ By **right-clicking** on a curve (not a key), you can choose the interpolation t
 	![](Clamped.png)
 
 
-
-
-
-
-### External Curves
+## External Curves
 
 To the left of the main key area is a field where you can add an external curve asset from the **Content Browser** to that track. This is different than using the ![](Engine/Blueprints/UserGuide\Timelines\Editor\add_external_curve_button.png) button, which creates a new
 track from a selected external curve.  
@@ -82,7 +72,7 @@ track from a selected external curve.
 | ![](external_curve_browse_button.png)| Browse in the **Content Browser** to select a curve for this track. |
 | ![](external_curve_internal_button.png)| Convert an imported external curve to an internal curve so that keys and curves can be edited. |
 
-However, you may also **right-click** on any curve (with at least one key selected) and choose **Create External Curve**. This brings up a context window in which you can edit the path into which the curve asset will be saved. In doing this, you may then use this curve again on other timelines.
+However, you may also **Right-click** on any curve (with at least one key selected) and choose **Create External Curve**. This brings up a context window in which you can edit the path into which the curve asset will be saved. In doing this, you may then use this curve again on other timelines.
 
 ![k2_timeline_createExternalCurve.png](k2_timeline_createExternalCurve.png)
 
