#UE4doc #docs:
Added engine version field
Removed italics from a few terms
Fixed some wording

diff --git a/Engine/Documentation/Source/Engine/Blueprints/UserGuide/Timelines/TimelineNodes/TimelineNodes.INT.udn b/Engine/Documentation/Source/Engine/Blueprints/UserGuide/Timelines/TimelineNodes/TimelineNodes.INT.udn
--- a/Engine/Documentation/Source/Engine/Blueprints/UserGuide/Timelines/TimelineNodes/TimelineNodes.INT.udn
+++ b/Engine/Documentation/Source/Engine/Blueprints/UserGuide/Timelines/TimelineNodes/TimelineNodes.INT.udn
@@ -2,6 +2,7 @@ Availability:Public
 Title:Timeline Nodes
 Crumbs: %ROOT%, Engine, Engine/Blueprints, Engine/Blueprints/UserGuide, Engine/Blueprints/UserGuide/Timelines
 Description: Nodes that are used to handle specific functionality of Blueprint Timelines.
+version: 4.9
 
 [VAR:Topic]
 [OBJECT:Topic]
@@ -84,7 +85,7 @@ The **Get Timeline Length** node returns the total length of the input Timeline
 
 ## Is Looping
 
-The **Is Looping** node returns a boolean value which will be _true_ if the input Timeline is looping, _false_ if it is not.
+The **Is Looping** node returns a boolean value which will be true if the input Timeline is looping, false if it is not.
 
 ![](Is_Looping.png)
 
@@ -97,7 +98,7 @@ The **Is Looping** node returns a boolean value which will be _true_ if the inpu
 
 ## Is Playing
 
-The **Is Playing** node returns a boolean value which will be _true_ if the input Timeline is currently playing, _false_ if it is not.
+The **Is Playing** node returns a boolean value which will be true if the input Timeline is currently playing, false if it is not.
 
 ![](Is_Playing.png)
 
@@ -110,7 +111,7 @@ The **Is Playing** node returns a boolean value which will be _true_ if the inpu
 
 ## Is Reversing
 
-The **Is Reversing** node returns a boolean value which will be _true_ if the input Timeline is currently reversing (or playing in reverse), _false_ if it is not.
+The **Is Reversing** node returns a boolean value which will be true if the input Timeline is currently reversing (or playing in reverse), false if it is not.
 
 ![](Is_Reversing.png)
 
@@ -159,7 +160,7 @@ The **Reverse** node tells the input Timeline to begin reversing (playing in rev
 
 ![](Reverse_From_End.png)
 
-The **Reverse from End** node tells the input Timeline to begin reversing (playing in reverse) from the start of end Timeline.
+The **Reverse from End** node tells the input Timeline to begin reversing (playing in reverse) from the end of the Timeline.
 
 | **Name** | **Type** | **Description** |
 | --- | --- |---|
