#UE4doc #docs:
Added engine version field
Removed iconsection regions from the page, as they don't meet the documentation guidelines

diff --git a/Engine/Documentation/Source/Engine/Blueprints/UserGuide/Types/BlueprintTypes.INT.udn b/Engine/Documentation/Source/Engine/Blueprints/UserGuide/Types/BlueprintTypes.INT.udn
--- a/Engine/Documentation/Source/Engine/Blueprints/UserGuide/Types/BlueprintTypes.INT.udn
+++ b/Engine/Documentation/Source/Engine/Blueprints/UserGuide/Types/BlueprintTypes.INT.udn
@@ -2,6 +2,7 @@ Availability:Public
 Title: Types of Blueprints
 Crumbs: %ROOT%, Engine/Blueprints
 Description:Landing page for information on different types of Blueprints.
+version: 4.9
 
 
 Sometimes it is difficult to tell which type of **Blueprint** you should use at first glance, especially 
@@ -10,75 +11,59 @@ when it comes to **Blueprint Macro Libraries** and **Blueprint Classes**. A good
 * _Are there multiple instances?_
 
 If you will have more than one or two instances (e.g., a TV set that can be shot or turned on/off), 
-then it probably makes sense to create a **Blueprint Class**, with associated code living there. 
+then it probably makes sense to create a Blueprint Class, with associated code living there. 
 If not, and you just want to have some helper functions (like find all Actors within X units), those 
-are ideal to live in a **Blueprint Macro Library**.
+are ideal to live in a Blueprint Macro Library.
 
 
 [EXCERPT:Main]
 
-[OBJECT:IconSection]
-	[PARAM:icon]
-		[![](level_blueprint_test_image.png)(w:250)](Engine/Blueprints/UserGuide/Types/LevelBlueprint)
-	[/PARAM]
-	[PARAM:title]
-		[%Engine/Blueprints/UserGuide/Types/LevelBlueprint:title%](Engine/Blueprints/UserGuide/Types/LevelBlueprint)
-	[/PARAM]
-	[PARAM:description]
-		[INCLUDE:Engine/Blueprints/UserGuide/Types/LevelBlueprint#Intro]
-	[/PARAM]
-[/OBJECT]
+##Level Blueprint
 
+![](level_blueprint_test_image.png)(w:250)
 
-[OBJECT:IconSection]
-	[PARAM:icon]
-		[![](class_blueprint_test_image.png)(w:250)](Engine/Blueprints/UserGuide/Types/ClassBlueprint)
-	[/PARAM]
-	[PARAM:title]
-		[%Engine/Blueprints/UserGuide/Types/ClassBlueprint:title%](Engine/Blueprints/UserGuide/Types/ClassBlueprint)
-	[/PARAM]
-	[PARAM:description]
-		[INCLUDE:Engine/Blueprints/UserGuide/Types/ClassBlueprint#Intro]
-	[/PARAM]
-[/OBJECT]
+[INCLUDE:Engine/Blueprints/UserGuide/Types/LevelBlueprint#Intro]
+
+For more information on this section, please see [%Engine/Blueprints/UserGuide/Types/LevelBlueprint:title%](Engine/Blueprints/UserGuide/Types/LevelBlueprint).
+
+
+##Blueprint Class
+
+![](class_blueprint_test_image.png)(w:250)
+
+[INCLUDE:Engine/Blueprints/UserGuide/Types/ClassBlueprint#Intro]
+
+For more information on this section, please see [%Engine/Blueprints/UserGuide/Types/ClassBlueprint:title%](Engine/Blueprints/UserGuide/Types/ClassBlueprint).
 
 [/EXCERPT:Main]
 
 
 [EXCERPT:Types]
-[OBJECT:IconSection]
-	[PARAM:icon]
-		[![](data_blueprint.png)(w:250)](Engine/Blueprints/UserGuide/Types/ClassBlueprint#Data-OnlyBlueprint)
-	[/PARAM]
-	[PARAM:title]
-		[Data-Only Blueprint](Engine/Blueprints/UserGuide/Types/ClassBlueprint#Data-OnlyBlueprint)
-	[/PARAM]
-	[PARAM:description]
-		[INCLUDE:Engine/Blueprints/UserGuide/Types/ClassBlueprint#DataBlueprint]
-	[/PARAM]
-[/OBJECT]
-[OBJECT:IconSection]
-	[PARAM:icon]
-		[![](blueprint_interface.png)(w:250)](Engine/Blueprints/UserGuide/Types/Interface)
-	[/PARAM]
-	[PARAM:title]
-		[](Engine/Blueprints/UserGuide/Types/Interface)
-	[/PARAM]
-	[PARAM:description]
-		[INCLUDE:Engine/Blueprints/UserGuide/Types/Interface#Intro]
-	[/PARAM]
-[/OBJECT]
-[OBJECT:IconSection]
-	[PARAM:icon]
-		[![](macro_blueprint.png)(w:250)](Engine/Blueprints/UserGuide/Types/MacroLibrary)
-	[/PARAM]
-	[PARAM:title]
-		[](Engine/Blueprints/UserGuide/Types/MacroLibrary)
-	[/PARAM]
-	[PARAM:description]
-		[INCLUDE:Engine/Blueprints/UserGuide/Types/MacroLibrary#Intro]
-	[/PARAM]
-[/OBJECT]
+
+##Data-Only Blueprint
+
+![](data_blueprint.png)(w:250)
+
+[INCLUDE:Engine/Blueprints/UserGuide/Types/ClassBlueprint#DataBlueprint]
+
+For more information on this section, please see [Data-Only Blueprint](Engine/Blueprints/UserGuide/Types/ClassBlueprint#Data-OnlyBlueprint).
+
+##Blueprint Interface
+
+![](blueprint_interface.png)(w:250)
+
+[INCLUDE:Engine/Blueprints/UserGuide/Types/Interface#Intro]
+
+For more information on this section, please see [](Engine/Blueprints/UserGuide/Types/Interface).
+
+##Blueprint Macro Library
+
+![](macro_blueprint.png)(w:250)
+
+[INCLUDE:Engine/Blueprints/UserGuide/Types/MacroLibrary#Intro]
+
+For more information on this section, please see [](Engine/Blueprints/UserGuide/Types/MacroLibrary).
+
 
 <!--
 [OBJECT:IconSection]
