#UE4doc #docs:
Added engine version field
Edited existing picture due to incorrect coloring
Replaced a couple pictures due to outdated UI
Removed italics from multiple terms

diff --git a/Engine/Documentation/Source/Engine/Blueprints/UserGuide/Types/ClassBlueprint/Creation/CreatingBlueprints.INT.udn b/Engine/Documentation/Source/Engine/Blueprints/UserGuide/Types/ClassBlueprint/Creation/CreatingBlueprints.INT.udn
--- a/Engine/Documentation/Source/Engine/Blueprints/UserGuide/Types/ClassBlueprint/Creation/CreatingBlueprints.INT.udn
+++ b/Engine/Documentation/Source/Engine/Blueprints/UserGuide/Types/ClassBlueprint/Creation/CreatingBlueprints.INT.udn
@@ -6,6 +6,7 @@ Related: Engine/Blueprints/UserGuide/Types/ClassBlueprint
 Related: Engine/Blueprints/UserGuide/Types/ClassBlueprint/Opening 
 Related: Engine/Blueprints/Editor/UIBreakdowns/ClassBPUI
 Related: Engine/Blueprints/Editor
+version: 4.9
 
 
 [TOC(start:2 end:2)]
@@ -22,13 +23,13 @@ Related: Engine/Blueprints/Editor
 
 1. In the **Content Browser**, click on ![Create Asset Button](button_new_asset_blueprint.png).
 
-1. Select **Blueprint** under the _Create Basic Asset_ section from the menu that appears.
+1. Select **Blueprint** under the **Create Basic Asset** section from the menu that appears.
 
 	[REGION:note]
-	Additional [](Engine/Blueprints/UserGuide/Types) can be created from the **Blueprints** option under _Create Advanced Asset_. 
+	Additional [](Engine/Blueprints/UserGuide/Types) can be created from the **Blueprints** option under **Create Advanced Asset**. 
 	[/REGION]
 
-1. Choose a **Parent Class** for your _Blueprint_.
+1. Choose a **Parent Class** for your Blueprint.
 
 	![Choose a Parent Class](new_asset_parent_class.png)
 
@@ -47,13 +48,13 @@ There are two other methods in which you can create a **Blueprint** from the **C
 
 1. **Right-click** in the **Asset View** (right-hand panel) of the **Content Browser**. 
 
-1. Select **Blueprint** under the _Create Basic Asset_ section from the menu that appears.
+1. Select **Blueprint** under the Create Basic Asset section from the menu that appears.
 
 	[REGION:raw]
 	![New Asset menu](new_asset_blueprint_menu.png)
 	[/REGION]
 
-1. Choose a **Parent Class** for your _Blueprint_.
+1. Choose a **Parent Class** for your Blueprint.
 
 	![Choose a Parent Class](new_asset_parent_class.png)
 
@@ -75,7 +76,7 @@ Alternatively:
 	![](rtclick_content_browser_asset_tree.png)
 	[/REGION]
 
-1. Choose a **Parent Class** for your _Blueprint_.
+1. Choose a **Parent Class** for your Blueprint.
 
 	![Choose a Parent Class](new_asset_parent_class.png)
 
@@ -86,25 +87,25 @@ Alternatively:
 [/EXCERPT:CreateBP3]
 
 <!--
-1.	In the **Content Browser**, right-click and choose the option for the type of _Blueprint_ you wish to 
-create. In the example, a normal _Blueprint_ is created by choosing **New Blueprint**:
+1.	In the **Content Browser**, right-click and choose the option for the type of Blueprint you wish to 
+create. In the example, a normal Blueprint is created by choosing **New Blueprint**:
 
 ![Create Blueprint Menu Option](k2_create_menu.png)
 
 Or click ![Content Browser New Button](k2_create_new.png) and set the **Factory** to the type of 
-_Blueprint_ you wish to create. In the example, a normal _Blueprint_ is created by choosing **Blueprint**:
+Blueprint you wish to create. In the example, a normal Blueprint is created by choosing **Blueprint**:
 
 ![New Dialog Blueprint Factory Option](k2_create_factory.png)
 
-The options specific to creating a _Blueprint_ are displayed.
+The options specific to creating a Blueprint are displayed.
 
 ![New Dialog Blueprint Optionts](k2_create_options.png)
 
-1.	In the **New** dialog, enter a **Path** for the new _Blueprint_.
+1.	In the **New** dialog, enter a **Path** for the new Blueprint.
 
 ![New Dialog Info](k2_create_info.png)
 
-1.	Click ![New Dialog OK Button](k2_create_ok.png) to create the _Blueprint_. The new _Blueprint_ asset 
+1.	Click ![New Dialog OK Button](k2_create_ok.png) to create the Blueprint. The new Blueprint asset 
 is displayed in the **Content Browser**.
 
 ![Blueprint Asset Thumbnail](k2_create_asset.png)
@@ -112,7 +113,7 @@ is displayed in the **Content Browser**.
 
 ##Creation via Class Viewer
 
-When creating Blueprints via the [Class Viewer](Engine/UI/ClassViewer), it can be helpful to first filter the displayed classes to _Blueprint_ bases only. 
+When creating Blueprints via the [Class Viewer](Engine/UI/ClassViewer), it can be helpful to first filter the displayed classes to Blueprint bases only. 
 
 1. In the **Class Viewer** toolbar, click on **Filters**.
 
@@ -123,28 +124,28 @@ When creating Blueprints via the [Class Viewer](Engine/UI/ClassViewer), it can b
 	![](filter_blueprint_class_viewer.png)
 
 
-Once you have filtered the list of classes, navigate to the class you would like to use as a base for your _Blueprint_. In the example here, we are creating a _Blueprint_ with a **CameraActor** as a base.
+Once you have filtered the list of classes, navigate to the class you would like to use as a base for your Blueprint. In the example here, we are creating a Blueprint with a **CameraActor** as a base.
 
 1. Click on the down arrow icon to the right of your chosen class. 
 
 	![](class_viewer_arrow.png)
 
-1. The option to **Create Blueprint** will appear. Click on this, and the dialog for creating a _Blueprint_ will appear.
+1. The option to **Create Blueprint** will appear. Click on this, and the dialog for creating a Blueprint will appear.
 
 	![](class_viewer_down_arrow.png)
 
-1. Enter a name for your _Blueprint_, and select the folder where it should be saved.
+1. Enter a name for your Blueprint, and select the folder where it should be saved.
 
-1. Click on "Create [Path]/[Name]" at the top of the _Blueprint_ creation dialog, and your _Blueprint_ will be created and open in the **Blueprint Editor**.
+1. Click on "Create [Path]/[Name]" at the top of the Blueprint creation dialog, and your Blueprint will be created and open in the **Blueprint Editor**.
 	
 	![](class_viewer_create_complete.png)
 
-1. To save your _Blueprint_ in the Editor and complete the creation process, click on ![Save](blueprint_editor_save.png) in the **Blueprint Editor** toolbar.
+1. To save your Blueprint in the Editor and complete the creation process, click on ![Save](blueprint_editor_save.png) in the **Blueprint Editor** toolbar.
 
 
 
-You can also access _Blueprint_ creation in the **Class Viewer** by **right-clicking** directly on the class you would like to use as a base for your _Blueprint_. 
-The _Blueprint_ creation process is the same as if you used the ![Down Arrow Icon](button_down_arrow.png) to the right of the class.
+You can also access Blueprint creation in the **Class Viewer** by **right-clicking** directly on the class you would like to use as a base for your Blueprint. 
+The Blueprint creation process is the same as if you used the ![Down Arrow Icon](button_down_arrow.png) to the right of the class.
 
 ![](rtclick_class_viewer_bp.png)
 
@@ -152,27 +153,27 @@ The _Blueprint_ creation process is the same as if you used the ![Down Arrow Ico
 
 ##Creation from a Selected Actor
 
-1. Select the **Actor** (StaticMesh, Camera, etc.) that you would like to create a _Blueprint_ from in the Unreal Editor.
+1. Select the **Actor** (StaticMesh, Camera, etc.) that you would like to create a Blueprint from in the Unreal Editor.
 
-1. In the **Details** pane, under _Blueprint_, click on the **Create (Type) Blueprint** button. 
+1. In the **Details** pane, under Blueprint, click on the **Create (Type) Blueprint** button. 
 
 	![](actor_details_create_blueprint.png)
 
 	[REGION:note]
-	You can also click on the **Replace With Composited Blueprint** button to add the Actor(s) to a new _Blueprint_ and replace them in your level with the new _Blueprint_. 
+	You can also click on the **Replace With Composited Blueprint** button to add the Actor(s) to a new Blueprint and replace them in your level with the new Blueprint. 
 	[/REGION]
 
-1. Select the path where your _Blueprint_ should be saved.
+1. Select the path where your Blueprint should be saved.
 
 	![](actor_select_path_menu.png)
 
-1. Your _Blueprint_ will now appear in the **Content Browser**. Save it by clicking on **Save All** in the **Content Browser's** Navigation Bar.
+1. Your Blueprint will now appear in the **Content Browser**. Save it by clicking on **Save All** in the **Content Browser's** Navigation Bar.
 
 	[REGION:raw]
 	![Save in Content Browser](Content_browser_new_bp_save.png)
 	[/REGION]
 
-1. _Blueprints_ created directly from Actors have automatically generated names. To change the name, **Right-click** (or press **F2**) on the _Blueprint's_ icon in the **Content Browser**, and select **Rename** in the menu that appears.
+1. Blueprints created directly from Actors have automatically generated names. To change the name, **Right-click** (or press **F2**) on the Blueprint's icon in the **Content Browser**, and select **Rename** in the menu that appears.
 
 	![](rename_blueprint.png)
 
