#UE4doc #docs:
Added engine version field
Replaced multiple images due to irrelevance
Reworded multiple steps as they were incorrect
Removed italics from a few terms

diff --git a/Engine/Documentation/Source/Engine/Blueprints/UserGuide/Types/ClassBlueprint/Opening/OpeningClassBlueprints.INT.udn b/Engine/Documentation/Source/Engine/Blueprints/UserGuide/Types/ClassBlueprint/Opening/OpeningClassBlueprints.INT.udn
--- a/Engine/Documentation/Source/Engine/Blueprints/UserGuide/Types/ClassBlueprint/Opening/OpeningClassBlueprints.INT.udn
+++ b/Engine/Documentation/Source/Engine/Blueprints/UserGuide/Types/ClassBlueprint/Opening/OpeningClassBlueprints.INT.udn
@@ -2,30 +2,31 @@ Availability:Public
 Title:Opening Blueprint Classes
 Crumbs: %ROOT%, Engine/Blueprints, Engine/Blueprints/UserGuide/Types/ClassBlueprint
 Description:Guide to opening Blueprint Classes for editing from the Content Browser and the Level Editor
+version: 4.9
 
-You can open a _Blueprint Class_ in the **Blueprint Editor** either by interacting with the asset directly in the **Content Browser**, or by interacting with an
-instance of the _Blueprint Class_ that has been placed in your level.
+You can open a **Blueprint Class** in the **Blueprint Editor** either by interacting with the asset directly in the **Content Browser**, or by interacting with an
+instance of the Blueprint Class that has been placed in your level.
 
 ## Content Browser
 * **Double-click** on the icon in the **Content Browser**.
 
 	![](Content_browser_blueprints.png)
 
-* **Right-click** on the icon in the **Content Browser** and under **Blueprint Actions**, select **Open in Full Editor**.
+* **Right-click** on the icon in the **Content Browser** and under **Common**, select **Edit...**.
 	
 	![](Content_browser_open_bp.png)
 
 
 ## Level Viewport
-* **Right-click** on an instance of the _Blueprint_ in the level and under **Blueprints**, select **Edit Blueprint**.
+* **Right-click** on an instance of the Blueprint in the level and under **Asset**, select **Edit _Blueprint name_**.
 	
 	![](level_editor_open_blueprint.png)
 
-* With an instance of the _Blueprint_ selected, in the **Details** tab, click on the **Edit** button next to the text **[BlueprintName] extends [ParentClass]**. 
+* With an instance of the Blueprint selected, in the **Details** tab, click on the **Edit Blueprint** button. 
 
 	![](details_panel_edit_blueprint.png)
 
 [REGION:warning]
-Although these two methods involve menus attached to an individual instance of a _Blueprint_, please note that any changes made in the **Blueprint Editor** are made to the _Blueprint_ asset. 
-As a result, the changes will be applied to all instances of the _Blueprint_.
+Although these two methods involve menus attached to an individual instance of a Blueprint, please note that any changes made in the **Blueprint Editor** are made to the Blueprint asset. 
+As a result, the changes will be applied to all instances of the Blueprint.
 [/REGION]
