#UE4 Docs: Jira UEDOC-1935 (Animation Blueprint Override).

diff --git a/Engine/Documentation/Source/Engine/Animation/AnimHowTo/AnimBPOverride/AnimBPOverrideHowTo.INT.udn b/Engine/Documentation/Source/Engine/Animation/AnimHowTo/AnimBPOverride/AnimBPOverrideHowTo.INT.udn
new file mode 100644
--- /dev/null
+++ b/Engine/Documentation/Source/Engine/Animation/AnimHowTo/AnimBPOverride/AnimBPOverrideHowTo.INT.udn
@@ -0,0 +1,67 @@
+Availability: Public
+Title:Animation Blueprint Override
+Crumbs: %ROOT%, Engine, Engine\Animation, Engine\Animation\AnimHowTo
+Description:Illustrates how you can override animations in a child Animation Blueprint. 
+Related: Engine\Animation\Sequences
+Related: Engine\Animation\AnimBlueprints
+Related: Engine\Animation\Skeleton
+Related: Engine\Animation\Overview
+version: 4.9
+skilllevel: Advanced
+
+When setting up characters and animation for those characters, there may be instances where you want to have one character perform an animation for an action and have another character perform a different animation for the same action. By using the **Asset Override Editor** inside of a child [Animation Blueprint](Engine/Animation/AnimBlueprints), you can override any previously set animation asset used in a [Sequence Player Node](Engine/Animation/Sequences/UserGuide) allowing you to quickly create variations of Animation Blueprints that each of your characters can use. 
+
+
+
+[REGION:note]
+For this guide, we are using the **Blueprint Third Person** template and included the **Infinity Blade: Warriors** and **Animation Starter Pack** assets which are available for free via the Marketplace. We also performed some [Animation Retargeting](Engine/Animation/AnimHowTo/Retargeting) and retargeted the Animation Blueprint and Animations that came with the Animation Starter Pack for use with the Infinity Blade characters. 
+[/REGION]
+
+## Steps
+
+1. **Right-click** on the **Animation Blueprint** you would like to override animations for and select **Create Child Blueprint Class**. 
+
+	![](01_CreateChild.png)
+
+1. Open the child Animation Blueprint, then click **Window** from the **File Menu** and select **Asset Override Editor**. 
+
+	![](02_OpenEditor.png)
+
+1. The **Asset Override Editor** will open showing the available animations which can be overridden. 
+
+	![](03_AssetOverrideEditor.png)
+
+	You can expand/collapse the display of assets by clicking the arrow next to each asset. 
+
+1. In the **Asset** column, click the drop-down window and specify the new asset to use (which will override the existing when called). 
+
+	![](04_SelectAssetToOverride.png)
+
+	Clicking the eyeball icon will allow you to see a preview of the node context in the parent graph as read only in the main graph panel. 
+
+	![](05_JumpToNode.png)
+
+1. You can view your changes after you click **Compile**. That's it! 
+
+
+## End Result
+
+Below is an example where a character's default movement (assigned as Jog) in the parent Animation Blueprint is overridden so that the character performs a Walk instead. 
+
+
+[OBJECT:EmbeddedVideo]
+[PARAMLITERAL:width]
+640
+[/PARAMLITERAL]
+[PARAMLITERAL:height]
+360
+[/PARAMLITERAL]
+[PARAMLITERAL:videoid]
+MuQ5VQIGfY4
+[/PARAMLITERAL]
+[/OBJECT]
+
+A prime usage case for this is if you have a character that performs a certain attack when a button is pressed but you want a different character to have the same general movement but perform a different motion when the attack button is pressed. 
+
+
+
diff --git a/Engine/Documentation/Source/Engine/Animation/AnimHowTo/AnimHowToOverview.INT.udn b/Engine/Documentation/Source/Engine/Animation/AnimHowTo/AnimHowToOverview.INT.udn
--- a/Engine/Documentation/Source/Engine/Animation/AnimHowTo/AnimHowToOverview.INT.udn
+++ b/Engine/Documentation/Source/Engine/Animation/AnimHowTo/AnimHowToOverview.INT.udn
@@ -28,6 +28,7 @@ Refer to the tables below for more information and click each link to view its c
 | ------- | ------------ |
 | [](Engine\Animation\AnimHowTo\AdditiveAnimations) | %Engine\Animation\AnimHowTo\AdditiveAnimations:description% |
 | [](Engine\Animation\AnimHowTo\LayerEditing) | %Engine\Animation\AnimHowTo\LayerEditing:description% |
+| [](Engine\Animation\AnimHowTo\AnimBPOverride) | %Engine\Animation\AnimHowTo\AnimBPOverride:description% |
 
 
 <!-- Convert these to Quick Starts
@@ -40,7 +41,3 @@ Split this into two How To's
 
 
 
-
-
-
-
