#UE4doc #docs:
Resubmitting as not all changes were saved

diff --git a/Engine/Documentation/Source/Engine/Blueprints/UserGuide/OnlineNodes/OnlineNodes.INT.udn b/Engine/Documentation/Source/Engine/Blueprints/UserGuide/OnlineNodes/OnlineNodes.INT.udn
--- a/Engine/Documentation/Source/Engine/Blueprints/UserGuide/OnlineNodes/OnlineNodes.INT.udn
+++ b/Engine/Documentation/Source/Engine/Blueprints/UserGuide/OnlineNodes/OnlineNodes.INT.udn
@@ -6,6 +6,7 @@ Related: Engine/Blueprints
 Related: Programming/Online
 Related: Gameplay/Networking/Actors
 Related: Resources/ContentExamples/Networking
+version: 4.9
 
 
 [VAR:Topic]
