#UE4 Docs: Jira UEDOC-2314 (Non-Linear Blends Feature Concept).

diff --git a/Engine/Documentation/Source/Engine/Animation/Animation.INT.udn b/Engine/Documentation/Source/Engine/Animation/Animation.INT.udn
--- a/Engine/Documentation/Source/Engine/Animation/Animation.INT.udn
+++ b/Engine/Documentation/Source/Engine/Animation/Animation.INT.udn
@@ -104,6 +104,7 @@ Navigation:topic
 		* [](Engine/Animation/PhysicallyDrivenAnimation "%Engine/Animation/PhysicallyDrivenAnimation:description%")
 		* [](Engine/Animation/RootMotion "%Engine/Animation/RootMotion:description%")
 		* [](Engine/Animation/NonUniformScale "%Engine/Animation/NonUniformScale:description%")
+		* [](Engine\Animation\NonLinearBlends "%Engine\Animation\NonLinearBlends:description%")
 	[/PARAM]
 [/OBJECT]
 
diff --git a/Engine/Documentation/Source/Engine/Animation/NonLinearBlends/NonLinearBlends.INT.udn b/Engine/Documentation/Source/Engine/Animation/NonLinearBlends/NonLinearBlends.INT.udn
new file mode 100644
--- /dev/null
+++ b/Engine/Documentation/Source/Engine/Animation/NonLinearBlends/NonLinearBlends.INT.udn
@@ -0,0 +1,51 @@
+Availability:Public
+Title:Non-Linear Blends
+Crumbs: %ROOT%, Engine, Engine/Animation
+Description:Overview of the different blend functions between animation poses and transitions.  
+version: 4.9
+Related: Engine\Animation\StateMachines\TransitionRules
+Related: Engine\Animation\StateMachines
+Related: Engine/UI/CurveEditor
+Related: Engine/Basics/Distributions
+Related: Engine/Animation/AnimMontage
+
+Unreal Engine 4 supports many different blending functions that you can use when blending between animation poses and transitions. 
+
+[OBJECT:EmbeddedVideo]
+[PARAMLITERAL:width]
+640
+[/PARAMLITERAL]
+[PARAMLITERAL:height]
+360
+[/PARAMLITERAL]
+[PARAMLITERAL:videoid]
+-WeUEpZd8jY
+[/PARAMLITERAL]
+[/OBJECT]
+
+_An example of a Custom curve used to Blend between different poses._
+
+You can select from many different blending functions: Linear, Cubic, Hermite Cubic, Sinusoidal, Quadratic, Circular, Exponential and Custom. For most types, you can separately control whether you want easing on either the _in_ or the _out_ of the curve.
+
+The **Custom** option requires a **Curve Float** asset. After creating one, you can assign it in the **Details** panel for the relevant nodes:
+
+![](image_44.png)
+
+![](image_45.png)
+
+This will allow you to specify just about any blend you want to perform. The length of the curve will be normalized and scaled down to meet the specified blend time and any values that are outside of the range 0-1 are clamped back to fit. If a Custom curve is not specified the system will fall back to a Linear blend as depicted below.
+
+[OBJECT:EmbeddedVideo]
+[PARAMLITERAL:width]
+640
+[/PARAMLITERAL]
+[PARAMLITERAL:height]
+360
+[/PARAMLITERAL]
+[PARAMLITERAL:videoid]
+GHxhplIScUY
+[/PARAMLITERAL]
+[/OBJECT]
+
+
+
diff --git a/Engine/Documentation/Source/Engine/Animation/StateMachines/TransitionRules/TransitionRules.INT.udn b/Engine/Documentation/Source/Engine/Animation/StateMachines/TransitionRules/TransitionRules.INT.udn
--- a/Engine/Documentation/Source/Engine/Animation/StateMachines/TransitionRules/TransitionRules.INT.udn
+++ b/Engine/Documentation/Source/Engine/Animation/StateMachines/TransitionRules/TransitionRules.INT.udn
@@ -38,7 +38,7 @@ When **Left-clicking** on the Transition Rule, the **Details** panel will update
 | **Transition Interrupt** | Set the Trigger Weight Threshold and specify a Notify State Class.  |
 | **Transition Crossfade Sharing** | Use an existing Transition Crossfade Blend or promote the existing Transition Crossfade to shared.  |
 | **Duration** | The duration to apply to the cross-fade. |
-| **Mode** | The type of blending applied to the cross-fade. |
+| **Mode** | The type of blending applied to the cross-fade. See [](Engine\Animation\NonLinearBlends) for more information. |
 | **Start Transition Event / Custom Blueprint Event** | Call and execute an Event from within the notified Blueprint at the start of the transition. |
 | **Blueprint Notify** | Blueprint that is notified when the transition starts.  |
 | **End Transition Event / Custom Blueprint Event** | Call and execute an Event from within the notified Blueprint at the end of the transition. |
