#UE4doc #docs - Added Engine Version Field. Updated images to match latest UI changes.

diff --git a/Engine/Documentation/Source/Engine/UI/LevelEditor/InEditorTesting/InEditorTesting.INT.udn b/Engine/Documentation/Source/Engine/UI/LevelEditor/InEditorTesting/InEditorTesting.INT.udn
--- a/Engine/Documentation/Source/Engine/UI/LevelEditor/InEditorTesting/InEditorTesting.INT.udn
+++ b/Engine/Documentation/Source/Engine/UI/LevelEditor/InEditorTesting/InEditorTesting.INT.udn
@@ -2,6 +2,7 @@ Availability:Public
 Title:In-Editor Testing (Play & Simulate)
 Crumbs: %ROOT%, Engine, Engine/UI, Engine/UI/LevelEditor
 Description: Guide to testing levels and gameplay directly in the editor.
+Version: 4.9
 
 [VAR:playlistnav]
 [OBJECT:PlaylistNavigation]
