#UE4doc #docs - Added Engine Version Field.

diff --git a/Engine/Documentation/Source/Engine/UI/LevelEditor/Modes/MeshPaintMode/VertexColor/BrushSettings/BrushSettings.INT.udn b/Engine/Documentation/Source/Engine/UI/LevelEditor/Modes/MeshPaintMode/VertexColor/BrushSettings/BrushSettings.INT.udn
--- a/Engine/Documentation/Source/Engine/UI/LevelEditor/Modes/MeshPaintMode/VertexColor/BrushSettings/BrushSettings.INT.udn
+++ b/Engine/Documentation/Source/Engine/UI/LevelEditor/Modes/MeshPaintMode/VertexColor/BrushSettings/BrushSettings.INT.udn
@@ -2,6 +2,7 @@ Availability: Public
 Title:Brush Settings
 Crumbs:%ROOT%, Engine, Engine/UI/LevelEditor/Modes/MeshPaintMode/VertexColor
 Description:Documentation for the brush settings of the Mesh Paint tool.
+Version: 4.9
 
 [EXCERPT:Intro]
 The Mesh Paint tool is a brush-based system. Using the included settings, you can establish a brush designed for wide areas or fine details. It is important to keep in mind, however, that as the nature of the tool is to paint color onto actual vertices, you will often need to make settings adjustments based on the vertex density of your mesh.
