#UE4doc #docs - Added Engine version field. Updated images to match latest UI changes.

diff --git a/Engine/Documentation/Source/Engine/UI/LevelEditor/InEditorTesting/NetworkedPlayInEditor/NetworkedPlayInEditor.INT.udn b/Engine/Documentation/Source/Engine/UI/LevelEditor/InEditorTesting/NetworkedPlayInEditor/NetworkedPlayInEditor.INT.udn
--- a/Engine/Documentation/Source/Engine/UI/LevelEditor/InEditorTesting/NetworkedPlayInEditor/NetworkedPlayInEditor.INT.udn
+++ b/Engine/Documentation/Source/Engine/UI/LevelEditor/InEditorTesting/NetworkedPlayInEditor/NetworkedPlayInEditor.INT.udn
@@ -2,6 +2,7 @@ Availability:Public
 Title:Networked Play In Editor
 Crumbs: %ROOT%, Engine, Engine/UI, Engine/UI/LevelEditor, Engine/UI/LevelEditor/InEditorTesting
 Description:Testing gameplay with clients and servers
+Version: 4.9
 
 
 _Play In Editor_ supports multiple worlds, so that you can test networked play. The editor spawns multiple worlds within the already loaded instance of UE4, which results in fast iteration time and 
