#UE4doc #docs - Added Engine Version Field. Updated Images to match latest UI changes.

diff --git a/Engine/Documentation/Source/Engine/UI/LevelEditor/InEditorTesting/Settings/PIE_SIE_Settings.INT.udn b/Engine/Documentation/Source/Engine/UI/LevelEditor/InEditorTesting/Settings/PIE_SIE_Settings.INT.udn
--- a/Engine/Documentation/Source/Engine/UI/LevelEditor/InEditorTesting/Settings/PIE_SIE_Settings.INT.udn
+++ b/Engine/Documentation/Source/Engine/UI/LevelEditor/InEditorTesting/Settings/PIE_SIE_Settings.INT.udn
@@ -2,6 +2,7 @@ Availability:Public
 Title:Play In Editor Settings
 Crumbs: %ROOT%, Engine, Engine/UI, Engine/UI/LevelEditor, Engine/UI/LevelEditor/InEditorTesting
 Description:Set up window sizes and other options for the Play-In feature
+Version: 4.9
 
 [TOC(start:2)]
 The **Play** settings panel can be opened two different ways:
