#UE4doc #docs - Added Engine Version Field. Updated image to match latest UI changes.

diff --git a/Engine/Documentation/Source/Engine/UI/LevelEditor/Modes/LevelEditorTools.INT.udn b/Engine/Documentation/Source/Engine/UI/LevelEditor/Modes/LevelEditorTools.INT.udn
--- a/Engine/Documentation/Source/Engine/UI/LevelEditor/Modes/LevelEditorTools.INT.udn
+++ b/Engine/Documentation/Source/Engine/UI/LevelEditor/Modes/LevelEditorTools.INT.udn
@@ -2,6 +2,7 @@ Availability:Public
 Title:Level Editor Modes
 Crumbs: %ROOT%, Engine, Engine/UI, Engine/UI/LevelEditor
 Description:Set of tools used to control the mode the Level Editor is currently in, reshape the builder brush, and create new Brush geometry and volumes.
+Version: 4.9
 
 [VAR:playlistnav]
 [OBJECT:PlaylistNavigation]
