#UE4doc #doc - Added engine version field.

diff --git a/Engine/Documentation/Source/Engine/UI/LevelEditor/Modes/MeshPaintMode/VertexColor/VertexColor.INT.udn b/Engine/Documentation/Source/Engine/UI/LevelEditor/Modes/MeshPaintMode/VertexColor/VertexColor.INT.udn
--- a/Engine/Documentation/Source/Engine/UI/LevelEditor/Modes/MeshPaintMode/VertexColor/VertexColor.INT.udn
+++ b/Engine/Documentation/Source/Engine/UI/LevelEditor/Modes/MeshPaintMode/VertexColor/VertexColor.INT.udn
@@ -2,6 +2,7 @@ Availability: Public
 Title:Mesh Paint User Guide
 Crumbs:%ROOT%, Engine
 Description:Tool for painting vertex colors on meshes interactively in the editor viewport.
+Version: 4.9
 
 **Mesh Paint Mode** provides a quick and easy way of adjusting color and texture on your meshes. The following sections cover all of the critical skills you will need to make the most out of this powerful tool.
 
