#UE4doc #docs - Added Engine Version Field.

diff --git a/Engine/Documentation/Source/Engine/UI/LevelEditor/Modes/MeshPaintMode/VertexColor/AssetVsInstance/AssetVsInstance.INT.udn b/Engine/Documentation/Source/Engine/UI/LevelEditor/Modes/MeshPaintMode/VertexColor/AssetVsInstance/AssetVsInstance.INT.udn
--- a/Engine/Documentation/Source/Engine/UI/LevelEditor/Modes/MeshPaintMode/VertexColor/AssetVsInstance/AssetVsInstance.INT.udn
+++ b/Engine/Documentation/Source/Engine/UI/LevelEditor/Modes/MeshPaintMode/VertexColor/AssetVsInstance/AssetVsInstance.INT.udn
@@ -2,6 +2,7 @@ Availability: Public
 Title:Asset vs. Instance
 Crumbs:%ROOT%, Engine, Engine/UI/LevelEditor/Modes/MeshPaintMode/VertexColor
 Description:A look at the difference between painting on assets and painting on an individual instance.
+Version: 4.9
 
 When painting Vertex Colors, you are only painting on an individual **Instance** of the material and not the actual material asset itself. 
 If you want to apply the vertex color data to another instance of the mesh, you can copy and paste the data onto another mesh or multiple meshes. 
