#UE4doc #docs - Added Engine Version Field.

diff --git a/Engine/Documentation/Source/Engine/UI/LevelEditor/Modes/MeshPaintMode/VertexColor/ColorsVsBlendWeights/ColorsVsBlendWeights.INT.udn b/Engine/Documentation/Source/Engine/UI/LevelEditor/Modes/MeshPaintMode/VertexColor/ColorsVsBlendWeights/ColorsVsBlendWeights.INT.udn
--- a/Engine/Documentation/Source/Engine/UI/LevelEditor/Modes/MeshPaintMode/VertexColor/ColorsVsBlendWeights/ColorsVsBlendWeights.INT.udn
+++ b/Engine/Documentation/Source/Engine/UI/LevelEditor/Modes/MeshPaintMode/VertexColor/ColorsVsBlendWeights/ColorsVsBlendWeights.INT.udn
@@ -2,6 +2,7 @@ Availability: Public
 Title:Colors vs. Blend Weights
 Crumbs:%ROOT%, Engine, Engine/UI/LevelEditor/Modes/MeshPaintMode/VertexColor
 Description:An overview of painting colors onto a mesh vs. painting texture blending weights.
+Version: 4.9
 
 
 
