#UE4doc #docs - Added Engine Version Field.

diff --git a/Engine/Documentation/Source/Engine/UI/LevelEditor/Modes/MeshPaintMode/VertexColor/MaterialSetup/MaterialSetup.INT.udn b/Engine/Documentation/Source/Engine/UI/LevelEditor/Modes/MeshPaintMode/VertexColor/MaterialSetup/MaterialSetup.INT.udn
--- a/Engine/Documentation/Source/Engine/UI/LevelEditor/Modes/MeshPaintMode/VertexColor/MaterialSetup/MaterialSetup.INT.udn
+++ b/Engine/Documentation/Source/Engine/UI/LevelEditor/Modes/MeshPaintMode/VertexColor/MaterialSetup/MaterialSetup.INT.udn
@@ -2,6 +2,6 @@ Availability: Public
 Title:Material Setup
 Crumbs:%ROOT%, Engine, Engine/UI/LevelEditor/Modes/MeshPaintMode/VertexColor
 Description:Examples for various material setups for utilizing Mesh Paint Mode's Blend Weight system.
-
+Version: 4.9
 
 [INCLUDE:Engine/UI/LevelEditor/Modes/MeshPaintMode/VertexColor#MaterialSetup]
