#UE4doc #docs - Added Engine Version Field. Added Skill Level field.

diff --git a/Engine/Documentation/Source/Engine/UI/LevelEditor/Modes/MeshPaintMode/VertexColor/MaterialSetup/2Way/2Way.INT.udn b/Engine/Documentation/Source/Engine/UI/LevelEditor/Modes/MeshPaintMode/VertexColor/MaterialSetup/2Way/2Way.INT.udn
--- a/Engine/Documentation/Source/Engine/UI/LevelEditor/Modes/MeshPaintMode/VertexColor/MaterialSetup/2Way/2Way.INT.udn
+++ b/Engine/Documentation/Source/Engine/UI/LevelEditor/Modes/MeshPaintMode/VertexColor/MaterialSetup/2Way/2Way.INT.udn
@@ -2,7 +2,8 @@ Availability: Public
 Title:2-Way Blend Materials
 Crumbs:%ROOT%, Engine, Engine/UI/LevelEditor/Modes/MeshPaintMode/VertexColor, Engine/UI/LevelEditor/Modes/MeshPaintMode/VertexColor/MaterialSetup
 Description:Guide to setting up materials for 2-way (alpha) texture blending.
-
+Version: 4.9
+SkillLevel: Intermediate
 
 Here is an example of how you might use the vertex alpha value to blend between two diffuse textures in the material editor.
 
