#UE4doc #docs - Added Engine Version Field. Added Skill Level field.

diff --git a/Engine/Documentation/Source/Engine/UI/LevelEditor/Modes/MeshPaintMode/VertexColor/MaterialSetup/3Way/3Way.INT.udn b/Engine/Documentation/Source/Engine/UI/LevelEditor/Modes/MeshPaintMode/VertexColor/MaterialSetup/3Way/3Way.INT.udn
--- a/Engine/Documentation/Source/Engine/UI/LevelEditor/Modes/MeshPaintMode/VertexColor/MaterialSetup/3Way/3Way.INT.udn
+++ b/Engine/Documentation/Source/Engine/UI/LevelEditor/Modes/MeshPaintMode/VertexColor/MaterialSetup/3Way/3Way.INT.udn
@@ -2,7 +2,8 @@ Availability: Public
 Title:3-Way Blend Materials
 Crumbs:%ROOT%, Engine, Engine/UI/LevelEditor/Modes/MeshPaintMode/VertexColor, Engine/UI/LevelEditor/Modes/MeshPaintMode/VertexColor/MaterialSetup
 Description:Guide to setting up materials for 3-way (RGB) texture blending.
-
+Version: 4.9
+SkillLevel: Intermediate
 
 
  Blending 3 separate textures can be handled by the Mesh Paint tool very easily, but it does require a special setup in the material. Consider the following:
