#UE4doc #docs -  Added Engine Version Field. Added Skill Level field.

diff --git a/Engine/Documentation/Source/Engine/UI/LevelEditor/Modes/MeshPaintMode/VertexColor/MaterialSetup/4Way/4Way.INT.udn b/Engine/Documentation/Source/Engine/UI/LevelEditor/Modes/MeshPaintMode/VertexColor/MaterialSetup/4Way/4Way.INT.udn
--- a/Engine/Documentation/Source/Engine/UI/LevelEditor/Modes/MeshPaintMode/VertexColor/MaterialSetup/4Way/4Way.INT.udn
+++ b/Engine/Documentation/Source/Engine/UI/LevelEditor/Modes/MeshPaintMode/VertexColor/MaterialSetup/4Way/4Way.INT.udn
@@ -2,6 +2,8 @@ Availability: Public
 Title:4-Way Blend Materials
 Crumbs:%ROOT%, Engine, Engine/UI/LevelEditor/Modes/MeshPaintMode/VertexColor, Engine/UI/LevelEditor/Modes/MeshPaintMode/VertexColor/MaterialSetup
 Description:Guide to setting up materials for 4-way (ARGB) texture blending.
+Version: 4.9
+SkillLevel: Intermediate
 
 
 
