#UE4doc #docs -  Added Engine Version Field. Added Skill Level field.

diff --git a/Engine/Documentation/Source/Engine/UI/LevelEditor/Modes/MeshPaintMode/VertexColor/MaterialSetup/5Way/5Way.INT.udn b/Engine/Documentation/Source/Engine/UI/LevelEditor/Modes/MeshPaintMode/VertexColor/MaterialSetup/5Way/5Way.INT.udn
--- a/Engine/Documentation/Source/Engine/UI/LevelEditor/Modes/MeshPaintMode/VertexColor/MaterialSetup/5Way/5Way.INT.udn
+++ b/Engine/Documentation/Source/Engine/UI/LevelEditor/Modes/MeshPaintMode/VertexColor/MaterialSetup/5Way/5Way.INT.udn
@@ -2,7 +2,8 @@ Availability: Public
 Title:5-Way Blend Materials
 Crumbs:%ROOT%, Engine, Engine/UI/LevelEditor/Modes/MeshPaintMode/VertexColor, Engine/UI/LevelEditor/Modes/MeshPaintMode/VertexColor/MaterialSetup
 Description:Guide to setting up materials for 5-way (1-ARGB) texture blending.
-
+Version: 4.9
+SkillLevel: Intermediate
 
 
 To handle 5 textures in the blend weights setup, you need to create a material similar to this:
