#UE4doc #docs -  Added Engine Version Field. Added Skill Level field.

diff --git a/Engine/Documentation/Source/Engine/UI/LevelEditor/Modes/MeshPaintMode/VertexColor/MaterialSetup/Color/Color.INT.udn b/Engine/Documentation/Source/Engine/UI/LevelEditor/Modes/MeshPaintMode/VertexColor/MaterialSetup/Color/Color.INT.udn
--- a/Engine/Documentation/Source/Engine/UI/LevelEditor/Modes/MeshPaintMode/VertexColor/MaterialSetup/Color/Color.INT.udn
+++ b/Engine/Documentation/Source/Engine/UI/LevelEditor/Modes/MeshPaintMode/VertexColor/MaterialSetup/Color/Color.INT.udn
@@ -2,7 +2,8 @@ Availability: Public
 Title:Vertex Color Materials
 Crumbs:%ROOT%, Engine, Engine/UI/LevelEditor/Modes/MeshPaintMode/VertexColor, Engine/UI/LevelEditor/Modes/MeshPaintMode/VertexColor/MaterialSetup
 Description:How to set up materials that utilize Vertex Color so that you can paint colors in Mesh Paint Mode.
-
+Version: 4.9
+SkillLevel: Intermediate
 
 There are really an infinite number of ways one could use the colors from a Vertex Color material expression. For this setup, we are going to keep things relatively simple and just multiply the colors into an existing material. This allows us to have the original texture in place, but to tint it with the Mesh Paint tool.
 
