#UE4doc #docs -  Added Engine Version Field.

diff --git a/Engine/Documentation/Source/Engine/UI/LevelEditor/Modes/MeshPaintMode/VertexColor/PaintingColors/PaintingColors.INT.udn b/Engine/Documentation/Source/Engine/UI/LevelEditor/Modes/MeshPaintMode/VertexColor/PaintingColors/PaintingColors.INT.udn
--- a/Engine/Documentation/Source/Engine/UI/LevelEditor/Modes/MeshPaintMode/VertexColor/PaintingColors/PaintingColors.INT.udn
+++ b/Engine/Documentation/Source/Engine/UI/LevelEditor/Modes/MeshPaintMode/VertexColor/PaintingColors/PaintingColors.INT.udn
@@ -2,7 +2,7 @@ Availability: Public
 Title:Painting Colors
 Crumbs:%ROOT%, Engine, Engine/UI/LevelEditor/Modes/MeshPaintMode/VertexColor
 Description:How to go about painting colors onto your meshes using Mesh Paint Mode.
-
+Version: 4.9
 
 
 Mesh Paint has two different modes: **Color**, which paints RGBA data directly onto a mesh's vertices, and **Blend Weights**, which serve as a means to blend between different textures. The default Mesh Paint mode is Color.  You can select a paint and erase color and then apply those colors to the mesh using the brush.  This mode is useful when your material is configured to combine the vertex color with your pixel shader in some interesting way.
