#UE4doc #docs - Added engine version field. Added Skill Level field.

diff --git a/Engine/Documentation/Source/Engine/UI/LevelEditor/Modes/MeshPaintMode/VertexColor/PaintingOnMeshes/PaintingOnMeshes.INT.udn b/Engine/Documentation/Source/Engine/UI/LevelEditor/Modes/MeshPaintMode/VertexColor/PaintingOnMeshes/PaintingOnMeshes.INT.udn
--- a/Engine/Documentation/Source/Engine/UI/LevelEditor/Modes/MeshPaintMode/VertexColor/PaintingOnMeshes/PaintingOnMeshes.INT.udn
+++ b/Engine/Documentation/Source/Engine/UI/LevelEditor/Modes/MeshPaintMode/VertexColor/PaintingOnMeshes/PaintingOnMeshes.INT.udn
@@ -2,6 +2,8 @@ Availability: Public
 Title:Painting on Meshes
 Crumbs:%ROOT%, Engine, Engine/UI/LevelEditor/Modes/MeshPaintMode/VertexColor
 Description:An overview of how to lay paint strokes down on meshes using Mesh Paint Mode.
+Version: 4.9
+SkillLevel: Intermediate
 
 [EXCERPT:Intro]
 Painting on meshes is easy! The first thing to remember is that you can only paint on selected meshes.  This helps you restrict painting to only specific objects, while still allowing you to paint across multiple selected _Mesh Actors_ if necessary.
@@ -20,7 +22,7 @@ After selecting one or more meshes, when you move the mouse over them, you shoul
 You can think of the paint brush as a cylinder projected along the surface normal of the mesh geometry, centered at the point on the mesh under your cursor.  
 The circular brush outline shows you the radius of the brush.  The inner circle shows the brush's _falloff_ (or inner radius.)  You probably figured this out already, but the small line sticking out from center of the brush is the mesh surface normal.
 
-To paint, hold down the **Left Mouse Button** and drag over the mesh's surface.  You can also hold down the **Shift key** while painting to erase instead.
+To paint, hold down the **Left Mouse Button** and drag over the mesh's surface.  You can also hold down the **Shift** key while painting to erase instead.
 
 Paint will be applied to the surface each time you click, and also each time the mouse position changes by dragging.  Also, if **Brush Flow** is enabled, paint will be applied every time the scene is rendered.
 
