#UE4doc #docs - Added Engine Version Field.

diff --git a/Engine/Documentation/Source/Engine/UI/LevelEditor/Modes/MeshPaintMode/VertexColor/QuickGuide/QuickGuide.INT.udn b/Engine/Documentation/Source/Engine/UI/LevelEditor/Modes/MeshPaintMode/VertexColor/QuickGuide/QuickGuide.INT.udn
--- a/Engine/Documentation/Source/Engine/UI/LevelEditor/Modes/MeshPaintMode/VertexColor/QuickGuide/QuickGuide.INT.udn
+++ b/Engine/Documentation/Source/Engine/UI/LevelEditor/Modes/MeshPaintMode/VertexColor/QuickGuide/QuickGuide.INT.udn
@@ -2,7 +2,7 @@ Availability: Public
 Title:Mesh Paint Overview
 Crumbs:%ROOT%, Engine, Engine/UI/LevelEditor/Modes/MeshPaintMode/VertexColor
 Description:A down-and-dirty quick guide to getting your feet wet with Mesh Paint Mode.
-
+Version: 4.9
 
 [REGION:headerimage]
 [REGION:fullwidth]
