#UE4doc #docs - Added engine version field.

diff --git a/Engine/Documentation/Source/Engine/UI/LevelEditor/Modes/MeshPaintMode/VertexColor/SelectingColors/SelectingColors.INT.udn b/Engine/Documentation/Source/Engine/UI/LevelEditor/Modes/MeshPaintMode/VertexColor/SelectingColors/SelectingColors.INT.udn
--- a/Engine/Documentation/Source/Engine/UI/LevelEditor/Modes/MeshPaintMode/VertexColor/SelectingColors/SelectingColors.INT.udn
+++ b/Engine/Documentation/Source/Engine/UI/LevelEditor/Modes/MeshPaintMode/VertexColor/SelectingColors/SelectingColors.INT.udn
@@ -2,6 +2,7 @@ Availability: Public
 Title:Selecting Colors
 Crumbs:%ROOT%, Engine, Engine/UI/LevelEditor/Modes/MeshPaintMode/VertexColor
 Description:An overiview of the process of color selection for Mesh Paint Mode.
+Version: 4.9
 
 
 [EXCERPT:Intro]
