#Docfix

Fixed DirectX Version confusion issues on the 'Recommended Specifications' page

diff --git a/Engine/Documentation/Source/Engine/Rendering/ParticleSystems/Reference/TypeData/GPUSprites/GPUSpritesTypeData.INT.udn b/Engine/Documentation/Source/Engine/Rendering/ParticleSystems/Reference/TypeData/GPUSprites/GPUSpritesTypeData.INT.udn
--- a/Engine/Documentation/Source/Engine/Rendering/ParticleSystems/Reference/TypeData/GPUSprites/GPUSpritesTypeData.INT.udn
+++ b/Engine/Documentation/Source/Engine/Rendering/ParticleSystems/Reference/TypeData/GPUSprites/GPUSpritesTypeData.INT.udn
@@ -13,7 +13,7 @@ The **GPUSprite** type data module supports simulating particles on the GPU. Tra
 
 GPU particles do not support all of the features available to traditional CPU particles, but they offer greater efficiency as well as a few unique features. Emission of particles still happens on the CPU to determine where, when, and how many particles to spawn. The CPU may also assign initial attributes such as size and velocity using methods available to traditional CPU particles.
 
-GPU particle sprites are authored in Cascade by adding the "GPU Sprites" type data module to an emitter. As with other effects authored in Cascade, the parameters of the simulation are modified by adding and editing modules to the emitter. The simulation updates in real time as properties change. 
+GPU particle sprites are authored in Cascade by adding the "GPU Sprites" type data module to an emitter. As with other effects authored in Cascade, the parameters of the simulation are modified by adding and editing modules on the emitter. The simulation updates in real time as properties change. 
 
 
 ## Properties
diff --git a/Engine/Documentation/Source/GettingStarted/RecommendedSpecifications/Specifications.INT.udn b/Engine/Documentation/Source/GettingStarted/RecommendedSpecifications/Specifications.INT.udn
--- a/Engine/Documentation/Source/GettingStarted/RecommendedSpecifications/Specifications.INT.udn
+++ b/Engine/Documentation/Source/GettingStarted/RecommendedSpecifications/Specifications.INT.udn
@@ -70,8 +70,7 @@ Unreal Engine has some specific hardware and software requirements for running t
 | **Operating System** | Windows 7/8 64-bit |
 | **Processor** | Quad-core Intel or AMD, 2.5 GHz or faster |
 | **Memory** | 8 GB RAM |
-| **Video Card** | DirectX 11 compatible graphics card |
-| **DirectX Version** | [DirectX End-User Runtimes (June 2010)](http://www.microsoft.com/en-us/download/details.aspx?id=8109) |
+| **Video Card/DirectX Version** | DirectX 11 compatible graphics card |
 [/EXCERPT:WindowsMinHardware]
 [EXCERPT:MacOSMinHardware]
 |||
@@ -196,7 +195,7 @@ Below are the requirements for running the engine or editor.
 | Running the Engine ||
 | -------- | --- |
 | **Operating System** | Windows 7/8 64-bit |
-| **DirectX Version** | [DirectX End-User Runtimes (June 2010)](http://www.microsoft.com/en-us/download/details.aspx?id=8109) |
+| **DirectX Runtime** | [DirectX End-User Runtimes (June 2010)](http://www.microsoft.com/en-us/download/details.aspx?id=8109) |
 
 [/EXCERPT:WindowsMinSoftware]
 
