#UE4doc #docs - Added Engine Version Field. Added skill level field.

diff --git a/Engine/Documentation/Source/Engine/UI/LevelEditor/Modes/MeshPaintMode/VertexColor/SharingColorData/SharingColorData.INT.udn b/Engine/Documentation/Source/Engine/UI/LevelEditor/Modes/MeshPaintMode/VertexColor/SharingColorData/SharingColorData.INT.udn
--- a/Engine/Documentation/Source/Engine/UI/LevelEditor/Modes/MeshPaintMode/VertexColor/SharingColorData/SharingColorData.INT.udn
+++ b/Engine/Documentation/Source/Engine/UI/LevelEditor/Modes/MeshPaintMode/VertexColor/SharingColorData/SharingColorData.INT.udn
@@ -2,7 +2,8 @@ Availability: Public
 Title:Sharing Color Data
 Crumbs:%ROOT%, Engine, Engine/UI/LevelEditor/Modes/MeshPaintMode/VertexColor
 Description:How to go about transferring vertex color data across meshes using Mesh Paint Mode.
-
+Version: 4.9
+SkillLevel: Intermediate
 
 
 There may be cases where you want to paint a mesh and then share those vertex colors with other instances of the mesh in the same level.  For example, if you had a pillar mesh you were using in an outdoor map and you wanted to paint the base of the pillar green, and have all instances in that map only appear that way.
