#UE4doc #docs - Added Engine version field. Replaced image table with comparison slider.

diff --git a/Engine/Documentation/Source/Engine/UI/LevelEditor/Modes/MeshPaintMode/VertexColor/VertexColorMatching/VertexColorMatching.INT.udn b/Engine/Documentation/Source/Engine/UI/LevelEditor/Modes/MeshPaintMode/VertexColor/VertexColorMatching/VertexColorMatching.INT.udn
--- a/Engine/Documentation/Source/Engine/UI/LevelEditor/Modes/MeshPaintMode/VertexColor/VertexColorMatching/VertexColorMatching.INT.udn
+++ b/Engine/Documentation/Source/Engine/UI/LevelEditor/Modes/MeshPaintMode/VertexColor/VertexColorMatching/VertexColorMatching.INT.udn
@@ -2,6 +2,8 @@ Availability: Public
 Title:Vertex Color Matching
 Crumbs:%ROOT%, Engine, Engine/UI/LevelEditor/Modes/MeshPaintMode/VertexColor
 Description:A tool which allows vertex colors to be repaired when a new version of a mesh is imported with a different vertex count.
+Version: 4.9
+
 
 
 
@@ -23,10 +25,11 @@ Mismatched instanced vertex colors against the mesh's vertices makes the vertex
 
 This feature will only be available if a mesh needs fixup; otherwise, it will be disabled. The tool works well in most cases especially when you are only making small tweaks. The more drastic the change to the mesh, the more unlikely the colors will match. The feature is designed to always match some color even if the changes are drastic. We elected not to make this tool automatic because that would prevent people from easily opting out of the fixup if they did not like the results. Additionally, adding the feature to mesh paint means you can easily touchup after the fix is applied.
 
-**Original mesh:**
-
-![fix_before.png](fix_before.png)(w:800 h:354)
-
-**Lower poly mesh with Fixup applied:**
-
-![fix_after.png](fix_after.png)(w:800 h:354)
+[OBJECT:ComparisonSlider]
+	[PARAM:before]
+	![Original Mesh - Fix Before](fix_before.png)(w:850 h:404)
+	[/PARAM]
+	[PARAM:after]
+	![Lower Poly Mesh - Fixup applied](fix_after.png)(w:850 h:404)
+	[/PARAM]
+[/OBJECT]
\ No newline at end of file
