#UE4doc #docs - Added Engine Version field.

diff --git a/Engine/Documentation/Source/Engine/UI/LevelEditor/Modes/MeshPaintMode/VertexColor/ViewModes/ViewModes.INT.udn b/Engine/Documentation/Source/Engine/UI/LevelEditor/Modes/MeshPaintMode/VertexColor/ViewModes/ViewModes.INT.udn
--- a/Engine/Documentation/Source/Engine/UI/LevelEditor/Modes/MeshPaintMode/VertexColor/ViewModes/ViewModes.INT.udn
+++ b/Engine/Documentation/Source/Engine/UI/LevelEditor/Modes/MeshPaintMode/VertexColor/ViewModes/ViewModes.INT.udn
@@ -2,6 +2,7 @@ Availability: Public
 Title:View Modes
 Crumbs:%ROOT%, Engine, Engine/UI/LevelEditor/Modes/MeshPaintMode/VertexColor
 Description:A look at the different view modes available in Mesh Paint Mode.
+Version: 4.9
 
 [EXCERPT:Intro]
 Mesh Paint Mode includes a variety of different view modes to aid in visualizing how vertex color is being applied to your meshes within the viewport. The view can display just the vertex colors, or can selectively view individual channels (red, green, blue, and alpha).
