#UE4 Docs: Fixing a few sentances so they read better.

diff --git a/Engine/Documentation/Source/Engine/Basics/Screenshots/Screenshots.INT.udn b/Engine/Documentation/Source/Engine/Basics/Screenshots/Screenshots.INT.udn
--- a/Engine/Documentation/Source/Engine/Basics/Screenshots/Screenshots.INT.udn
+++ b/Engine/Documentation/Source/Engine/Basics/Screenshots/Screenshots.INT.udn
@@ -7,8 +7,8 @@ version: 4.9
 
 
 
-Unreal Engine 4 (UE4) offers users a multitude of different ways of taking in-game screenshots of your Unreal Engine game project. 
-This document will not only cover which tools you have available for taking screenshots but also how you can use the tools. 
+Unreal Engine 4 (UE4) offers users a multitude of different ways of taking in-game screenshots of your project. 
+This document will not only cover which tools are available for taking screenshots in UE4 but also how you can use the tools. 
 
 
 ## Default Save Location
@@ -20,7 +20,7 @@ When you do take a screenshot, the following message will be displayed in the lo
 
 ![](HRSS_Save_Location.png)
 
-If you click on that message a folder with the Screenshots will be displayed for you. 
+If you click on that message the folder with the Screenshots will be displayed for you. 
 
 ## Standard Screen Shot
 
