TypoCorrection

diff --git a/Engine/Documentation/Source/Engine/Blueprints/UserGuide/Types/BlueprintTypes.INT.udn b/Engine/Documentation/Source/Engine/Blueprints/UserGuide/Types/BlueprintTypes.INT.udn
--- a/Engine/Documentation/Source/Engine/Blueprints/UserGuide/Types/BlueprintTypes.INT.udn
+++ b/Engine/Documentation/Source/Engine/Blueprints/UserGuide/Types/BlueprintTypes.INT.udn
@@ -1,4 +1,4 @@
-Availability:Public
+﻿Availability:Public
 Title: Types of Blueprints
 Crumbs: %ROOT%, Engine/Blueprints
 Description:Landing page for information on different types of Blueprints.
@@ -24,7 +24,7 @@ are ideal to live in a Blueprint Macro Library.
 
 [INCLUDE:Engine/Blueprints/UserGuide/Types/LevelBlueprint#Intro]
 
-For more information on this section, please see [%Engine/Blueprints/UserGuide/Types/LevelBlueprint:title%](Engine/Blueprints/UserGuide/Types/LevelBlueprint).
+For more information on this section, please see [](Engine/Blueprints/UserGuide/Types/LevelBlueprint).
 
 
 ##Blueprint Class
@@ -33,7 +33,7 @@ For more information on this section, please see [%Engine/Blueprints/UserGuide/T
 
 [INCLUDE:Engine/Blueprints/UserGuide/Types/ClassBlueprint#Intro]
 
-For more information on this section, please see [%Engine/Blueprints/UserGuide/Types/ClassBlueprint:title%](Engine/Blueprints/UserGuide/Types/ClassBlueprint).
+For more information on this section, please see [](Engine/Blueprints/UserGuide/Types/ClassBlueprint).
 
 [/EXCERPT:Main]
 
