#UE4doc #docs - Added Engine Version Field.

diff --git a/Engine/Documentation/Source/Engine/UI/LevelEditor/PerformanceTools/PerformanceTools.INT.udn b/Engine/Documentation/Source/Engine/UI/LevelEditor/PerformanceTools/PerformanceTools.INT.udn
--- a/Engine/Documentation/Source/Engine/UI/LevelEditor/PerformanceTools/PerformanceTools.INT.udn
+++ b/Engine/Documentation/Source/Engine/UI/LevelEditor/PerformanceTools/PerformanceTools.INT.udn
@@ -2,6 +2,7 @@ Availability:Public
 Title:Performance Tools
 Crumbs: %ROOT%, Engine, Engine/UI, Engine/UI/LevelEditor
 Description:Collection of commands that provide quick access to commonly used tools and operations.
+Version: 4.9
 
 At the top of the Level Editor window, at the right side of the Tab Bar, you will find a series of performance and debugging tools. 
 
