#UE4 Docs: Making the first check in for the Infilitrator docs.

diff --git a/Engine/Documentation/Source/Resources/Showcases/Infiltrator/Infiltrator.INT.udn b/Engine/Documentation/Source/Resources/Showcases/Infiltrator/Infiltrator.INT.udn
new file mode 100644
--- /dev/null
+++ b/Engine/Documentation/Source/Resources/Showcases/Infiltrator/Infiltrator.INT.udn
@@ -0,0 +1,138 @@
+Availability: Public
+Title:Infiltrator Demo
+Crumbs:%ROOT%, Resources,
+Description:An overview of the real time Infiltrator Demo
+Related: Engine/Matinee
+Related: Engine/Rendering/Materials/HowTo/CreatingLayeredMaterials
+Related: Engine/Rendering/ParticleSystems
+Related: Engine/Animation
+Related: Engine/Content/Tools/MayaRiggingTool
+Version: 4.9
+Skilllevel: Intermediate
+
+![](Infiltrator_Page_Header.png)
+
+First shown at GDC 2013, Infiltrator showcases Unreal Engine 4's (UE4) PBR and cinematic capabilities at the time. 
+The project is full of expert-level content to help you learn about all the different aspects of what goes into creating a AAA-quality scene. 
+In the following document we will high light some of the UE4 specific tools that were used to create the Infiltrator demo.
+
+[region:note]
+Please note, if you use the Infiltrator assets in a new project with reference lighting, you'll have to make a few tweaks to make the Infiltrator content match your content.
+[/region]
+
+## Location
+You can find and download the Infiltrator project from the Engine Features Samples section of the Learn Tab in the Epic Games Launcher. 
+
+ ![](Infiltrator_Launcher_Location.png)
+
+
+## Tools Used
+The Infiltrator demo makes uses of almost all of the tools UE4 has to offer. 
+In the following section we will highlight some of these tools as well as provide links to documentation for further reading.
+
+### Matinee 
+
+![](Infiltrator_Matinee.png)
+
+The Matinee cinematic tool provides the ability to animate the properties of Actors over time, to create either dynamic game play or cinematic in-game sequences. 
+Matinee is the main tool used to bring the Infiltrator demo from an idea into reality. 
+For more information about Matinee and how to use it make sure to check out the official [Matinee documentation](https://docs.unrealengine.com/latest/INT/Engine/Matinee/index.html) for more information.
+
+
+### Material Editor
+
+![](Infiltrator_Material_Graph_Whole.png)
+
+As with any Unreal-based project the Material Editor was used to create all of the Materials that are being used. 
+However Infiltrator make use of a new way of thinking about how to apply Materials to objects. 
+Instead of using a single Material, Infiltrator uses Material Layers which allow you to mix and match different Materials together instead of just Textures. 
+For more information about Layered Materials and how to use it make sure to check out the official [How To on creating Layered Materials](https://docs.unrealengine.com/latest/INT/Engine/Rendering/Materials/HowTo/CreatingLayeredMaterials/index.html) for more information on how to create, setup and use Layered Materials.
+
+### Cascade Particle Editor
+
+![](Infiltrator_Cascade_Editor.png)
+
+ The Cascade particle editor allows artist to quickly and easily create incredible visual effects ranging from smoke and sparks to effects for pick ups or things the player should be aware of. 
+ All of the effects you see in Infiltrator were created using the Cascade Particle Editor. 
+ For more information on all the features Cascade offers check out the official documentation on [Cascade Particle Systems](https://docs.unrealengine.com/latest/INT/Engine/Rendering/ParticleSystems/index.html).
+
+
+### Persona Animation System
+
+![](Infiltrator_Animation_Tools.png)
+
+ UE4's Persona animation system allows users to create animations for their project using a wide variety of animation data. 
+ From using skeletal-based animations to morph-based animations, the animation system enables both simple and complex animation across a wide variety of different characters. 
+ For more information about what Persona Animation can do for your project make sure to check out the official [Persona Animation System](https://docs.unrealengine.com/latest/INT/Engine/Animation/index.html) documentation.
+
+
+### Maya Animation Rigging Toolset (A.R.T)
+
+![](Infiltrator_ART_Tools.jpg)(h:550)
+
+The Maya Animation Rigging Toolset or A.R.T for short, is a comprehensive toolset that takes a character through skeleton creation, skeleton placement, and rig creation, as well as many animation tools to speed up animator workflow.
+ For more information about the A.R.T tools and how to install and use them, check out the official documentation.
+
+
+## Infiltrator Matinee Breakdown
+
+To access the Infiltrator cinematic you will need to first go to the Main Toolbar and open the Matinee sequence label **Cinematic**. 
+
+![](Infiltrator_Open_Cinematic.png)
+
+After the Infiltrator cinematic Matinee has been opened you can play the cinematic using the Play, Loop, Stop, or Reverse buttons or by scrubbing the Timeline. 
+
+![](Infiltrator_Matinee_Scrub.png)
+
+
+
+
+
+
+
+
+
+
+
+
+
+
+
+
+
+
+
+
+
+
+
+
+
+
+
+
+
+
+
+
+
+
+
+
+
+
+
+
+
+
+
+
+
+
+
+
+
+
+
+
+
