#UE4doc #docs - Added Engine version field. Re-organized list to match UI changes. Added multiple settings that were not listed.

diff --git a/Engine/Documentation/Source/Engine/UI/LevelEditor/QuickSettings/ViewOptions.INT.udn b/Engine/Documentation/Source/Engine/UI/LevelEditor/QuickSettings/ViewOptions.INT.udn
--- a/Engine/Documentation/Source/Engine/UI/LevelEditor/QuickSettings/ViewOptions.INT.udn
+++ b/Engine/Documentation/Source/Engine/UI/LevelEditor/QuickSettings/ViewOptions.INT.udn
@@ -2,8 +2,9 @@ Availability:Public
 Title:Quick Settings
 Crumbs: %ROOT%, Engine, Engine/UI, Engine/UI/LevelEditor
 Description:Options for controlling selection, editing, and preview aspects of the Level Editor.
+Version: 4.9
 
-The **Quick Settings** are a set of properties for controlling selection, editing, and previewing aspects of the Level Editor. These can be accessed in the [](Engine/UI/LevelEditor/Toolbar).
+The **Settings** are a set of properties for controlling selection, editing, and previewing aspects of the Level Editor. These can be accessed in the [](Engine/UI/LevelEditor/Toolbar).
 
 | Option | Description |
 | ------ | ----------- |
@@ -11,28 +12,32 @@ The **Quick Settings** are a set of properties for controlling selection, editin
 | **Allow Translucent Selection** | If checked, clicking on geometry with translucent materials applied will select the Actor. [REGION:tip]Enabling this option makes it possible to select meshes that are transparent, such as glass objects. Disabling the option can be very helpful in other cases, like in the case of a room filled with translucent particles. Selecting other objects in the room would be difficult otherwise since clicking on the particles would select the Emitter Actor.[/REGION] |
 | **Allow Group Selection** | If checked, selecting an Actor in a group selects the group instead of the individual Actor. |
 | **Strict Box Selection** | If checked, an Actor must be completely within the marquee selection box in order to be selected. |
-| [REGION:tablesection]Editing[/REGION] ||
+| **Box Select Occluded Objects** | When enabled, marquee box select operations will also select objects that are occluded by other objects. |
 | **Show Transform Widget** | Toggles the visibility of the [Transform widget](Engine/Actors/Transform) in the Viewports. |
-| **Draw Brush Marker Polys** | If checked, semi-translucent polygons are rendered for the faces of a CSG brush when selected. [INCLUDE:Engine/UI/LevelEditor/QuickSettings#BrushPolys] |
-| [REGION:tablesection]Previewing[/REGION] ||
-| **Only Load Visible Levels in Game Preview** | If checked, only visible levels will be loaded when game preview starts. |
-| **Enable Particle System LOD Switching** | If checked, particle systems will use distance LOD switching in perspective Viewports. |
-| **Toggle Particle System Helpers** | Toggles showing particles system helpers in Viewports. |
-| **Freeze Particle Simulation** | If checked, particle systems will freeze their simulation state. |
-| **Enable LOD View Locking** | If checked, Viewports of the same type will use the same LOD. |
-| **Enable Automatic Level Streaming** | If checked, the Viewport will stream in levels automatically when the camera is moved. |
+| [REGION:tablesection]Scalability[/REGION] ||
 | **Engine Scalability Settings** | [INCLUDE:#scalability] |
 | **Material Quality level** | [INCLUDE:#materialquality] |
+| **Preview Rendering Level** | This sets the rendering level used by the editor. |
 | [REGION:tablesection]Real Time Audio[/REGION] ||
 | **Volume** | Controls the volume of sounds from the world that are played when the Viewport is set to realtime. |
-| [REGION:tablesection]Actor Snap[/REGION] ||
+| [REGION:tablesection]Snapping[/REGION] ||
 | **Enable Actor Snapping** | If checked, Actors will snap to the location of other Actors if they are within a specified distance. |
 | **Distance** | Sets the distance Actors must be within in order to snap to one another if **Enable Actor Snapping** is checked. |
-| [REGION:tablesection]Snapping[/REGION] ||
 | **Enable Socket Snapping** | If checked, allows Actors to be snapped to sockets. |
 | **Enable Vertex Snapping** | If checked, Actors will snap to the nearest vertex on another Actor in the direction of movement. |
-| [REGION:tablesection]Viewport[/REGION] ||
+| **Enable Planar Snapping** | If checked, actors will snap to the nearest location on the constraint plane. This only works correctly in Perspective views. |
+[REGION:tablesection]Viewport[/REGION] ||
 | **Hide Viewport UI** | Toggles hidden Viewport UI mode, and hides all overlaid Viewport UI widgets. |
+| [REGION:tablesection]Previewing[/REGION] ||
+| **Draw Brush Polys** | If checked, semi-translucent polygons are rendered for the faces of a CSG brush when selected. [INCLUDE:Engine/UI/LevelEditor/QuickSettings#BrushPolys] |
+| **Only Load Visible Levels in Game Preview** | If checked, only visible levels will be loaded when game preview starts. |
+| **Enable Particle System LOD Switching** | If checked, particle systems will use distance LOD switching in perspective Viewports. |
+| **Toggle Particle System Helpers** | Toggles showing particles system helpers in Viewports. |
+| **Freeze Particle Simulation** | If checked, particle systems will freeze their simulation state. |
+| **Enable LOD View Locking** | If checked, Viewports of the same type will use the same LOD. |
+| **Enable Automatic Level Streaming** | If checked, the Viewport will stream in levels automatically when the camera is moved. |
+
+
 
 <!--
 [EXCERPT:BrushPolys]
