#UE4doc #docs - Reviewed Infiltrator doc and made corrections. See UEDOC-2360 for more details.

diff --git a/Engine/Documentation/Source/Resources/Showcases/Infiltrator/Infiltrator.INT.udn b/Engine/Documentation/Source/Resources/Showcases/Infiltrator/Infiltrator.INT.udn
--- a/Engine/Documentation/Source/Resources/Showcases/Infiltrator/Infiltrator.INT.udn
+++ b/Engine/Documentation/Source/Resources/Showcases/Infiltrator/Infiltrator.INT.udn
@@ -21,7 +21,7 @@ Please note, if you use the Infiltrator assets in a new project with reference l
 [/region]
 
 ## Location
-You can find and download the Infiltrator project from the Engine Features Samples section of the Learn Tab in the Epic Games Launcher. 
+You can find and download the Infiltrator project from the **Engine Features Samples** section of the **Learn Tab** in the Epic Games Launcher. 
 
  ![](Infiltrator_Launcher_Location.png)
 
@@ -46,7 +46,7 @@ For more information about Matinee and how to use it make sure to check out the
 As with any Unreal-based project the Material Editor was used to create all of the Materials that are being used. 
 However Infiltrator make use of a new way of thinking about how to apply Materials to objects. 
 Instead of using a single Material, Infiltrator uses Material Layers which allow you to mix and match different Materials together instead of just Textures. 
-For more information about Layered Materials and how to use it make sure to check out the official [How To on creating Layered Materials](https://docs.unrealengine.com/latest/INT/Engine/Rendering/Materials/HowTo/CreatingLayeredMaterials/index.html) for more information on how to create, setup and use Layered Materials.
+For more information about Layered Materials and how to use it make sure to check out the official [How To on creating Layered Materials](https://docs.unrealengine.com/latest/INT/Engine/Rendering/Materials/HowTo/CreatingLayeredMaterials/index.html) to create, setup and use Layered Materials.
 
 ### Cascade Particle Editor
 
@@ -71,7 +71,7 @@ For more information about Layered Materials and how to use it make sure to chec
 ![](Infiltrator_ART_Tools.jpg)(h:550)
 
 The Maya Animation Rigging Toolset or A.R.T for short, is a comprehensive toolset that takes a character through skeleton creation, skeleton placement, and rig creation, as well as many animation tools to speed up animator workflow.
- For more information about the A.R.T tools and how to install and use them, check out the official documentation.
+ For more information about the A.R.T tools and how to install and use them, check out the official documentation on the [Animation Rigging Toolset](https://docs.unrealengine.com/latest/INT/Engine/Content/Tools/MayaRiggingTool/index.html).
 
 
 ## Infiltrator Matinee Breakdown
