#UE4doc #docs -  Added Engine Version Field. Removed Overview header. Unbolded unnecessary terms.

diff --git a/Engine/Documentation/Source/Engine/UI/LevelEditor/SceneOutliner/SceneOutliner.INT.udn b/Engine/Documentation/Source/Engine/UI/LevelEditor/SceneOutliner/SceneOutliner.INT.udn
--- a/Engine/Documentation/Source/Engine/UI/LevelEditor/SceneOutliner/SceneOutliner.INT.udn
+++ b/Engine/Documentation/Source/Engine/UI/LevelEditor/SceneOutliner/SceneOutliner.INT.udn
@@ -2,6 +2,7 @@ Availability:Public
 Title:World Outliner
 Crumbs: %ROOT%, Engine, Engine/UI, Engine/UI/LevelEditor
 Description:Hierarchical tree-view of all Actors within the current level used for selection, attachment, and more.
+Version: 4.9
 
 [VAR:playlistnav]
 [OBJECT:PlaylistNavigation]
@@ -25,15 +26,14 @@ Description:Hierarchical tree-view of all Actors within the current level used f
 
 [TOC(start:2)]
 
-## Overview
 
 [EXCERPT:SceneOutliner]
-The **World Outliner** panel displays all of the Actors within the scene in a hierarchical tree view. Actors can be selected and modified directly from the **World Outliner**. You can also use the **Info** drop down menu to display an extra column that shows Levels, Layers, or ID Names.
+The **World Outliner** panel displays all of the Actors within the scene in a hierarchical tree view. Actors can be selected and modified directly from the World Outliner. You can also use the **Info** drop down menu to display an extra column that shows Levels, Layers, or ID Names.
 [/EXCERPT:SceneOutliner]
 
 ![Scene Outliner Panel](scene_outliner.png)
 
-For more information about the **World Outliner** UI, see [World Outliner UI](Engine/UI/LevelEditor\SceneOutliner\Interface).
+For more information about the [World Outliner UI](Engine/UI/LevelEditor\SceneOutliner\Interface) see the documentation.
 
 ## Outliner Actions
 
@@ -45,7 +45,7 @@ For more information about the **World Outliner** UI, see [World Outliner UI](En
 
 ## Searching the World
 
-The **World Outliner** contains a text box that is used to search and quickly filter the list of Actors in the scene. All Actors with even a partial match to the search term will be displayed. You can use multiple words in the search box now and only Actors that match all the terms will be displayed.
+The World Outliner contains a text box that is used to search and quickly filter the list of Actors in the scene. All Actors with even a partial match to the search term will be displayed. You can use multiple words in the search box now and only Actors that match all the terms will be displayed.
 
 ![Scene Outliner Search](outliner_search.png)
 
