#UE4doc #docs  - Added Engine Version field. Unbolded unnecessary term.

diff --git a/Engine/Documentation/Source/Engine/UI/LevelEditor/SceneOutliner/Interface/SceneOutlinerInterface.INT.udn b/Engine/Documentation/Source/Engine/UI/LevelEditor/SceneOutliner/Interface/SceneOutlinerInterface.INT.udn
--- a/Engine/Documentation/Source/Engine/UI/LevelEditor/SceneOutliner/Interface/SceneOutlinerInterface.INT.udn
+++ b/Engine/Documentation/Source/Engine/UI/LevelEditor/SceneOutliner/Interface/SceneOutlinerInterface.INT.udn
@@ -2,6 +2,7 @@ Availability:Public
 Title:World Outliner UI
 Crumbs: %ROOT%, Engine, Engine/UI, Engine/UI/LevelEditor, Engine/UI/LevelEditor/SceneOutliner
 Description:Interface for hierarchical tree-view of all Actors within the current level used for selection, attachment, and more.
+Version: 4.9
 
 [TOC(start:2)]
 
@@ -10,7 +11,7 @@ Description:Interface for hierarchical tree-view of all Actors within the curren
 [INCLUDE:Engine/UI/LevelEditor/SceneOutliner/#SceneOutliner]
 
 [REGION:note]
-When the Level Editor is in Play mode or Simulate mode, the **World Outliner** may display Actors that are dynamically spawned from gameplay code or from Blueprints, in addition to the existing Actors.
+When the Level Editor is in Play mode or Simulate mode, the World Outliner may display Actors that are dynamically spawned from gameplay code or from Blueprints, in addition to the existing Actors.
 [/REGION]
 
 ![](LE_SceneOutlinerUI.png)
@@ -26,4 +27,4 @@ When the Level Editor is in Play mode or Simulate mode, the **World Outliner** m
 
 [Include:Shared/Editors/Common/EditorMenuItems/#UITabTips]
 
-For more information about the **World Outliner**, see [](Engine/UI/LevelEditor\SceneOutliner).
+For more information about the **World Outliner**, see the [](Engine/UI/LevelEditor\SceneOutliner) documentation.
