#UE4doc #docs - Added Engine Version Field. Corrected outdated term.

diff --git a/Engine/Documentation/Source/Engine/UI/LevelEditor/Toolbar/LevelEditorToolbar.INT.udn b/Engine/Documentation/Source/Engine/UI/LevelEditor/Toolbar/LevelEditorToolbar.INT.udn
--- a/Engine/Documentation/Source/Engine/UI/LevelEditor/Toolbar/LevelEditorToolbar.INT.udn
+++ b/Engine/Documentation/Source/Engine/UI/LevelEditor/Toolbar/LevelEditorToolbar.INT.udn
@@ -2,6 +2,7 @@ Availability:Public
 Title:Level Editor Toolbar
 Crumbs: %ROOT%, Engine, Engine/UI, Engine/UI/LevelEditor
 Description:Collection of commands that provide quick access to commonly used tools and operations.
+Version: 4.9
 
 [VAR:playlistnav]
 [OBJECT:PlaylistNavigation]
@@ -35,7 +36,7 @@ The **Toolbar** panel, like in most applications, is a group of commands providi
 | ![](SourceControl.png)| Allows you to connect to or assign your [Source Control](Engine/UI/SourceControl) solution. |
 | ![](button_Toolbar_Content.png)| Opens the [Content Browser](Engine/Content/Browser) displaying all of the content assets contained in your game. This is where you go to create, import, and edit all content. |
 | ![](button_Toolbar_marketplace.png) | Opens the Unreal Engine Launcher to the Marketplace section.  |
-| ![](button_Toolbar_View.png) | Displays the [Quick Settings](Engine/UI/LevelEditor/QuickSettings) menu providing easy access to commonly used options controlling selection, editing, and previewing aspects of the Level Editor. |
+| ![](button_Toolbar_View.png) | Displays the [Settings](Engine/UI/LevelEditor/QuickSettings) menu providing easy access to commonly used options controlling selection, editing, and previewing aspects of the Level Editor. |
 | ![](button_Toolbar_blueprints.png) | Provides access to create or edit any Blueprints in the world, including opening the [Level Blueprint](Engine/Blueprints/UserGuide/Types/LevelBlueprint) for the current level in the [Blueprint Editor](Engine/Blueprints/Editor\UIBreakdowns\LevelBPUI). This menu give you quick access to setting up the framework for your game - game rules, player type, HUD, etc. - from within the editor. |
 | ![](button_Toolbar_matinee.png)| Enables you to create a new [Matinee](Engine/Matinee) sequence or edit any existing **Matinee** sequences in the level. |
 | ![](button_toolbar_build.png) | Performs a the _build_ operation on all levels (persistent and streaming) open in the editor. Building precalculates as much data as possible relating to various aspects of the level. For instance, static lighting - lightmaps, shadows, global illumination - and geometry are calculated during this process. Clicking the arrow displays the **Build Options** menu. |
