#UE4doc #docs - Removed Overview header. Added Engine Version field. Bolded term.

diff --git a/Engine/Documentation/Source/Engine/UI/LevelEditor/Viewports/EditorViewports.INT.udn b/Engine/Documentation/Source/Engine/UI/LevelEditor/Viewports/EditorViewports.INT.udn
--- a/Engine/Documentation/Source/Engine/UI/LevelEditor/Viewports/EditorViewports.INT.udn
+++ b/Engine/Documentation/Source/Engine/UI/LevelEditor/Viewports/EditorViewports.INT.udn
@@ -2,6 +2,7 @@ Availability:Public
 Title:Editor Viewports
 Crumbs: %ROOT%, Engine, Engine/UI
 Description:A guide to navigation techniques and visualization settings for the Unreal Editor viewports.
+Version: 4.9
 
 [VAR:Topic]
 [OBJECT:Topic]
@@ -47,8 +48,6 @@ Description:A guide to navigation techniques and visualization settings for the
 ![](Viewports_Topic.png)
 [/REGION]
 
-## Overview
-
 [EXCERPT:Intro]
 [INCLUDE:Shared/Tutorials/inEditorTutorialCamera#viewportintro]
 [/EXCERPT:Intro]
diff --git a/Engine/Documentation/Source/Shared/Tutorials/inEditorTutorialCamera/inEditorTutorialCamera.INT.udn b/Engine/Documentation/Source/Shared/Tutorials/inEditorTutorialCamera/inEditorTutorialCamera.INT.udn
--- a/Engine/Documentation/Source/Shared/Tutorials/inEditorTutorialCamera/inEditorTutorialCamera.INT.udn
+++ b/Engine/Documentation/Source/Shared/Tutorials/inEditorTutorialCamera/inEditorTutorialCamera.INT.udn
@@ -10,7 +10,7 @@ Viewport Introduction
 [/VAR]
 
 [EXCERPT:ViewportIntro]
-The Viewports are your window into the worlds you create in Unreal. They can be navigated just as you would in a game, or can be used in a more schematic design sense as you would for an architectural blueprint. The Unreal Editor viewports contain a variety of tools and visualizers to help you see exactly the data you need. 
+The **Viewports** are your window into the worlds you create in Unreal. They can be navigated just as you would in a game, or can be used in a more schematic design sense as you would for an architectural blueprint. The Unreal Editor viewports contain a variety of tools and visualizers to help you see exactly the data you need. 
 [/EXCERPT:ViewportIntro]
 
 ![](viewports.png)
