#UE4 Docs: fixed broken include

diff --git a/Engine/Documentation/Source/Engine/Blueprints/UserGuide/FlowControl/FlowControl.INT.udn b/Engine/Documentation/Source/Engine/Blueprints/UserGuide/FlowControl/FlowControl.INT.udn
--- a/Engine/Documentation/Source/Engine/Blueprints/UserGuide/FlowControl/FlowControl.INT.udn
+++ b/Engine/Documentation/Source/Engine/Blueprints/UserGuide/FlowControl/FlowControl.INT.udn
@@ -47,7 +47,7 @@ version: 4.9
 
 [EXCERPT:Intro]
 
-[INCLUDE:Shared/Glossary/F/#GlossaryFlow_Control]
+[INCLUDE:Shared/Glossary/F#Glossary_Flow_Control]
 
 [/EXCERPT:Intro]
 
