#UE4doc #docs - Removed Overview Header. Added Engine Version Field. Added Skill Level Field.

diff --git a/Engine/Documentation/Source/Engine/UI/LevelEditor/Viewports/Basics/ViewportBasics.INT.udn b/Engine/Documentation/Source/Engine/UI/LevelEditor/Viewports/Basics/ViewportBasics.INT.udn
--- a/Engine/Documentation/Source/Engine/UI/LevelEditor/Viewports/Basics/ViewportBasics.INT.udn
+++ b/Engine/Documentation/Source/Engine/UI/LevelEditor/Viewports/Basics/ViewportBasics.INT.udn
@@ -2,6 +2,8 @@ Availability:Public
 Title:Viewport Basics
 Crumbs: %ROOT%, Engine, Engine/UI, Engine/UI/LevelEditor/Viewports
 Description:Basic concepts and features of the Viewports in Unreal Editor.
+Version: 4.9
+SkillLevel: Beginner
 
 [VAR:playlistnav]
 [OBJECT:PlaylistNavigation]
@@ -65,8 +67,6 @@ Description:Basic concepts and features of the Viewports in Unreal Editor.
 
 [TOC(start:2 end:3)]
 
-## Overview
-
 [INCLUDE:Engine/UI/LevelEditor/Viewports#intro]
 
 ![](Engine/UI/LevelEditor/Viewports/Viewports_Topic.png)(w:640)
