#UE4doc #docs - Removed Overview header. Added Skill Level field. Added Engine Version Field. Corrected editor term for Actor Locking that is now referred to as Piloting.

diff --git a/Engine/Documentation/Source/Engine/UI/LevelEditor/Viewports/PIP/ViewportPIP.INT.udn b/Engine/Documentation/Source/Engine/UI/LevelEditor/Viewports/PIP/ViewportPIP.INT.udn
--- a/Engine/Documentation/Source/Engine/UI/LevelEditor/Viewports/PIP/ViewportPIP.INT.udn
+++ b/Engine/Documentation/Source/Engine/UI/LevelEditor/Viewports/PIP/ViewportPIP.INT.udn
@@ -2,6 +2,8 @@ Availability:Public
 Title:Picture-in-Picture
 Crumbs: %ROOT%, Engine, Engine/UI, Engine/UI/LevelEditor/Viewports
 Description:A viewport feature that creates a picture-in-picture for accurate camera placement.
+Version: 4.9
+SkillLevel:Beginner
 
 [VAR:Topic]
 [OBJECT:Topic]
@@ -55,7 +57,7 @@ Selecting a Camera Actor shows a preview of that camera overlaid in the active V
 ![](pip_single.png)(w:720)
 
 [REGION:note]
-	The preview is not displayed when a selected camera is _locked_ to the active Viewport using the **Actor Locking** fly-out in the Viewport dropdown menu.
+	The preview is not displayed when a selected camera is _locked_ to the active Viewport using the **Piloting** option in the right-click context menu.
 [/REGION]
 
 When multiple Camera Actors are selected, a preview for each one is displayed.
