#UE4 Docs: minor text updates

diff --git a/Engine/Documentation/Source/Resources/Showcases/Infiltrator/Infiltrator.INT.udn b/Engine/Documentation/Source/Resources/Showcases/Infiltrator/Infiltrator.INT.udn
--- a/Engine/Documentation/Source/Resources/Showcases/Infiltrator/Infiltrator.INT.udn
+++ b/Engine/Documentation/Source/Resources/Showcases/Infiltrator/Infiltrator.INT.udn
@@ -34,9 +34,9 @@ In the following section we will highlight some of these tools as well as provid
 
 ![](Infiltrator_Matinee.png)
 
-The Matinee cinematic tool provides the ability to animate the properties of Actors over time, to create either dynamic game play or cinematic in-game sequences. 
+The Matinee Cinematic Tool provides the ability to animate the properties of Actors over time, to create either dynamic game play or cinematic in-game sequences. 
 Matinee is the main tool used to bring the Infiltrator demo from an idea into reality. 
-For more information about Matinee and how to use it make sure to check out the official [Matinee documentation](https://docs.unrealengine.com/latest/INT/Engine/Matinee/index.html) for more information.
+For more information about Matinee and how to use it, make sure to check out [](Engine/Matinee).
 
 
 ### Material Editor
@@ -46,32 +46,32 @@ For more information about Matinee and how to use it make sure to check out the
 As with any Unreal-based project the Material Editor was used to create all of the Materials that are being used. 
 However Infiltrator make use of a new way of thinking about how to apply Materials to objects. 
 Instead of using a single Material, Infiltrator uses Material Layers which allow you to mix and match different Materials together instead of just Textures. 
-For more information about Layered Materials and how to use it make sure to check out the official [How To on creating Layered Materials](https://docs.unrealengine.com/latest/INT/Engine/Rendering/Materials/HowTo/CreatingLayeredMaterials/index.html) to create, setup and use Layered Materials.
+To learn more about creating, setting up, and using Layered Materials, see [](Engine/Rendering/Materials/HowTo/CreatingLayeredMaterials).
 
 ### Cascade Particle Editor
 
 ![](Infiltrator_Cascade_Editor.png)
 
- The Cascade particle editor allows artist to quickly and easily create incredible visual effects ranging from smoke and sparks to effects for pick ups or things the player should be aware of. 
+ The Cascade Particle Editor allows artist to quickly and easily create incredible visual effects ranging from smoke and sparks to effects for pick ups or things the player should be aware of. 
  All of the effects you see in Infiltrator were created using the Cascade Particle Editor. 
- For more information on all the features Cascade offers check out the official documentation on [Cascade Particle Systems](https://docs.unrealengine.com/latest/INT/Engine/Rendering/ParticleSystems/index.html).
+ For a complete overview of all the features Cascade offers, check out [](Engine/Rendering/ParticleSystems).
 
 
 ### Persona Animation System
 
 ![](Infiltrator_Animation_Tools.png)
 
- UE4's Persona animation system allows users to create animations for their project using a wide variety of animation data. 
+ UE4's Persona Animation System allows users to create animations for their project using a wide variety of animation data. 
  From using skeletal-based animations to morph-based animations, the animation system enables both simple and complex animation across a wide variety of different characters. 
- For more information about what Persona Animation can do for your project make sure to check out the official [Persona Animation System](https://docs.unrealengine.com/latest/INT/Engine/Animation/index.html) documentation.
+ For details on how to incorporate animation into your project, please see [](Engine/Animation).
 
 
 ### Maya Animation Rigging Toolset (A.R.T)
 
 ![](Infiltrator_ART_Tools.jpg)(h:550)
 
-The Maya Animation Rigging Toolset or A.R.T for short, is a comprehensive toolset that takes a character through skeleton creation, skeleton placement, and rig creation, as well as many animation tools to speed up animator workflow.
- For more information about the A.R.T tools and how to install and use them, check out the official documentation on the [Animation Rigging Toolset](https://docs.unrealengine.com/latest/INT/Engine/Content/Tools/MayaRiggingTool/index.html).
+The Maya Animation Rigging Toolset, or A.R.T for short, is a comprehensive toolset that takes a character through skeleton creation, skeleton placement, and rig creation, as well as many animation tools to speed up animator workflow.
+ If you would like to learn more about the A.R.T tools and how to install and use them, check out the [Animation Rigging Toolset](Engine/Content/Tools/MayaRiggingTool) guide.
 
 
 ## Infiltrator Matinee Breakdown
