#UE4 Docs: added topic image and link

diff --git a/Engine/Documentation/Source/Resources/Showcases/Infiltrator/Infiltrator.INT.udn b/Engine/Documentation/Source/Resources/Showcases/Infiltrator/Infiltrator.INT.udn
--- a/Engine/Documentation/Source/Resources/Showcases/Infiltrator/Infiltrator.INT.udn
+++ b/Engine/Documentation/Source/Resources/Showcases/Infiltrator/Infiltrator.INT.udn
@@ -10,6 +10,45 @@ Related: Engine/Content/Tools/MayaRiggingTool
 Version: 4.9
 Skilllevel: Intermediate
 
+[VAR:Topic]
+[OBJECT:Topic]
+	[PARAM:image]
+		![%Resources/Showcases/Infiltrator:title%](Infiltrator_Topic.png)
+	[/PARAM]
+	[PARAM:icon]
+	[/PARAM]
+	[PARAM:title]
+		%Resources/Showcases/Infiltrator:title%
+	[/PARAM]
+	[PARAM:description]
+		%Resources/Showcases/Infiltrator:description%
+	[/PARAM]
+	[PARAM:path]
+		[RELATIVE:Resources/Showcases/Infiltrator]
+	[/PARAM]
+[/OBJECT]
+[/VAR]
+
+
+[VAR:TopicCompact]
+[OBJECT:TopicCompact]
+	[PARAM:image]
+		![%Resources/Showcases/Infiltrator:title%](Infiltrator_Topic.png)
+	[/PARAM]
+	[PARAM:icon]
+	[/PARAM]
+	[PARAM:title]
+		%Resources/Showcases/Infiltrator:title%
+	[/PARAM]
+	[PARAM:description]
+		%Resources/Showcases/Infiltrator:description%
+	[/PARAM]
+	[PARAM:path]
+		[RELATIVE:Resources/Showcases/Infiltrator]
+	[/PARAM]
+[/OBJECT]
+[/VAR]
+
 ![](Infiltrator_Page_Header.png)
 
 First shown at GDC 2013, Infiltrator showcases Unreal Engine 4's (UE4) PBR and cinematic capabilities at the time. 
diff --git a/Engine/Documentation/Source/Resources/Showcases/Showcases.INT.udn b/Engine/Documentation/Source/Resources/Showcases/Showcases.INT.udn
--- a/Engine/Documentation/Source/Resources/Showcases/Showcases.INT.udn
+++ b/Engine/Documentation/Source/Resources/Showcases/Showcases.INT.udn
@@ -50,6 +50,7 @@ The corresponding documentation is designed to provide some insight into how eac
 
 [REGION:topics third]
 %Resources/ContentExamples:TopicCompact%
+%Resources/Showcases/Infiltrator:topiccompact%
 %watertopic%
 %featurestopic%
 %Resources/Showcases/MatineeFightScene:topiccompact%
