#UE4doc #docs:
Multiple changes copied over from another file by Wes
Changed font settings for a few terms
Added some images

diff --git a/Engine/Documentation/Source/Engine/Blueprints/UserGuide/Types/Interface/UsingInterfaces/UsingInterfaces.INT.udn b/Engine/Documentation/Source/Engine/Blueprints/UserGuide/Types/Interface/UsingInterfaces/UsingInterfaces.INT.udn
--- a/Engine/Documentation/Source/Engine/Blueprints/UserGuide/Types/Interface/UsingInterfaces/UsingInterfaces.INT.udn
+++ b/Engine/Documentation/Source/Engine/Blueprints/UserGuide/Types/Interface/UsingInterfaces/UsingInterfaces.INT.udn
@@ -2,6 +2,7 @@ Availability:Public
 Title:Implementing Blueprint Interfaces
 Crumbs: %ROOT%, Engine/Blueprints, Engine/Blueprints/UserGuide/Types/Interface
 Description:Blueprint Interfaces allow Blueprints to call functions in different types of targets, if they also implement the interface.
+Version: 4.9
 
 [VAR:Topic]
 [OBJECT:Topic]
@@ -46,57 +47,66 @@ Description:Blueprint Interfaces allow Blueprints to call functions in different
 ## Using Interfaces in Blueprints
 
 [EXCERPT:intro]
-The use of _Blueprint Interfaces_ allows for a common method of interacting with multiple different types of objects that all share 
+The use of **Blueprint Interfaces** allows for a common method of interacting with multiple different types of objects that all share 
 some specific functionality. This means you can have completely different types of objects, such as a car and a tree, that 
-share one specific thing, like they can both be shot by weapon fire and take damage. By creating a _Blueprint Interface_ that 
-contains an `OnTakeWeaponFire` function, and having both the car and the tree implement that _Blueprint Interface_, you can treat the 
-car and the tree as the same type and simply call the `OnTakeWeaponFire` function when either of them is shot.
+share one specific thing, like they can both be shot by weapon fire and take damage. By creating a Blueprint Interface that 
+contains an OnTakeWeaponFire function, and having both the car and the tree implement that Blueprint Interface, you can treat the 
+car and the tree as the same type and only call the OnTakeWeaponFire function when either of them is shot.
 [/EXCERPT:intro]
 
 
 ### Adding Interfaces
 
-_Blueprint Interfaces_ as well as _Interfaces_ defined in native code (C++) can be added to any type of _Blueprint_ other than a _Level Blueprint_.
+Blueprint Interfaces as well as Interfaces defined in native code (C++) can be added to any type of [Blueprint](Engine/Blueprints).
 
-1. Click the ![](BlueprintProps.png) button in the [Toolbar](Engine/Blueprints/Editor/UIComponents/Toolbar).
+1. From the Main Level Editor window, open the **Level Blueprint** inside your project. 
 
-1. In the **Details** tab, scroll to the **Interfaces** category. Click the ![](button_AddInterface.png) button to see a list of the available Interfaces. Click the one you want to add it.
+	![](OpenLevelBP.png)(w:640)
+
+1. Click the **Class Settings** button from the [Toolbar](Engine/Blueprints/Editor/UIComponents/Toolbar) across the top.
+
+	![](ClassSettings.png)
+
+1. In the **Details** panel on the right, locate the **Interfaces** category. 
+
+1. Click the **Add Interface** button to see a list of the available Interfaces and select the Inteface to use. 
 
 	![](InterfacesCategory.png)
 
 	![](InterfacesListPopulated.png)
 
 [REGION:tip]
-	Interfaces can be removed by clicking the ![](button_X_Remove.png) button located next to each one on the _Implemented Interfaces_ list. You can also migrate from Blueprint Interfaces to normal Blueprint functions. For more information, see [](Engine/Blueprints\UserGuide\BlueprintComms\InterfaceMigration).
+	Interfaces can be removed by clicking the ![](button_X_Remove.png) button located next to each one on the **Implemented Interfaces** list. You can also migrate from Blueprint Interfaces to normal Blueprint functions. For more information, see [](Engine/Blueprints\UserGuide\BlueprintComms\InterfaceMigration).
 [/REGION]
 
 [REGION:note]
-	The _Inherited Interfaces_ section within the Interfaces category automatically populates with any Interfaces that were inherited by this object (i.e. implemented in the Blueprint from which the current Blueprint was derived).
+	The **Inherited Interfaces** section within the Interfaces category automatically populates with any Interfaces that were inherited by this object (i.e. implemented in the Blueprint from which the current Blueprint was derived).
 [/REGION]
 
 
 
 ### Implementing Interface Members
 
-When an _Interface_ is added to a _Blueprint_, that _Blueprint_ must **implement** the _Interface's_ functions in order to have 
-them perform any action when called. This is done by placing _Event_ nodes for the _Interface_ functions in the **EventGraph**. 
-The nodes connected to the _Event_ define the functionality, or implementation, for the function.
+When an Interface is added to a Blueprint, that Blueprint must implement the Interface's functions in order to have 
+them perform any action when called. This is done by placing **Event** nodes for the Interface functions in the **EventGraph**. 
+The nodes connected to the Event define the functionality, or implementation, for the function.
 
-1. **Right-click** in the **EventGraph** of the _Blueprint_ and choose **Add Event > Event Example Function** (substitute the name of the 
-function in the _Interface_ you are implementing for **Example Function**). The _Event_ node is shown in the **Graph** tab. 
+1. **Right-click** in the **EventGraph** of the Blueprint and choose **Add Event > Event ...** 
 
 	![Implement Interface - Add Event](k2_interface_event_menu.png) ![Implement Interface - Event Node](k2_interface_event_graph.png)
 
-1. Create a network of nodes connected to the _Event's_ exec pin to give it functionality.
+	Select the type of Event you would like to add and the Event will be added to the graph. 
+
+1. Create a network of nodes connected to the Event's executable pin to give it functionality.
 
 	![Implement Interface - Event Implementation](k2_interface_event_implementation.png)
 
 	
 ### Calling Interface Functions
 
-The functions of the _Interface_ are executed through the use of **Interface Message** nodes, These look and act like regular 
+The functions of the Interface are executed through the use of **Interface Message** nodes, These look and act like regular 
 [](Engine/Blueprints/UserGuide/Functions/FunctionCalls), except that they can fail silently if the target does not implement 
-the _Interface_.
+the Interface.
 
 ![Implement Interface - Call Function](k2_interface_event_call.png)
 
@@ -135,6 +145,12 @@ Note in this image we are sending another Blueprint as the Target input. This wo
 	To specify a Blueprint that exists within the level, it is best to make a Public variable of type 'Actor', and assign the appropriate Blueprint in the value for that variable using the **Details** tab of the Level Editor.
 [/REGION]
 
+### Calling Interface Functions on Level Blueprints
+
+Interfaces can also allow for communication between sub-levels using the **Get Streaming Level** node. If the level is loaded and valid, the **Interface Message** will call the function with the level's Blueprint Script as the target.
+
+![](InterfaceGetLevelStreaming.png)
+
 ### Overriding an Interface Function
 
 One of the most important aspects of Interface functions is that, since they have no implementation of their own, their functionality can be independently defined for each Blueprint that implements the Interface.
@@ -143,9 +159,9 @@ Once an Interface has been implemented by a Blueprint, the functions of that int
 
 ![](iFunctionInExplorer.png)
 
-If you **double-click** that function, it will open up within the Blueprint Editor. This shows the local implementation for this function. Whatever functionality is defined here will be unique to this Blueprint. In this way, you can define custom functionality for each Blueprint that implements the interface.
+If you **Double-click** that function, it will open up within the Blueprint Editor. This shows the local implementation for this function. Whatever functionality is defined here will be unique to this Blueprint. In this way, you can define custom functionality for each Blueprint that implements the interface.
 
-For example, in this image we see a basic interface function as opened within a Blueprint. Right now, the function does nothing.
+For example, in this image we see a basic Interface function as opened within a Blueprint. Right now, the function does nothing.
 
 ![](DefaultFunction.png)(w:600)
 
