#UE4doc #docs:
Added engine version field
Changed font settings for multiple terms

diff --git a/Engine/Documentation/Source/Engine/Blueprints/UserGuide/Types/LevelBlueprint/LevelScript.INT.udn b/Engine/Documentation/Source/Engine/Blueprints/UserGuide/Types/LevelBlueprint/LevelScript.INT.udn
--- a/Engine/Documentation/Source/Engine/Blueprints/UserGuide/Types/LevelBlueprint/LevelScript.INT.udn
+++ b/Engine/Documentation/Source/Engine/Blueprints/UserGuide/Types/LevelBlueprint/LevelScript.INT.udn
@@ -4,6 +4,7 @@ Crumbs: %ROOT%, Engine/Blueprints, Engine/Blueprints/UserGuide/Types
 Description:Blueprints used for scripting level-specific events within maps.
 Related: Engine/Blueprints/UserGuide/Types
 Related: Engine/Blueprints
+version: 4.9
 
 [VAR:Topic]
 [OBJECT:Topic]
@@ -59,37 +60,37 @@ Level Blueprints also provide a control mechanism for level streaming and [Matin
 as for binding events to Actors placed within the level. 
 [/EXCERPT:Intro]
 
-For more information about the _Level Blueprint_ UI, see [](Engine/Blueprints\Editor\UIBreakdowns\LevelBPUI).
+For more information about the Level Blueprint UI, see [](Engine/Blueprints\Editor\UIBreakdowns\LevelBPUI).
 
 
 ## Default Level Blueprint
 
-Each game can specify the default _Level Blueprint_ class in the `DefaultGame.ini` config file. The _Level Blueprints_ 
+Each game can specify the default Level Blueprint class in the DefaultGame.ini config file. The Level Blueprints 
 for all new maps will be created using this class allowing for game-specific additions and functionality.
 
 ## Opening a Level Blueprint
 
-To open a _Level Blueprint_ for editing, click the **Blueprints** button in the **Level Editor Toolbar** and select **Open Level Blueprint**.
+To open a Level Blueprint for editing, click the **Blueprints** button in the **Level Editor Toolbar** and select **Open Level Blueprint**.
 
 ![Level Blueprint Button](toolbar_level_editor.png)
 
-This opens the _Level Blueprint_ in the **Blueprint Editor**:
+This opens the Level Blueprint in the **Blueprint Editor**:
 
 [REGION:fullwidth]
 ![](level_blueprint_editor.png)
 [/REGION]
 
-The **Blueprint Editor** only uses [**Graph**](Engine/Blueprints/Editor/UIComponents/GraphEditor) editor, **My Blueprints** panel, and **Details** panel. The **Class Defaults** Panel using the **Class Defaults** ![](Engine\Blueprints\UserGuide\Variables\classDefaults_button.png) button on the menu bar.
+The **Blueprint Editor** only uses [Graph Editor](Engine/Blueprints/Editor/UIComponents/GraphEditor), **My Blueprints** panel, and **Details** panel. The **Class Defaults** panel using the **Class Defaults** ![](Engine\Blueprints\UserGuide\Variables\classDefaults_button.png) button on the menu bar.
 
 ## Referencing Actors
 
-Often, you will need to connect a reference to an Actor to a **Target** pin on a node in your _Level Blueprint_. To get a node that contains an Actor reference:
+Often, you will need to connect a reference to an Actor to a **Target** pin on a node in your Level Blueprint. To get a node that contains an Actor reference:
 
 1. Select the Actor in the **Level Viewport** or in the **World Outliner**. 
 	
 	![](selected_actor.png)
 
-1. Open the _Level Blueprint_.
+1. Open the Level Blueprint.
 
 	![Level Blueprint Button](toolbar_level_editor.png)
 
@@ -100,11 +101,11 @@ Often, you will need to connect a reference to an Actor to a **Target** pin on a
 	
 Alternatively:
 
-1.  Drag and drop an _Actor_ from the **World Outliner** tab into a graph in the _Level Blueprint_.
+1.  Drag and drop an Actor from the **World Outliner** tab into a graph in the Level Blueprint.
 
 	![](add_reference_drag_drop.png)
 
-The _Actor_ reference node that appears can be connected to any compatible **Target** pin.
+The Actor reference node that appears can be connected to any compatible **Target** pin.
 
 ![](actor_reference.png)
 
@@ -114,7 +115,7 @@ In some cases, you will not need a reference node, as you can connect an output
 
 ## Adding Events
 
-There are two ways that [](Engine/Blueprints\UserGuide\Events) for a specific Actor can be added to a _Level Blueprint_. 
+There are two ways that [](Engine/Blueprints\UserGuide\Events) for a specific Actor can be added to a Level Blueprint. 
 
 1. Select the Actor in the **Level Viewport** or the **World Outliner**.
 
@@ -126,7 +127,7 @@ There are two ways that [](Engine/Blueprints\UserGuide\Events) for a specific Ac
 
 Alternatively, once you have your Actor selected:
 
-1. Open the _Level Blueprint_.
+1. Open the Level Blueprint.
 
 	![Level Blueprint Button](toolbar_level_editor.png)
 
