#UE4doc #docs:
Added engine version field
Changed font settings for multiple terms
Replaced multiple images due to outdated UI
Changed some steps as they were out of date

diff --git a/Engine/Documentation/Source/Engine/Blueprints/UserGuide/Types/MacroLibrary/MacroLibrary.INT.udn b/Engine/Documentation/Source/Engine/Blueprints/UserGuide/Types/MacroLibrary/MacroLibrary.INT.udn
--- a/Engine/Documentation/Source/Engine/Blueprints/UserGuide/Types/MacroLibrary/MacroLibrary.INT.udn
+++ b/Engine/Documentation/Source/Engine/Blueprints/UserGuide/Types/MacroLibrary/MacroLibrary.INT.udn
@@ -2,6 +2,7 @@ Availability:Public
 Title:Blueprint Macro Library
 Crumbs: %ROOT%, Engine/Blueprints, Engine/Blueprints/UserGuide/Types
 Description:Blueprints that declare and define macros for reuse in other Blueprints.
+version: 4.9
 
 [VAR:Topic]
 [OBJECT:Topic]
@@ -56,47 +57,47 @@ recompiled.
 
 ## Creating a Blueprint Macro Library
 
-_Blueprint Macro Libraries_ are stored in packages and can be created through the **Content Browser** like any other asset.
+Blueprint Macro Libraries are stored in packages and can be created through the **Content Browser** like any other asset.
 
 1. In the **Content Browser**, click on ![New Asset button](button_new_asset_blueprint.png).
 
-1. Select **Miscellaneous > Blueprint Macro Library** under **New Asset** from the menu that appears.
+1. Select **Blueprints > Blueprint Macro Library** under **Create Advanced Asset** from the menu that appears.
 
 	![](new_asset_macrolib.png)
 
-1. Choose a **Parent Class** for your _Blueprint Macro Library_.
+1. Choose a **Parent Class** for your Blueprint Macro Library.
 
 	![Choose a Parent Class](new_asset_parent_class.png)
 
-1. Your _Blueprint Macro Library_ will now appear in the **Content Browser**. Type a name for your _Blueprint Macro Library_ under its icon in the **Content Browser**.
+1. Your Blueprint Macro Library will now appear in the **Content Browser**. Type a name for your Blueprint Macro Library under its icon in the **Content Browser**.
 
 	![](name_macro_library.png)
 
-	When your _Blueprint Macro Library_ is first created, or when you make changes to it in the **Blueprint Editor**, an asterisk will be added to the _Blueprint Macro Library_'s icon in the **Content Browser**. This indicates that the _Blueprint Macro Library_ is not saved.
+	When your Blueprint Macro Library is first created, or when you make changes to it in the **Blueprint Editor**, an asterisk will be added to the Blueprint Macro Library's icon in the **Content Browser**. This indicates that the Blueprint Macro Library is not saved.
 
 	![](unsaved_macro_library.png)
 
 <!--
-1. To save the _Blueprint Macro Library_ directly from the **Content Browser**, click on **Save...** in the **Content Browser's** Navigation Bar.
+1. To save the Blueprint Macro Library directly from the **Content Browser**, click on **Save...** in the **Content Browser's** Navigation Bar.
 
 	![](save_macrolib_content_browser.png)
 
-1. In the **Save Content** window that appears, make sure the checkbox for your new _Blueprint Macro Library_ is selected, then click ![Save Selected](button_save_selected.png).
+1. In the **Save Content** window that appears, make sure the checkbox for your new Blueprint Macro Library is selected, then click ![Save Selected](button_save_selected.png).
 	
 	![](save_selected_macrolib.png)
 
-1. **Double-click** the _Blueprint Macro Library_'s icon to open it in the _Blueprint Editor_.
+1. **Double-click** the Blueprint Macro Library's icon to open it in the Blueprint Editor.
 -->
 
-There are two other methods by which to access the _Blueprint Macro Library_ Creation menu from the **Content Browser**.
+There are two other methods by which to access the Blueprint Macro Library Creation menu from the **Content Browser**.
 
 1. **Right-click** in the Asset View (right-hand) panel of the **Content Browser**, or on a folder in the Asset Tree (left-hand) panel of the **Content Browser**. 
 
-1. In the menu that appears, select **Miscellaneous > Blueprint Macro Library** under **New Asset**.  
+1. In the menu that appears, select **Blueprints > Blueprint Macro Library** under **Create Advanced Asset**.  
 
 	![](rt_click_content_browser_macrolib.png)
 
-1. The **Pick Parent Class** window will appear, and from this point the _Blueprint Macro Library_ creation process is the same as if you had used the **New Asset** button. 
+1. The **Pick Parent Class** window will appear, and from this point the Blueprint Macro Library creation process is the same as if you had used the **New Asset** button. 
 
 
 ## Blueprint Macros
@@ -114,9 +115,9 @@ There are two other methods by which to access the _Blueprint Macro Library_ Cre
 
 ### Using Macros stored in Blueprint Macro Libraries
 
-Storing _Macros_ in a _Blueprint Macro Library_ makes them accessible to _Blueprint Classes_ and _Level Blueprints_ throughout your project.  
+Storing Macros in a Blueprint Macro Library makes them accessible to Blueprint Classes and Level Blueprints throughout your project.  
 [INCLUDE:Engine/Blueprints/UserGuide/Macros#UniversalImplement]
 
-Once the _Macro_ node has been added to the graph, it behaves like any other node and the input and output pins can be wired accordingly. **Double-clicking** on a _Macro_ node in any _Blueprint_ graph
- will open the _Macro_'s graph.
+Once the Macro node has been added to the graph, it behaves like any other node and the input and output pins can be wired accordingly. **Double-clicking** on a Macro node in any Blueprint graph
+ will open the Macro's graph.
 
