#UE4doc #docs:
Removed italics from multiple terms in all 3 files
Removed bolding from a couple terms
Added engine version field to UserConstructionScript
Removed 'On this page...' section from UserConstructionScript

diff --git a/Engine/Documentation/Source/Engine/Blueprints/UserGuide/UserConstructionScript/UserConstructionScript.INT.udn b/Engine/Documentation/Source/Engine/Blueprints/UserGuide/UserConstructionScript/UserConstructionScript.INT.udn
--- a/Engine/Documentation/Source/Engine/Blueprints/UserGuide/UserConstructionScript/UserConstructionScript.INT.udn
+++ b/Engine/Documentation/Source/Engine/Blueprints/UserGuide/UserConstructionScript/UserConstructionScript.INT.udn
@@ -2,6 +2,7 @@ Availability:Public
 Title:Construction Script
 Crumbs: %ROOT%, Engine, Engine/Blueprints, Engine/Blueprints/UserGuide
 Description:Executed when instances of a Blueprint are created to perform initialization actions.
+version: 4.9
 
 [VAR:Topic]
 [OBJECT:Topic]
@@ -47,7 +48,6 @@ Description:Executed when instances of a Blueprint are created to perform initia
 [/REGION]
 
 
-[TOC(start:2 end:3)]
 
 [EXCERPT:Intro]
 
@@ -60,7 +60,7 @@ Description:Executed when instances of a Blueprint are created to perform initia
 Only Blueprint Classes contain **Construction Scripts**. Level Blueprints do not have them.
 [/REGION]
 
-The execution entry point into the **Construction Script** Graph is through a _ConstructionScript_ node that is always present.
+The execution entry point into the **Construction Script** Graph is through a **ConstructionScript** node that is always present.
 
 ![](construction_script.png)
 
@@ -68,4 +68,4 @@ The execution entry point into the **Construction Script** Graph is through a _C
 
 [INCLUDE:Shared/Glossary/G/#Glossary_Graph_Mode]
 
-See the [](Engine/Blueprints/Editor/UIComponents/GraphEditor) for a detailed guide to working with the **ConstructionScript** and other **Graphs** within Blueprints.
+See the [](Engine/Blueprints/Editor/UIComponents/GraphEditor) for a detailed guide to working with the **ConstructionScript** and other Graphs within Blueprints.
diff --git a/Engine/Documentation/Source/Shared/Glossary/C/Glossary_C.INT.udn b/Engine/Documentation/Source/Shared/Glossary/C/Glossary_C.INT.udn
--- a/Engine/Documentation/Source/Shared/Glossary/C/Glossary_C.INT.udn
+++ b/Engine/Documentation/Source/Shared/Glossary/C/Glossary_C.INT.udn
@@ -142,11 +142,11 @@ Construction Script
 	[PARAM:Definition]
 
 [EXCERPT:Glossary_Construction_Script]
-The **Construction Script** runs following the **Components** list when an instance of a _Blueprint Class_ 
-is created. It contains a node graph that is executed allowing the _Blueprint Class_ instance to perform initialization operations. 
+The **Construction Script** runs following the Components list when an instance of a Blueprint Class 
+is created. It contains a node graph that is executed allowing the Blueprint Class instance to perform initialization operations. 
 This can be extremely powerful as actions like performing traces into the world, setting meshes and materials, and so on can be 
-used to achieve context-specific setup. For instance, a light _Blueprint_ could determine what type of ground it is placed 
-upon and choose the correct mesh to use from a set of meshes or a fence _Blueprint_ could perform traces extending out in 
+used to achieve context-specific setup. For instance, a light Blueprint could determine what type of ground it is placed 
+upon and choose the correct mesh to use from a set of meshes or a fence Blueprint could perform traces extending out in 
 each direction to determine how long of a fence is needed to span the distance.
 [/EXCERPT:Glossary_Construction_Script]
 
diff --git a/Engine/Documentation/Source/Shared/Glossary/G/Glossary_G.INT.udn b/Engine/Documentation/Source/Shared/Glossary/G/Glossary_G.INT.udn
--- a/Engine/Documentation/Source/Shared/Glossary/G/Glossary_G.INT.udn
+++ b/Engine/Documentation/Source/Shared/Glossary/G/Glossary_G.INT.udn
@@ -49,7 +49,7 @@ Graph
 The **Graph** is a visual representation of a particular graph of nodes as it shows 
 all of the nodes contained in the graph as well as the connections between them. It provides 
 editing capabilities for adding and removing nodes, arranging nodes, and creating links between 
-nodes. Breakpoints can also be set in the Graph to aid in [debugging](Engine/Blueprints\UserGuide\Debugging) _Blueprints_.
+nodes. Breakpoints can also be set in the Graph to aid in [debugging](Engine/Blueprints\UserGuide\Debugging) Blueprints.
 [/EXCERPT:Glossary_Graph_Mode]
 
 [INCLUDE:Globals#bpModesRemoved]
