#UE4doc #docs:
Added engine version field
Changed font settings for multiple terms

diff --git a/Engine/Documentation/Source/Engine/Blueprints/UserGuide/Variables/Variables.INT.udn b/Engine/Documentation/Source/Engine/Blueprints/UserGuide/Variables/Variables.INT.udn
--- a/Engine/Documentation/Source/Engine/Blueprints/UserGuide/Variables/Variables.INT.udn
+++ b/Engine/Documentation/Source/Engine/Blueprints/UserGuide/Variables/Variables.INT.udn
@@ -2,6 +2,7 @@ Availability:Public
 Title:Blueprints User Guide: Variables
 Crumbs: %ROOT%, Engine, Engine/Blueprints, Engine/Blueprints/UserGuide
 Description: Properties that hold a value or reference an Object or Actor in the world.
+version: 4.9
 
 [VAR:Topic]
 [OBJECT:Topic]
@@ -73,12 +74,12 @@ Variables can be created in a variety of different types, including data types -
 | **Boolean**   | Red     | ![](RedWire.png) | Red variables represent Boolean (true/false) data. |
 | **Integer**   | Cyan    | ![](CyanWire.png) | Cyan variables represent integer data, or numbers without decimals such as 0, 152, and -226.
 | **Float**     | Green   | ![](GreenWire.png) | Green variables represent float data, or numbers with decimals such as 0.0553, 101.2887, and -78.322.
-| **String**    | Magenta | ![](MagentaWire.png) | Magenta variables represent string data, or a group of alphanumeric characters such as _Hello World_. |
+| **String**    | Magenta | ![](MagentaWire.png) | Magenta variables represent string data, or a group of alphanumeric characters such as Hello World. |
 | **Text**      | Pink	  | ![](PinkWire.png) | Pink variables represent displayed text data, especially where that text is localization aware. |
 | **Vector**    | Gold    | ![](GoldWire.png) | Gold variables represent vector data, or numbers consisting of float numbers for 3 elements or axes such as XYZ or RGB information. |
 | **Rotator**   | Purple  | ![](PurpleWire.png) | Purple variables represent Rotator data, which is a group of numbers that define rotation in 3D space. |
 | **Transform** | Orange  | ![](OrangeWire.png) | Orange variables represent Transform data, which combines translation (3D position), rotation, and scale. |
-| **Object**	| Blue    |  ![](BlueWire.png)	| Blue variables represent Objects, including _Lights_, _Actors_, _StaticMeshes_, _Cameras_, and _SoundCues_. |
+| **Object**	| Blue    |  ![](BlueWire.png)	| Blue variables represent Objects, including Lights, Actors, StaticMeshes, Cameras, and SoundCues. |
 
 [/EXCERPT:VarColor]
 
@@ -86,7 +87,7 @@ Variables can be created in a variety of different types, including data types -
 
 [EXCERPT:Panel]
 
-The **My Blueprint** tab allows custom variables to be added to the _Blueprint_ and lists any existing 
+The **My Blueprint** tab allows custom variables to be added to the Blueprint and lists any existing 
 variables, including component instance variables added in the [Components List](Engine/Blueprints/UserGuide/Components) 
 or variables created by promoting a value to a variable in the graph.
 
@@ -277,16 +278,16 @@ You can create a **Set** node (1 above) or **Get** node (2 above) for a variable
 
 Variable values can be set either as a part of the Blueprint node network or as a default value prior to execution. To set variable default values:
 
-1. Click the **Class Defaults** ![](classDefaults_button.png) Button on the Blueprint Editor Tool bar to open the Defaults in the **Details** panel.
+1. Click the **Class Defaults** button on the Blueprint Editor Tool bar to open the Defaults in the **Details** panel.
 
 1. In the **Details** panel, enter the desired default values to the right of your variable name.
 
 	![](blueprint_details_variable_callout.png)
 
-	Above we have highlighted our _Color_ variable in which we can set its default color. 
+	Above we have highlighted our Color variable in which we can set its default color. 
 
 [REGION:note]
-If you do not see your variable listed in the defaults, make sure you have clicked the  ![Compile Needed Button](Engine/Blueprints/Editor/UIComponents/Toolbar/bp_needs_compile.png) button. 
+If you do not see your variable listed in the defaults, make sure you have clicked the  **Compile** button. 
 [/REGION]
 
 
@@ -305,15 +306,15 @@ To rename a variable:
 
 ### Variable Properties
 
-You can set all the properties for a variable in the **Details** panel.  Some variables may have more properties than those shown here, such as **Expose to Matinee** for Vectors or a **Slider Range** for numeric variables like _Integers_ or _Floats_.
+You can set all the properties for a variable in the **Details** panel.  Some variables may have more properties than those shown here, such as **Expose to Matinee** for Vectors or a **Slider Range** for numeric variables like Integers or Floats.
 
 |Property|Description|
 |--------|-----------|
 |**Variable Type** | Set the variable type in the dropdown menu and determine whether or not the variable is an array.| 
-|**Editable**| Set whether or not it is possible to edit the value of the variable in the  **Class Defaults** and the _Blueprint_'s **Details** tab. |
+|**Editable**| Set whether or not it is possible to edit the value of the variable in the  **Class Defaults** and the Blueprint's **Details** tab. |
 |**Tooltip**| Set the tooltip for the variable. |
-|**Private** | Set whether or not this variable should be private and not modified by derived _Blueprints_. |
-|**Category** | Choose from an existing category or type a new Category name.  Setting the Category determines where the variable appears in the **Class Defaults**, the **My Blueprint** tab, and the _Blueprint_'s **Details** tab.| 
+|**Private** | Set whether or not this variable should be private and not modified by derived Blueprints. |
+|**Category** | Choose from an existing category or type a new Category name.  Setting the Category determines where the variable appears in the **Class Defaults**, the **My Blueprint** tab, and the Blueprint's **Details** tab.| 
 |**Replication** | Choose if the variable's value should or should not be replicated across clients, and if the value is replicated, if there should be a notification by way of a callback function. | 
 
 ### Getting & Setting Variable Values
