#UE4doc #docs:
Added engine version field

diff --git a/Engine/Documentation/Source/Engine/Components/Components.INT.udn b/Engine/Documentation/Source/Engine/Components/Components.INT.udn
--- a/Engine/Documentation/Source/Engine/Components/Components.INT.udn
+++ b/Engine/Documentation/Source/Engine/Components/Components.INT.udn
@@ -2,6 +2,7 @@ Availability:Public
 Title:Components
 Crumbs:%ROOT%, Engine
 Description: A compendium of the different types of Components designed to be used as sub-objects within an Actor.  
+version: 4.9
 
 [VAR:Topic]
 [OBJECT:Topic]
@@ -41,7 +42,7 @@ Description: A compendium of the different types of Components designed to be us
 [INCLUDE:Shared\Glossary\C#glossary_component]
 
 When adding Components to an Actor, you are putting together the bits and pieces that will make up that Actor as a whole who will be placed in your level(s). 
-For example, the wheels on a car, the steering wheel, the body of the car, the lights etc. would all be considered as Components while the car itself the Actor. 
+For example, the wheels on a car, the steering wheel, the body of the car, the lights etc. would all be considered as Components while the car itself is the Actor. 
 
 Once you have added the Components that define your Actor, it is important to note that you can place the Actor in your level even without providing any **Blueprint** script (or C++ Code) as to how the car should function (as in the real world, a car can exist as an object without a driver providing the instructions to tell the car what to do). 
 Since the "gas pedal" would be a Component, through script or code you would indicate that the gas pedal has been pressed and in turn provide the logic for the car to speed up. 
