#UE4doc #docs:
Added engine version field

diff --git a/Engine/Documentation/Source/Engine/Components/AI/AIComponents.INT.udn b/Engine/Documentation/Source/Engine/Components/AI/AIComponents.INT.udn
--- a/Engine/Documentation/Source/Engine/Components/AI/AIComponents.INT.udn
+++ b/Engine/Documentation/Source/Engine/Components/AI/AIComponents.INT.udn
@@ -6,6 +6,7 @@ Related: Gameplay/AI
 Related: Gameplay/Framework/Pawn
 Related: Gameplay/Framework/Controller
 Related: Engine/Blueprints/UserGuide/Components
+version: 4.9
 
 AI Components are a type of component that enable Pawns receive sensory like data from the environment, such as where noises are coming from or if the Pawn can see something.
 
@@ -19,7 +20,7 @@ A **PawnNoiseEmitterComponent** tracks noise event data used by **SensingCompone
 
 ## Pawn Sensing Component
 
-The Pawn **SensingComponent** encapsulates sensory (i.e. sight and hearing) settings and functionality for an Actor, allowing the Actor to see/hear Pawns in the world. It does nothing on network clients. 
+The Pawn SensingComponent encapsulates sensory (i.e. sight and hearing) settings and functionality for an Actor, allowing the Actor to see/hear Pawns in the world. It does nothing on network clients. 
 
 
 
