#UE4doc #docs:
Added engine version field
Removed bolding from a few terms

diff --git a/Engine/Documentation/Source/Engine/Components/Audio/AudioComponents.INT.udn b/Engine/Documentation/Source/Engine/Components/Audio/AudioComponents.INT.udn
--- a/Engine/Documentation/Source/Engine/Components/Audio/AudioComponents.INT.udn
+++ b/Engine/Documentation/Source/Engine/Components/Audio/AudioComponents.INT.udn
@@ -5,17 +5,18 @@ Description:An AudioComponent is used to create and control an instance of sound
 Related: Engine/Audio/Overview
 Related: Engine/Audio/WAV
 Related: Engine/Blueprints/UserGuide/Components
+version: 4.9
 
 
 An **AudioComponent** is used to create and control an instance of sound inside your game. 
 
 ## Audio Component
 
-The **AudioComponent** allows you to add a [Sound Wave](Engine/Audio/WAV) or [Sound Cue](Engine/Audio/Overview#soundcue) as a sub-object to an Actor, providing a sound source. 
+The AudioComponent allows you to add a [Sound Wave](Engine/Audio/WAV) or [Sound Cue](Engine/Audio/Overview#soundcue) as a sub-object to an Actor, providing a sound source. 
 
-For example, say you wanted a particle effect that resembles a fire in your level. While the particle effect alone can be used to convey the visuals of a fire, by adding an **AudioComponent** as a sub-Actor to it that plays an audio asset of a crackling fire, the fire effect will be that much more convincing. 
+For example, say you wanted a particle effect that resembles a fire in your level. While the particle effect alone can be used to convey the visuals of a fire, by adding an AudioComponent as a sub-Actor to it that plays an audio asset of a crackling fire, the fire effect will be that much more convincing. 
 
-You can change several of the properties or settings of the **AudioComponent** through **Blueprints** or **C++** during runtime such as fading the audio in or out, playing or stopping it entirely, adjusting the volume and other settings which can be found on the [Sound Actors](Engine/Audio/SoundActors) page under the different properties sections. 
+You can change several of the properties or settings of the AudioComponent through Blueprints or C++ during runtime such as fading the audio in or out, playing or stopping it entirely, adjusting the volume and other settings which can be found on the [Sound Actors](Engine/Audio/SoundActors) page under the different properties sections. 
 
 
 
