#UE4doc #docs:
Added engine version field
Removed bolding from multiple terms

diff --git a/Engine/Documentation/Source/Engine/Components/Camera/CameraComponents.INT.udn b/Engine/Documentation/Source/Engine/Components/Camera/CameraComponents.INT.udn
--- a/Engine/Documentation/Source/Engine/Components/Camera/CameraComponents.INT.udn
+++ b/Engine/Documentation/Source/Engine/Components/Camera/CameraComponents.INT.udn
@@ -6,20 +6,21 @@ Related: Gameplay/Framework/Camera
 Related: Engine/Actors/CameraActors
 Related: Gameplay/HowTo/UsingCameras
 Related: Engine/Blueprints/UserGuide/Components
+version: 4.9
 
 The **CameraComponent** (which adds a camera perspective to view from) and **SpringArmComponent** (which extends its children at a fixed distance then retracts when a collision occurs), when used together provide functionality for a 3rd person perspective that can dynamically adjust to your game world. 
 
 ## Camera Component
 
-The **CameraComponent** allows you to add a camera perspective as a sub-object to an Actor. The CameraComponent will provide information about the camera's properties if the **ViewTarget** is a **CameraActor** or an Actor that contains a CameraComponent and has its `bFindCameraComponentWhenViewTarget` option set to _True_. 
+The CameraComponent allows you to add a camera perspective as a sub-object to an Actor. The CameraComponent will provide information about the camera's properties if the **ViewTarget** is a **CameraActor** or an Actor that contains a CameraComponent and has its **bFindCameraComponentWhenViewTarget** option set to true. 
 
-For example, say you wanted to [switch between multiple cameras](Gameplay/HowTo/UsingCameras/Blueprints) in your level during gameplay. By using `SetViewTargetWithBlend` and a CameraActor, you can switch between each of your cameras and use the properties that you define within the CameraActor for each camera (including field of view, angle, or any Post Processing effects, etc.) 
+For example, say you wanted to [switch between multiple cameras](Gameplay/HowTo/UsingCameras/Blueprints) in your level during gameplay. By using **SetViewTargetWithBlend** and a CameraActor, you can switch between each of your cameras and use the properties that you define within the CameraActor for each camera (including field of view, angle, or any Post Processing effects, etc.) 
 
-A related property that can be set for any **Pawn** is `bTakeCameraControlWhenPossessed`, where the Pawn will automatically become the ViewTarget upon possession by a **PlayerController**. So for example if you had multiple Characters (which are a form of Pawn) that you wanted to switch between, and each of them had their own CameraComponents assigned to them providing the camera perspective to view from, you could use set `bTakeCameraControlWhenPossessed` to _True_ for each of them and whenever you switch between them, the CameraComponent for each will be used (where each could have its own properties and settings defined). 
+A related property that can be set for any **Pawn** is **bTakeCameraControlWhenPossessed**, where the Pawn will automatically become the ViewTarget upon possession by a PlayerController. So for example if you had multiple Characters (which are a form of Pawn) that you wanted to switch between, and each of them had their own CameraComponents assigned to them providing the camera perspective to view from, you could use set bTakeCameraControlWhenPossessed to true for each of them and whenever you switch between them, the CameraComponent for each will be used (where each could have its own properties and settings defined). 
 
 ## Spring Arm Component
 
-The **SpringArmComponent** tries to maintain its children at a fixed distance from the parent, but will retract the children if there is a collision, and spring back when there is no collision. Typically, the SpringArmComponent is used as a "camera boom" to keep the follow camera for a player from colliding into the world (without the SpringArmComponent, the CameraComponent will remain at the fixed distance specified regardless of objects that could potentially get in the way of it and the object it is attached to), see [Using a Spring Arm Component](Gameplay/HowTo/UsingCameras/SpringArmComponents). 
+The SpringArmComponent tries to maintain its children at a fixed distance from the parent, but will retract the children if there is a collision, and spring back when there is no collision. Typically, the SpringArmComponent is used as a "camera boom" to keep the follow camera for a player from colliding into the world (without the SpringArmComponent, the CameraComponent will remain at the fixed distance specified regardless of objects that could potentially get in the way of it and the object it is attached to), see [Using a Spring Arm Component](Gameplay/HowTo/UsingCameras/SpringArmComponents). 
 
 ![](camera_springArm.png)
 
