#UE4doc #docs:
Added engine version field
Changed font settings for multiple terms

diff --git a/Engine/Documentation/Source/Engine/Components/Lights/LightComponents.INT.udn b/Engine/Documentation/Source/Engine/Components/Lights/LightComponents.INT.udn
--- a/Engine/Documentation/Source/Engine/Components/Lights/LightComponents.INT.udn
+++ b/Engine/Documentation/Source/Engine/Components/Lights/LightComponents.INT.udn
@@ -5,12 +5,13 @@ Description:Descriptions of the different LightComponents that can be used in Un
 Related: Engine/Rendering/LightingAndShadows/LightTypes
 Related: Engine/Rendering/LightingAndShadows
 Related: Engine/Blueprints/UserGuide/Components
+version: 4.9
 
-Different types of **Light Components** will allow you to add a light source as a sub-object to an Actor, depending upon the effect you aim to achieve. Regardless of which Light Component type you choose, there are several light settings such as _Light Color_ or _Intensity_ that apply to all as well as individual settings that you can adjust (refer to each individual light for their specific settings).  
+Different types of **Light Components** will allow you to add a light source as a sub-object to an Actor, depending upon the effect you aim to achieve. Regardless of which Light Component type you choose, there are several light settings such as Light Color or Intensity that apply to all as well as individual settings that you can adjust (refer to each individual light for their specific settings).  
 
 ![](light_component.png)(w:720)
 
-_Above, a **PointLightComponent** is added to a **Blueprint** which also includes a **Static Mesh** to create a functional Ceiling Lamp._ 
+_Above, a PointLightComponent is added to a Blueprint which also includes a Static Mesh to create a functional Ceiling Lamp._ 
 
 ## Directional Light Component
 
@@ -28,14 +29,14 @@ See [](Engine/Rendering/LightingAndShadows/LightTypes/Point) for more informatio
 
 ## Sky Light Component
 
-**SkyLightComponents** capture the distant parts of your level (everything further than the `SkyDistanceThreshold`) and applies that to the scene as a light. That means the sky's appearance and its lighting / reflections will match. 
+**SkyLightComponents** capture the distant parts of your level (everything further than the SkyDistanceThreshold) and applies that to the scene as a light. That means the sky's appearance and its lighting / reflections will match. 
 
 See [](Engine/Rendering/LightingAndShadows/LightTypes/Spot) for more information. 
 
 
 ## Spot Light Component
 
-A **SpotLightComponent** emits a directional light from a single point in a cone shape. There are settings for adjusting the _Inner Cone Angle_ and _Outer Cone Angle_ as well as the other types of light settings such as _Intensity_, _Light Color_ and shadow settings. This type of light is excellent for a flash light due to its ability to define the inner/outer cone radius.
+A **SpotLightComponent** emits a directional light from a single point in a cone shape. There are settings for adjusting the Inner Cone Angle and Outer Cone Angle as well as the other types of light settings such as Intensity, Light Color and shadow settings. This type of light is excellent for a flash light due to its ability to define the inner/outer cone radius.
 
 See [](Engine/Rendering/LightingAndShadows/LightTypes/Skylight) for more information. 
 
