#UE4doc #docs:
Added engine version field
Changed font settings for multiple terms

diff --git a/Engine/Documentation/Source/Engine/Components/Movement/MovementComponents.INT.udn b/Engine/Documentation/Source/Engine/Components/Movement/MovementComponents.INT.udn
--- a/Engine/Documentation/Source/Engine/Components/Movement/MovementComponents.INT.udn
+++ b/Engine/Documentation/Source/Engine/Components/Movement/MovementComponents.INT.udn
@@ -1,10 +1,11 @@
 Availability:Public
 Title:Movement Components
 Crumbs:%ROOT%, Engine, Engine/Components
-Description:Any aspect of Movement, whether it�s a Character or a Projectile uses a Movement Component. 
+Description:Any aspect of Movement, whether it's a Character or a Projectile uses a Movement Component. 
 Related: Gameplay/Framework/Pawn/Character
 Related: Engine/Animation
 Related: Engine/Blueprints/UserGuide/Components
+version: 4.9
 
 **Movement Components** provide a form a movement to the Actor (or Character) that they are a sub-object of. 
 
@@ -14,7 +15,7 @@ The **CharacterMovementComponent** allows avatars not using rigid body physics t
 It is specific to **Characters** and cannot be implemented by any other class. It is automatically added when creating **Blueprints** based on the Character class and not manually added. 
 
 Properties that can be set include values for falling and walking friction, speeds for travel through air and water and across land, buoyancy, gravity scale, and the physics forces the Character can exert on Physics objects. 
-The **CharacterMovementComponent** also includes root motion parameters that come from the animation and are already transformed in world space, ready for use by physics. See [](Engine/Animation/RootMotion) for more information.  
+The CharacterMovementComponent also includes root motion parameters that come from the animation and are already transformed in world space, ready for use by physics. See [](Engine/Animation/RootMotion) for more information.  
 
 For information on working with Character Movement, see [](Gameplay/HowTo/CharacterMovement). 
 
@@ -22,10 +23,10 @@ For information on working with Character Movement, see [](Gameplay/HowTo/Charac
 ## Projectile Movement Component
 
 A **ProjectileMovementComponent** updates the position of another Component during its tick. Behavior such as bouncing after impacts and homing toward a target are supported by this type of Component. 
-Normally the root Component of the owing Actor is moved, however another Component may be selected (see [`SetUpdatedComponent`](https://docs.unrealengine.com/latest/INT/API/Runtime/Engine/GameFramework/UMovementComponent/SetUpdatedComponent/index.html)). 
+Normally the root Component of the owing Actor is moved, however another Component may be selected (see [SetUpdatedComponent](https://docs.unrealengine.com/latest/INT/API/Runtime/Engine/GameFramework/UMovementComponent/SetUpdatedComponent/index.html)). 
 If the Updated Component is simulating physics, only the initial launch parameters (when initial velocity is non-zero) will affect the projectile, and the physics simulation will take over from there. 
 
-An example of a **Blueprint** using a **ProjectileMovementComponent** is shown below (click for full image). 
+An example of a Blueprint using a ProjectileMovementComponent is shown below (click for full sized image). 
 
 [REGION:lightbox]
 [![](projectile.png)(w:820)](projectile.png)
@@ -35,7 +36,7 @@ An example of a **Blueprint** using a **ProjectileMovementComponent** is shown b
 
 The **RotatingMovementComponent** performs continuous rotation of a Component at a specific rotation rate. Rotation can optionally be offset around a pivot point. Important to note, collision testing is not performed during movement. 
 
-An example of using a **RotatingMovementComponent** might be in the form of an airplane's propellers, or a windmill, or even a series of planets revolving around the sun. 
+An example of using a RotatingMovementComponent might be in the form of an airplane's propellers, or a windmill, or even a series of planets revolving around the sun. 
 
 See [Blueprint Rotating Movement Component](https://wiki.unrealengine.com/Blueprint_Rotating_Movement_Component) for an example of how to utilize this component. 
 
