#UE4doc #docs:
Added engine version field
Changed font settings for multiple terms

diff --git a/Engine/Documentation/Source/Engine/Components/Paper2D/Paper2DComponents.INT.udn b/Engine/Documentation/Source/Engine/Components/Paper2D/Paper2DComponents.INT.udn
--- a/Engine/Documentation/Source/Engine/Components/Paper2D/Paper2DComponents.INT.udn
+++ b/Engine/Documentation/Source/Engine/Components/Paper2D/Paper2DComponents.INT.udn
@@ -6,15 +6,16 @@ Related: Engine\Paper2D
 Related: Engine\Paper2D\Flipbooks
 Related:Engine\Paper2D\Sprites
 Related: Engine/Blueprints/UserGuide/Components
+version: 4.9
 
 Paper 2D Components are used to add (or modify) 2D **Sprites** or **Flipbooks** as a sub-object to an Actor. 
 
 ## Paper Sprite Component
 
-A **PaperSpriteComponent** handles rendering and collision for a single instance of a `UPaperSprite` asset. 
+A **PaperSpriteComponent** handles rendering and collision for a single instance of a **UPaperSprite** asset. 
 This Component is created when you drag a Sprite asset from the **Content Browser** into a **Blueprint** or contained inside the Actor created when you drag one into a level. 
 
-An example of using this type of Component may come in the form of Sprite assets used to construct a level, for instance ledges or platforms, ladders, ramps etc. Placing those Sprite assets in your level creates a **PaperSpriteActor** that uses an instanced **PaperSpriteComponent** (or unique copy) of the selected Sprite asset. 
+An example of using this type of Component may come in the form of Sprite assets used to construct a level, for instance ledges or platforms, ladders, ramps etc. Placing those Sprite assets in your level creates a **PaperSpriteActor** that uses an instanced PaperSpriteComponent (or unique copy) of the selected Sprite asset. 
 
 For information on creating Sprites in Unreal Engine 4, see the [](Engine\Paper2D\Sprites\Creation) documentation. 
 
@@ -24,10 +25,10 @@ For information on creating Sprites in Unreal Engine 4, see the [](Engine\Paper2
 The **PaperFlipbookComponent** in short, is a wrapper for a **SourceFlipbook** (which is a series of Sprites played sequentially to create a 2D animation). 
 When creating a **Paper2DCharacter**, this component type is automatically added which allows you to create playable, animated 2D characters. 
 
-**PaperFlipbookComponents** can be posed arbitrarily in 3D, attached to other Components, or have other Components attached to them.
+PaperFlipbookComponents can be posed arbitrarily in 3D, attached to other Components, or have other Components attached to them.
 Each instance can specify a custom color that will be passed down to the Flipbook Material as a Vertex Color. They can also have a custom material specified that will override the default Material defined in the SourceFlipbook. 
 
-Through script you can change the current SourceFlipbook asset by calling the **SetFlipbook** function, but note that this will only work if the **Mobility** property is set to _Moveable_ (or if it is called during the construction of the Actor). 
+Through script you can change the current SourceFlipbook asset by calling the **SetFlipbook** function, but note that this will only work if the **Mobility** property is set to Moveable (or if it is called during the construction of the Actor). 
 You can also control play rate, play direction, looping, etc, with various other methods on the Component. 
 
 For information on setting up and working with Flipbook Components, see the [](Engine\Paper2D\Flipbooks\Components) documentation. 
