added commandline options to the docs

diff --git a/Engine/Documentation/Source/Programming/Rendering/Programming.INT.udn b/Engine/Documentation/Source/Programming/Rendering/Programming.INT.udn
--- a/Engine/Documentation/Source/Programming/Rendering/Programming.INT.udn
+++ b/Engine/Documentation/Source/Programming/Rendering/Programming.INT.udn
@@ -38,6 +38,10 @@ Useful command lines when working on rendering:
 |-d3ddebug|Enables the D3D11 debug layer, useful for catching API errors.|
 |-sm4|Forces Feature Level SM4 with the D3D11 RHI.|
 |-opengl3 / -opengl4|Forces use of OpenGL RHI at the specified feature level.|
+|-dx11 | currently the default on windows |
+|-dx12 | experimental |
+|-featureleveles2 | ignored when running with editor, there the UI has to be used |
+|-featureleveles31 | experimental, ignored when running with editor, there it needs to be enabled in Editor Preferences/Experimental/Rendering  |
 |-ddc=noshared|Prevents the use of network (shared) Derived Data Cache.  Can be useful when debugging shader caching issues.|
 
 ## Modules
