#UE4doc #docs - Adde Engine Version field. Added missing properties to match latest UI.

diff --git a/Engine/Documentation/Source/Engine/UI/LevelEditor/Viewports/ShowFlags/ShowFlags.INT.udn b/Engine/Documentation/Source/Engine/UI/LevelEditor/Viewports/ShowFlags/ShowFlags.INT.udn
--- a/Engine/Documentation/Source/Engine/UI/LevelEditor/Viewports/ShowFlags/ShowFlags.INT.udn
+++ b/Engine/Documentation/Source/Engine/UI/LevelEditor/Viewports/ShowFlags/ShowFlags.INT.udn
@@ -2,6 +2,7 @@ Availability:Public
 Title:Viewport Show Flags
 Crumbs: %ROOT%, Engine, Engine/UI, Engine/UI/LevelEditor/Viewports
 Description:Descriptions of all of the available Show Flags for the viewports.
+Version: 4.9
 
 
 [VAR:Topic]
@@ -62,7 +63,6 @@ The Show Flags are found under the Viewport **Show** menu. They are a means by w
 |---|---|
 |[REGION:tablesection]Name[/REGION]|[REGION:tablesection]Description[/REGION]|
 |**Anti-aliasing**|Shows/hides the effect of anti-aliasing.|
-|**Atmosphere**|Shows/hides any Atmospheric Fog in the level.|
 |**BSP** |Shows/hides any Brushes in the level. _(**Alt + Q**)_ |
 |**Collision**|Shows/hides collision objects in the level. _(**Alt + C**)_ |
 |**Decals**|Shows/hides decals in the scene.|
@@ -97,6 +97,15 @@ The Show Flags are found under the Viewport **Show** menu. They are a means by w
 |**Tonemapper**|Shows/hides Tonemapper effects. |
 |**Vignette**|Shows/hides Vignette effects.|
 
+## Light Types
+|Lighting Types||
+|---|---|
+|[REGION:tablesection]Name[/REGION]|[REGION:tablesection]Description[/REGION]|
+|**Directional Lights**|Shows/hides the result of lighting from the Directional Light actor in the scene.|
+|**Point Lights**|Shows/hides the result of lighting from the Point Light actor in the scene.|
+|**Sky Lighting**|Shows/hides the result of lighting from the Sky Light actor in the scene.|
+|**Spot Lights**|Shows/hides the result of lighting from the Spot Light actor in the scene.|
+
 ## Lighting Components
 
 |Lighting Components||
@@ -105,14 +114,9 @@ The Show Flags are found under the Viewport **Show** menu. They are a means by w
 |**Ambient Occlusion**|Shows/hides ambient occlusion in the scene.|
 |**Diffuse**|Shows/hides the result of diffuse lighting in the scene, which contains most of the scene color.|
 |**Direct Lighting**|Shows/hides the result of direct lighting, which contains most of the scene specularity.|
-|**Directional Lights**|Shows/hides the result of directional lights in the scene.|
 |**Dynamic Shadows**|Shows/hides dynamic shadows in the scene.|
 |**Global Illumination**|Shows/hides the result of global illumination in the scene.|
-|**Point Lights**|Shows/hides the result of point lights in the scene.|
-|**Sky Lighting**|Shows/hides the result of sky lights in the scene.|
 |**Specular**|Shows/Hides the result of specularity in the scene.|
-|**Spot Lights**|Shows/hides the result of spot lights in the scene.|
-
 
 ## Lighting Features
 
@@ -120,12 +124,15 @@ The Show Flags are found under the Viewport **Show** menu. They are a means by w
 |---|---|
 |[REGION:tablesection]Name[/REGION]|[REGION:tablesection]Description[/REGION]|
 |**Ambient Cubemap**|Shows/hides the result of the ambient cubemap.|
+|**Distance Field Ambient Occlusion**|Shows/hides the result of the AO generated from Distance Field meshes.|
+|**Distance Field Global Illumination**|Shows/hides the result of the GI generated from Distance Field meshes.|
 |**Indirect Lighting Cache**|Shows/hides the result of indirect lighting on dynamic objects.|
 |**Light Functions**|Shows/hides the result of light function materials on lights.|
 |**Light Shafts**|Shows/hides the result of light shafts.|
 |**Reflection Environment**|Shows/hides the result of the reflection environment.|
-|**Reflection Environment Lightmap Mixing**|Shows/hides the result of any lightmap data that is the result of the reflection environment.|
+|**Screen Space Ambient Occlusion**|Shows/hides the result of any screen space ambient occlusion.|
 |**Screen Space Reflections**|Shows/hides the result of screen space reflections.|
+|**Subsurface Scattering**| Shows/hides the result of any subsurface scattering within the screen space.|
 |**Textured Light Profiles (IES Textures)**|Shows/hides the result of IES lighting profiles.|
 
 
@@ -135,7 +142,8 @@ The Show Flags are found under the Viewport **Show** menu. They are a means by w
 |---|---|
 |[REGION:tablesection]Name[/REGION]|[REGION:tablesection]Description[/REGION]|
 |**AI Debug**|This shows the paths an AI-driven pawn is using to navigate along a Nav Mesh.|
-|**Blur GBuffer**|Shows/hides the results of Blue GBuffer Intensity, a property found on Post Process Volumes. This feature blurs the GBuffer normals and adjusts the GBuffer roughness to reduce specular aliasing. It is, however, a high performance cost for the result and you should generally consider just using Temporal AA.  |
+|**Composite Editor Primitives**|Show/hide editor primitives in the scene.|
+|**Gameplay Debug**|This shows the paths an AI-driven pawn is using to navigate along a Nav Mesh.|
 |**GBuffer Hints (material attributes)**|Colorizes objects that have improper Materials. Yellow for objects that are impossibly non-reflective, red for objects that emit more light than they receive.|
 |**Hit Proxies**|Shows/hides the proxy shapes used by hit detection.|
 |**Log Visualizer**|Draws any data to the screen that is required by the Log Visualizer system.|
@@ -145,6 +153,7 @@ The Show Flags are found under the Viewport **Show** menu. They are a means by w
 |**Test Image**|Shows a test image for configuration of monitor and graphics settings.|
 |**Vector Fields**|Shows/hides any vector fields in the scene.|
 
+|**Blur GBuffer**|Shows/hides the results of Blue GBuffer Intensity, a property found on Post Process Volumes. This feature blurs the GBuffer normals and adjusts the GBuffer roughness to reduce specular aliasing. It is, however, a high performance cost for the result and you should generally consider just using Temporal AA.  |
 
 ## Visualize
 
@@ -152,10 +161,21 @@ The Show Flags are found under the Viewport **Show** menu. They are a means by w
 |---|---|
 |[REGION:tablesection]Name[/REGION]|[REGION:tablesection]Description[/REGION]|
 |**Adaptive Depth of Field**|Applies a color coding filter to the view such that objects that are blurred by Depth of Field are tinted green and only the area in focus receives its normal color.|
+|**Bloom**| Shows a visualization window that expresses the scene histogram, as well as how settings are affecting bloom. |
 |**Depth of Field Layers**|Applies a color coding filter to the view such that objects blurred by near DOF are green, objects blurred by far DOF are blue. Objects in focus appear black. |
+|**Distance Field Ambient Occlusion**| Shows a visualization of ambient occlusion generated based on mesh distance fields. |
+|**Distance Field Global Illumination**| Shows a visualization of global illumination generated based on mesh distance fields. |
 |**HDR**|Shows a visualization window that expresses the scene histogram, as well as how exposure settings are affecting brightness. |
+|**Light Propagation Volumes**| Shows a visualization of the volume that is used for light propagation volumes. |
+|**Mesh Distance Fields**| Shows a visualization of the generated mesh distance fields for static meshes in the scene. |
 |**Motion Blur**|When Motion Blur is on, this changes the view to a Motion Blur visualizer, allowing you to see the motion vectors being used to calculate blur.|
-| [**Out of Bounds Pixels**](Engine/Performance/Visualizers/OutOfBoundsPixels) | If a pixel of an object exists outside that object's bounding volume, it can lead to flickering or the object disapearing when the bounding volume (but not the mesh) is off camera. This visualizer colorizes those pixels in blue, yellow, and white. |
+| [Out of Bounds Pixels](Engine/Performance/Visualizers/OutOfBoundsPixels) | If a pixel of an object exists outside that object's bounding volume, it can lead to flickering or the object disapearing when the bounding volume (but not the mesh) is off camera. This visualizer colorizes those pixels in blue, yellow, and white. |
+|**Precomputed Visibility Cells**|Shows a visualization of precomputed visibility cells when used with Precomputed Visibility Volumes. |
+|**Preview Shadows Indicator**| When lighting has not been built for Static/Stationary lights the "Preview" text will show in unbuilt shadows. |
+|**Screen Space Reflections**| A visualization for Screen Space Reflections.|
+|**Subsurface Scattering (Screen Space)**| A visualization mode that allows you to see subsurface scattering within screen space. |
+|**Volume Lighting Samples**| Shows a visualization of the volume lighting samples that are placed within a Lightmass Importance Volume when static lighting is built. |
+
 
 
 ## Advanced
@@ -163,14 +183,18 @@ The Show Flags are found under the Viewport **Show** menu. They are a means by w
 |Advanced||
 |---|---|
 |[REGION:tablesection]Name[/REGION]|[REGION:tablesection]Description[/REGION]|
+|**Atmospheric Fog**|Shows/hides atmospheric fog Actors in the scene.|
 |**Audio Radius**|Shows/hides the radius of any audio Actors in the scene.|
 |**Billboard Sprites**|Shows/hides any billboard sprites in the level. This does not extend to sprite particles, only Billboard Components.|
 |**Bounds**|Shows/hides bounds on selected objects.|
 |**BSP Split**|Shows/hides Brush surface splits on world geometry.|
-|**Builder Brush**|Shows/hides the builder brush.|
+|**Camera Aspect Ratio Bars**|Shows/hides the bars used with specific aspect ratios set in the camera actor.|
 |**Camera Frustums**|Shows/hides view frustums of any cameras in the scene.|
+|**Camera Safe Frames**| When looking through a Camera Actor, this shows the safe frames for that camera. Safe frames are areas a set distance from the edge of the screen where you can guarantee the image will be visible.|
 |**Constraints**|Shows/hides rigid body constraints. TODO|
-|**High Res Screenshot Mask**|TODO|
+|**Deferred Lighting**|Shows/hides the deferred lighting in the scene.|
+|**Foliage**| Shows/hides any Foliage Actors in the scene.|
+|**Grass**| Shows/hides any Grass Type Actors in the scene.|
 |**Instanced Static Meshes**|Shows/hides any instanced Static Meshes in the scene.|
 |**Large Vertices**|Shows/hides larger-than-usual vertices on any selected Brushes or Static Meshes.|
 |**Level Coloration**|Color codes the scene such that each streamed level gets another color. Irrelevant if not using level streaming.|
@@ -179,8 +203,9 @@ The Show Flags are found under the Viewport **Show** menu. They are a means by w
 |**LOD Parenting**|Use LOD parenting, MinDrawDistance, etc. If disabled, will show LOD parenting lines.|
 |**Mesh Edges**|This overlays the mesh wireframe on top of the view.|
 |**Mode Widgets**|Shows/hides the transform widgets, such as the Move tool.|
+|**Paper 2D Sprites**| Shows/hides the radius of any Paper 2D Sprite Actors in the scene.|
+|**Precomputed Visibility**| Shows/hides the radius of any Precomputed Visibility volumes in the scene.|
 |**Render (3D) Text**|Whether or not to render any Text Actors or Components that may be in the scene.|
-|**Safe Frames**|When looking through a Camera Actor, this shows the safe frames for that camera. Safe frames are areas a set distance from the edge of the screen where you can guarantee the image will be visible. |
 |**Selection**|Shows/hides the result of the Selection Highlight.|
 |**Separate Translucency**|This shows/hides any objects that are making use of Separate Translucency in their Materials.|
 |**Shadow Frustums**|Shows/hides a selection of color-coded boxes which represent the shadow frustums used to cast shadows from each light.|
@@ -202,7 +227,10 @@ The Show Flags are found under the Viewport **Show** menu. They are a means by w
 |**Hide All**|Hides all types of volumes in the scene.|
 |**Audio**|Shows/hides Audio Volumes.|
 |**Blocking**|Shows/hides Blocking Volumes.|
+|**Camera Blocking Volume**| Shows/hides Camera Blocking Volumes. |
 |**Cull Distance**|Shows/hides Cull Distance Volumes.|
+|**Hierarchical LOD Volume**| Shows/hides Hierarchical LOD Volumes. |
+|**Kill Z Volume**| Shows/hides Kill Z Volumes. |
 |**Level Streaming**|Shows/hides Level Streaming Volumes.|
 |**Lightmass Character Indirect Detail**|Shows/hides Lightmass Character Indirect Detail Volumes.|
 |**Lightmass Importance**|Shows/hides Lightmass Importance Volumes.|
@@ -213,6 +241,8 @@ The Show Flags are found under the Viewport **Show** menu. They are a means by w
 |**Post Process**|Shows/hides Post Process Volumes.|
 |**Precomputed Visibility Override**|Shows/hides Precomputed Visibility Override Volumes.|
 |**Precomputed Visibility**|Shows/hides Precomputed Visibility Volumes.|
+|**Procedural Foliage Blocking Volume**| Shows/hides Procedural Foliage Blocking Volumes. |
+|**Procedural Foliage Volume**| Shows/hides Procedural Foliage Volumes. |
 |**Trigger**|Shows/hides Trigger Volumes.|
 
 
@@ -237,20 +267,22 @@ Billboard sprites are generally in-editor icons used for placement of various ob
 |[REGION:tablesection]Name[/REGION]|[REGION:tablesection]Description[/REGION]|
 |**Show All**|Shows all types of billboard sprites.|
 |**Hide All**|Hides all types of billboard sprites.|
+|**Characters**|Shows/hides Character billboard sprites.|
 |**Decals**|Shows/hides Decal billboard sprites.|
 |**Effects**|Shows/hides Effects billboard sprites, which represent the location of Emitter Actors.|
 |**Fog**|Shows/hides Fog Actor billboard sprites.|
+|**FTests**|Shows/hides FTest billboard sprites.|
 |**Info**|Shows/hides Info billboard sprites.|
-|**Landscape**|Shows/hides Landscape billboard sprites.|
 |**Lighting**|Shows/hides Light Actor billboard sprites.|
 |**Materials**|Shows/hides Materials billboard sprites.|
+|**Matinee**|Shows/hides Matinee billboard sprites.|
 |**Misc**|Shows/hides miscellaneous billboard sprites which do not fall into the other categories.|
 |**Navigation**|Shows/hides Navigation billboard sprites.|
 |**Notes**|Shows/hides Notes billboard sprites.|
 |**Physics**|Shows/hides Physics billboard sprites.|
 |**Player Start**|Shows/hides Player Start billboard sprites.|
+|**Sky**|Shows/hides Sky billboard sprites.|
 |**Sounds**|Shows/hides Sounds billboard sprites.|
 |**Target Points**|Shows/hides Target Points billboard sprites.|
 |**Triggers**|Shows/hides Triggers billboard sprites.|
-|**Unreal Matinee**|Shows/hides Unreal Matinee billboard sprites.|
 |**Wind**|Shows/hides Wind billboard sprites.|
