#UE4doc #docs - Added Engine Version field. Replaced two image tables with comparison sliders.

diff --git a/Engine/Documentation/Source/Engine/UI/LevelEditor/Viewports/ViewModes/ViewModes.INT.udn b/Engine/Documentation/Source/Engine/UI/LevelEditor/Viewports/ViewModes/ViewModes.INT.udn
--- a/Engine/Documentation/Source/Engine/UI/LevelEditor/Viewports/ViewModes/ViewModes.INT.udn
+++ b/Engine/Documentation/Source/Engine/UI/LevelEditor/Viewports/ViewModes/ViewModes.INT.udn
@@ -2,6 +2,7 @@ Availability:Public
 Title:View Modes
 Crumbs: %ROOT%, Engine, Engine/UI, Engine/UI/LevelEditor/Viewports
 Description:Explanations for the available View Modes within the viewports.
+Version: 4.9
 
 
 [VAR:Topic]
@@ -333,11 +334,14 @@ The **Shading Model** mode shows value of the Shading Model property for each Ma
 
 ### Material AO
 
-[REGION:imagetable]
-|![](Buffer_MaterialAO_1.png)(w:%WireSize%)|![](Buffer_MaterialAO_2.png)(w:%WireSize%)|
-|---|---|
-|Scene in Lit Mode (Game View On) |Scene in Buffer Material AO Mode (Game View On)|
-[/REGION]
+[OBJECT:ComparisonSlider]
+	[PARAM:before]
+	![Scene in Lit Mode (Game View On)](Buffer_MaterialAO_1.png)(w:500)
+	[/PARAM]
+	[PARAM:after]
+	![Scene in Buffer Material AO Mode (Game View On)](Buffer_MaterialAO_2.png)(w:500)
+	[/PARAM]
+[/OBJECT]
 
 * View Mode Hotkey: **None (Menu only by default)**
 
@@ -359,11 +363,14 @@ Note: Materials should generally have Metallic values of 0 or 1, not in between.
 
 ### Opacity
 
-[REGION:imagetable]
-|![](Buffer_Opacity_1.png)(w:%WireSize%)|![](Buffer_Opacity_2.png)(w:%WireSize%)|
-|---|---|
-|Scene in Lit Mode (Game View On) |Scene in Buffer Material AO Mode (Game View On)|
-[/REGION]
+[OBJECT:ComparisonSlider]
+	[PARAM:before]
+	![Scene in Lit Mode](Buffer_Opacity_1.png)(w:500)
+	[/PARAM]
+	[PARAM:after]
+	![Scene in Buffer Material Opacity Mode (Game View On)](Buffer_Opacity_2.png)(w:500)
+	[/PARAM]
+[/OBJECT]
 
 * View Mode Hotkey: **None (Menu only by default)**
 
