#UE4doc #docs - Added engine version field.

diff --git a/Engine/Documentation/Source/Engine/UI/LevelEditor/Viewports/ViewportControls/ViewportControls.INT.udn b/Engine/Documentation/Source/Engine/UI/LevelEditor/Viewports/ViewportControls/ViewportControls.INT.udn
--- a/Engine/Documentation/Source/Engine/UI/LevelEditor/Viewports/ViewportControls/ViewportControls.INT.udn
+++ b/Engine/Documentation/Source/Engine/UI/LevelEditor/Viewports/ViewportControls/ViewportControls.INT.udn
@@ -2,6 +2,7 @@ Availability:Public
 Title:Viewport Controls
 Crumbs: %ROOT%, Engine, Engine/UI, Engine/UI/LevelEditor/Viewports
 Description:Learn about the various control schemes for the editor's viewports.
+Version: 4.9
 
 [VAR:Topic]
 [OBJECT:Topic]
