#UE4doc #docs - Added engine version field.

diff --git a/Engine/Documentation/Source/Engine/UI/LevelEditor/Viewports/ViewportOptions/ViewportOptions.INT.udn b/Engine/Documentation/Source/Engine/UI/LevelEditor/Viewports/ViewportOptions/ViewportOptions.INT.udn
--- a/Engine/Documentation/Source/Engine/UI/LevelEditor/Viewports/ViewportOptions/ViewportOptions.INT.udn
+++ b/Engine/Documentation/Source/Engine/UI/LevelEditor/Viewports/ViewportOptions/ViewportOptions.INT.udn
@@ -2,6 +2,7 @@ Availability:Public
 Title:Viewport Options
 Crumbs: %ROOT%, Engine, Engine/UI, Engine/UI/LevelEditor/Viewports
 Description:Explanations of the various options available within the editor viewports.
+Version: 4.9
 
 
 [VAR:Topic]
