#UE4doc #docs - Added engine version field. Bolded terms where needed.

diff --git a/Engine/Documentation/Source/Engine/UI/LevelEditor/Viewports/ViewportToolbar/ViewportToolbar.INT.udn b/Engine/Documentation/Source/Engine/UI/LevelEditor/Viewports/ViewportToolbar/ViewportToolbar.INT.udn
--- a/Engine/Documentation/Source/Engine/UI/LevelEditor/Viewports/ViewportToolbar/ViewportToolbar.INT.udn
+++ b/Engine/Documentation/Source/Engine/UI/LevelEditor/Viewports/ViewportToolbar/ViewportToolbar.INT.udn
@@ -2,6 +2,7 @@ Availability:Public
 Title:Viewport Toolbar
 Crumbs: %ROOT%, Engine, Engine/UI, Engine/UI/LevelEditor/Viewports
 Description:A rundown of the various tools available within the viewports integrated toolbar.
+Version: 4.9
 
 
 [VAR:Topic]
@@ -81,7 +82,7 @@ Each one of these tools also has settings to handle snapping. For more informati
 
 ![](ViewportToolbar_CameraSpeed.png)
 
-The Camera Speed slider controls how fast the camera moves through your scenes in the Perspective viewport. 
+The **Camera Speed** slider controls how fast the camera moves through your scenes in the Perspective viewport. 
 
 ![](SettingCameraSpeed.png)
 
@@ -92,7 +93,7 @@ While using the WASD controls, you can scroll the mouse wheel to speed up or slo
 
 ![](ViewportToolbar_Split.png)
 
-The Maximize button, located in the upper-right corner of the viewport, is used to toggle between a single viewport and a split 4-view style. 
+The **Maximize** button, located in the upper-right corner of the viewport, is used to toggle between a single viewport and a split 4-view style. 
 
 [VAR:ViewportImage]
 w:450
