#UE4doc #docs - Added Engine Version field. Updated Editor Preferences image. Added missing property from Editor Preferences. Bolded terms where needed.

diff --git a/Engine/Documentation/Source/Engine/UI/SourceControl/SourceControl.INT.udn b/Engine/Documentation/Source/Engine/UI/SourceControl/SourceControl.INT.udn
--- a/Engine/Documentation/Source/Engine/UI/SourceControl/SourceControl.INT.udn
+++ b/Engine/Documentation/Source/Engine/UI/SourceControl/SourceControl.INT.udn
@@ -2,6 +2,7 @@ Availability: Public
 Title:Source Control
 Crumbs:%ROOT%, Engine, Engine/UI
 Description:This document overviews the Source Control features.
+Version: 4.9
 
 [VAR:TopicCompact]
 [OBJECT:TopicCompact]
@@ -41,6 +42,7 @@ You can activate source control via the editor's main Preferences window (**Edit
 
 | Item | Description |
 | --- | --- |
+| **Automatically Checkout on Asset Modification** | When checked, this will automatically check out any asset that has been modified.  |
 | **Prompt for Checkout on Package Modification** | When checked, a prompt will appear when you make changes to a source controlled-package, asking if you would like to check out (lock) that package. |
 | **Add New Files when Modified** | This will add new files into source control when they are modified. |
 | **Use Global Settings** | Use global source control login settings, rather than per-project.  Changing this will require you to login again. |
@@ -80,10 +82,10 @@ While source control is active, the following context menu will be available whe
 
 | Item | Description |
 | --- | --- |
-| Check Out | Checks out (locks) the asset for editing. This prevents other users from editing the asset at the same time. |
-| Refresh | Refreshes the source control status of the asset. |
-| History | Gives a revision history list of the selected asset, allowing you to see previous edits. |
-| Diff Against Depot | This allows you to check the asset against the version currently stored in the source control depot. |
+| **Check Out** | Checks out (locks) the asset for editing. This prevents other users from editing the asset at the same time. |
+| **Refresh** | Refreshes the source control status of the asset. |
+| **History** | Gives a revision history list of the selected asset, allowing you to see previous edits. |
+| **Diff Against Depot** | This allows you to check the asset against the version currently stored in the source control depot. |
 
 ## Checking Out and In
 
