#UE4doc #docs - Added Engine Version Field.

diff --git a/Engine/Documentation/Source/Engine/UI/Tools/EditorTools.INT.udn b/Engine/Documentation/Source/Engine/UI/Tools/EditorTools.INT.udn
--- a/Engine/Documentation/Source/Engine/UI/Tools/EditorTools.INT.udn
+++ b/Engine/Documentation/Source/Engine/UI/Tools/EditorTools.INT.udn
@@ -3,7 +3,7 @@ Title:Tools and Systems
 Crumbs:%ROOT%, Engine, Engine/UI
 Description:Individual tools and systems built into Unreal Editor.
 Navigation:topic
-
+Version: 4.9
 
 [REGION:features half]
 [OBJECT:Role]
