#UE4doc #docs - Added Engine Version field.

diff --git a/Engine/Documentation/Source/Engine/UI/Tools/Simplygon/Simplygon.INT.udn b/Engine/Documentation/Source/Engine/UI/Tools/Simplygon/Simplygon.INT.udn
--- a/Engine/Documentation/Source/Engine/UI/Tools/Simplygon/Simplygon.INT.udn
+++ b/Engine/Documentation/Source/Engine/UI/Tools/Simplygon/Simplygon.INT.udn
@@ -2,6 +2,7 @@ Availability:Licensee
 Title:Simplygon
 Crumbs:%ROOT%, Engine
 Description:Third-party tools for mesh simplification and merging that have been integrated into Unreal Engine 4.
+Version: 4.9
 
 [REGION:raw]
 ![](SimplygonLogo.png)
