#UE4doc #docs - Added engine version field. Added Skill Level field. Removed Overview Header.

diff --git a/Engine/Documentation/Source/Engine/UI/Tools/Simplygon/Compiling/Compiling.INT.udn b/Engine/Documentation/Source/Engine/UI/Tools/Simplygon/Compiling/Compiling.INT.udn
--- a/Engine/Documentation/Source/Engine/UI/Tools/Simplygon/Compiling/Compiling.INT.udn
+++ b/Engine/Documentation/Source/Engine/UI/Tools/Simplygon/Compiling/Compiling.INT.udn
@@ -2,11 +2,13 @@ Availability: Licensee
 Title:Compiling Simplygon
 Crumbs:%ROOT%, Engine, Engine/UI/Tools/Simplygon
 Description:Instructions on compiling and terms of use for the Simplygon integration in Unreal Engine 4.
+Version: 4.9
+SkillLevel: Advanced
 
 [TOC(start:2)]
 
 
-## Overview
+
 
 Simplygon ships with the Editor but licensees must first obtain a license from Donya Labs (http://simplygon.com/) to use this feature. Please contact them for further details on obtaining the license and SDK.
 
