#UE4doc #docs - Added Engine Version Field.

diff --git a/Engine/Documentation/Source/Engine/UMG/UMG.INT.udn b/Engine/Documentation/Source/Engine/UMG/UMG.INT.udn
--- a/Engine/Documentation/Source/Engine/UMG/UMG.INT.udn
+++ b/Engine/Documentation/Source/Engine/UMG/UMG.INT.udn
@@ -2,6 +2,7 @@ Availability:Public
 Title: UMG UI Designer
 Crumbs:%ROOT%, Engine
 Description:A guide to using Unreal Motion Graphics to create UI elements. 
+Version: 4.9
 
 [VAR:Topic]
 [OBJECT:Topic]
