#UE4doc #docs - Added Engine Version field. Bolder terms.

diff --git a/Engine/Documentation/Source/Gameplay/Framework/Camera/Camera.INT.udn b/Engine/Documentation/Source/Gameplay/Framework/Camera/Camera.INT.udn
--- a/Engine/Documentation/Source/Gameplay/Framework/Camera/Camera.INT.udn
+++ b/Engine/Documentation/Source/Gameplay/Framework/Camera/Camera.INT.udn
@@ -2,6 +2,7 @@ Availability:Public
 Title:Camera
 Crumbs:%ROOT%, Gameplay, Gameplay/Framework
 Description:The Camera represents the player's point of view; how the player sees the world.
+Version: 4.9
 
 The **Camera** represents the player's point of view; how the player sees the world. For this reason, 
 cameras only have relevance to human-controlled players. The **PlayerController** specifies a camera class 
@@ -14,21 +15,21 @@ For basic examples of how to work with Cameras, refer to the [Using Cameras](Gam
 
 ## CameraComponent and CameraActor
 
-All of the camera's properties and behavior are set up in the CameraComponent; the CameraActor class primarily acts as a wrapper for the CameraComponent, so that the camera can be placed directly in the
+All of the camera's properties and behavior are set up in the **CameraComponent**; the `CameraActor` class primarily acts as a wrapper for the CameraComponent, so that the camera can be placed directly in the
 level rather than within another class.  
 
 In the CameraComponent, it is possible to set whether the camera is in perspective mode or orthographic mode.  The vertical field of view (FOV) can be set for perspective mode, 
 and the width in world units can be set for orthographic mode.  For both modes, the aspect ratio can be designated and preset aspect ratios for common devices and display types are 
 provided.  [Post process effects](Engine/Rendering\PostProcessEffects) can be added to the camera, and it is also possible to scale the strength of the post process effects.
 
-Two components are added to the CameraComponent to aid in visual placement in the Editor, although they will not be visible during gameplay.  A FrustumComponent shows where the 
+Two components are added to the CameraComponent to aid in visual placement in the Editor, although they will not be visible during gameplay.  A **FrustumComponent** shows where the 
 camera field of view is. This does not show by default, but can be turned on by selecting **Show > Advanced > Camera Frustums** in the **Viewport**. A StaticMeshComponent represents the camera's placement 
  in the level. 
 
 
 ## PlayerCameraManager
 
-The PlayerCameraManager class is a camera manager.  By default, its own built-in behavior is blending between pending view targets and debug cameras triggered by console commands. Otherwise, it queries the 
+The **PlayerCameraManager** class is a camera manager.  By default, its own built-in behavior is blending between pending view targets and debug cameras triggered by console commands. Otherwise, it queries the 
 ViewTarget for what to do for the camera's viewpoint, and all other camera settings.  Usually you will not need a PlayerCameraManager subclass - little modification to PlayerCameraManager is needed beyond 
 possibly adding rules for setting the ViewTarget if the automatic rules are insufficient.
 
@@ -37,7 +38,7 @@ using this function, return _true_ to use the returned values, or return _false_
 
 ### ViewTarget
 
-The ViewTarget struct, defined in PlayerCameraManager, is responsible for providing the PlayerCameraManager with an ideal Point of View (POV).  ViewTarget contains information on the target Actor, the Controller
+The **ViewTarget** struct, defined in PlayerCameraManager, is responsible for providing the PlayerCameraManager with an ideal Point of View (POV).  ViewTarget contains information on the target Actor, the Controller
 of the target Actor (for non-locally controlled Pawns), and the PlayerState, which is used to follow the same player through Pawn transitions and other changes while spectating.  The camera information
 passed to PlayerCameraManager through the POV property is in the form of a `FMinimalViewInfo` struct.  This struct contains the basic camera information from a CameraComponent, including the location, rotation,
 projection mode (Perspective or Orthographic), FOV, orthographic width, aspect ratio, and post process effects.  Providing the PlayerCameraManager with access to these values allows the PlayerCameraManager to
@@ -50,18 +51,18 @@ Game-specific camera behavior can be provided at any point along the camera "res
 
 ### CameraComponent 
 
-A CameraComponent will provide information about the camera's properties if the ViewTarget is a CameraActor or an Actor that contains a CameraComponent and has `bFindCameraComponentWhenViewTarget`
+A **CameraComponent** will provide information about the camera's properties if the ViewTarget is a CameraActor or an Actor that contains a CameraComponent and has `bFindCameraComponentWhenViewTarget`
 set to _true_. A related property that can be set for any Pawn is `bTakeCameraControlWhenPossessed`, where the Pawn will automatically become the ViewTarget upon possession by the PlayerController.
 
 ### Actor or PlayerController
 
-Both PlayerControllers and Actors contain a CalcCamera function.  An Actor's CalcCamera function returns the camera view of the first CameraComponent in the Actor, if `bFindCameraComponentWhenViewTarget` is _true_ and a CameraComponent is present.  
+Both **PlayerControllers** and **Actors** contain a CalcCamera function.  An Actor's CalcCamera function returns the camera view of the first CameraComponent in the Actor, if `bFindCameraComponentWhenViewTarget` is _true_ and a CameraComponent is present.  
 Otherwise, it gets the Actor's location and rotation.  In PlayerController, the CalcCamera function behaves similarly to this second case, returning the location of the possessed Pawn if it exists, and the 
 control rotation of the PlayerController.  
 
 
 ### PlayerCameraManager
 
-The UpdateViewTarget function in PlayerCameraManager queries the ViewTarget and returns that ViewTarget's Point Of View.  This function is also what calls BlueprintUpdateCamera if you have 
+The **UpdateViewTarget** function in PlayerCameraManager queries the ViewTarget and returns that ViewTarget's Point Of View.  This function is also what calls BlueprintUpdateCamera if you have 
 subclassed APlayerCameraManager and are not looking through a CameraComponent.
 
