#UE4doc #docs - Added engine version field. Added skill level field.

diff --git a/Engine/Documentation/Source/Gameplay/HowTo/ReferenceAssets/Blueprints/GHT_5B.INT.udn b/Engine/Documentation/Source/Gameplay/HowTo/ReferenceAssets/Blueprints/GHT_5B.INT.udn
--- a/Engine/Documentation/Source/Gameplay/HowTo/ReferenceAssets/Blueprints/GHT_5B.INT.udn
+++ b/Engine/Documentation/Source/Gameplay/HowTo/ReferenceAssets/Blueprints/GHT_5B.INT.udn
@@ -2,6 +2,8 @@ Availability: Public
 Title:Referencing Actors in Blueprints
 Crumbs: %ROOT%, Gameplay, Gameplay/HowTo/ReferenceAssets
 Description:A How To Guide for referencing Actors in Blueprints. 
+Version: 4.9
+SkillLevel: Advanced
 
 [TOC (start:2)]
 
