#UE4doc #docs - Added engine version field.

diff --git a/Engine/Documentation/Source/Gameplay/HowTo/RespawnPlayer/GHT_3_Overview.INT.udn b/Engine/Documentation/Source/Gameplay/HowTo/RespawnPlayer/GHT_3_Overview.INT.udn
--- a/Engine/Documentation/Source/Gameplay/HowTo/RespawnPlayer/GHT_3_Overview.INT.udn
+++ b/Engine/Documentation/Source/Gameplay/HowTo/RespawnPlayer/GHT_3_Overview.INT.udn
@@ -2,6 +2,7 @@ Availability: Public
 Title:Respawning a Player
 Crumbs: %ROOT%, Gameplay
 Description:A How To Guide for respawning player characters in Unreal Engine 4. 
+Version: 4.9
 
 In most games, there comes a time when the player character is killed and you want to respawn them in the world. In this How To guide, you will learn how to set up a simple respawn system for single player purposes. 
 
