#UE4doc #docs - Added Engine Version Field. Added Skill Level field.

diff --git a/Engine/Documentation/Source/Gameplay/HowTo/RespawnPlayer/Blueprints/GHT_3B.INT.udn b/Engine/Documentation/Source/Gameplay/HowTo/RespawnPlayer/Blueprints/GHT_3B.INT.udn
--- a/Engine/Documentation/Source/Gameplay/HowTo/RespawnPlayer/Blueprints/GHT_3B.INT.udn
+++ b/Engine/Documentation/Source/Gameplay/HowTo/RespawnPlayer/Blueprints/GHT_3B.INT.udn
@@ -2,6 +2,8 @@ Availability: Public
 Title:Respawning a Player in Blueprints
 Crumbs: %ROOT%, Gameplay, Gameplay/HowTo/RespawnPlayer
 Description:A How To Guide for respawning a player character through Blueprints in Unreal Engine 4. 
+Version: 4.9
+SkillLevel: Advanced
 
 [TOC (start:2)]
 
