#UE4doc #docs - Added Engine Version field.

diff --git a/Engine/Documentation/Source/Gameplay/HowTo/SettingUpAGameMode/GameModeOverview.INT.udn b/Engine/Documentation/Source/Gameplay/HowTo/SettingUpAGameMode/GameModeOverview.INT.udn
--- a/Engine/Documentation/Source/Gameplay/HowTo/SettingUpAGameMode/GameModeOverview.INT.udn
+++ b/Engine/Documentation/Source/Gameplay/HowTo/SettingUpAGameMode/GameModeOverview.INT.udn
@@ -2,6 +2,7 @@ Availability: Public
 Title:Setting Up a Game Mode
 Crumbs: %ROOT%, Gameplay
 Description:A How To Guide for Setting Up a Game Mode in Unreal Engine 4. 
+Version: 4.9
 
 The **GameMode** defines the game's set of rules. The rules can include how players join the game, game pausing, and level transition, as well as any game-specific behavior such as win conditions. The GameMode is set for each level, and GameModes can be reused in multiple levels. 
 
