#UE4doc #docs - Added Engine Version field. Added Skill Level field.

diff --git a/Engine/Documentation/Source/Gameplay/HowTo/SettingUpAGameMode/Blueprints/GameMode_BP.INT.udn b/Engine/Documentation/Source/Gameplay/HowTo/SettingUpAGameMode/Blueprints/GameMode_BP.INT.udn
--- a/Engine/Documentation/Source/Gameplay/HowTo/SettingUpAGameMode/Blueprints/GameMode_BP.INT.udn
+++ b/Engine/Documentation/Source/Gameplay/HowTo/SettingUpAGameMode/Blueprints/GameMode_BP.INT.udn
@@ -2,6 +2,8 @@ Availability: Public
 Title:Setting Up a Game Mode in Blueprints
 Crumbs: %ROOT%, Gameplay, Gameplay\HowTo\SettingUpAGameMode
 Description:A How To Guide for Setting Up a Game Mode in Blueprints.
+Version: 4.9
+SkillLevel: Advanced
 
 [TOC (start:2 end:2)]
 
