#UE4doc #docs:
Added engine version field
Bolded a couple terms

diff --git a/Engine/Documentation/Source/Engine/Landscape/Landscape.INT.udn b/Engine/Documentation/Source/Engine/Landscape/Landscape.INT.udn
--- a/Engine/Documentation/Source/Engine/Landscape/Landscape.INT.udn
+++ b/Engine/Documentation/Source/Engine/Landscape/Landscape.INT.udn
@@ -2,7 +2,7 @@ Availability: Public
 Title:Landscape Outdoor Terrain
 Crumbs:%ROOT%, Engine
 Description:Landscape system for creating terrain for large, open, outdoor environments.
-
+version: 4.9
 [VAR:Topic]
 [OBJECT:Topic]
 	[PARAM:image]
@@ -122,7 +122,7 @@ Description:Landscape system for creating terrain for large, open, outdoor envir
 [TOC(start:2 end:3)]
 
 [EXCERPT:Intro]
-The Landscape system enables you to create terrain for your world - mountains, valleys, uneven or sloped ground, even openings for caves - and easily modify both its shape and its appearance by using a range of tools.
+The **Landscape** system enables you to create terrain for your world - mountains, valleys, uneven or sloped ground, even openings for caves - and easily modify both its shape and its appearance by using a range of tools.
 [/EXCERPT:Intro]
 
 For information about opening and using the Landscape tool, see [](Engine/Landscape\Creation).
@@ -155,7 +155,7 @@ The main features and techniques employed by the Landscape terrain system are de
 ### Large Terrain Sizes
 
 [EXCERPT:landscape001]
-The Landscape system paves the way for terrains that are orders of magnitude larger than what has been possible in Unreal Engine previously. Because of its powerful LOD system and the way it makes efficient use of memory, heightmaps up to 8192x8192 are now legitimately possible and feasible. The Unreal Engine now supports huge outdoor worlds that allow for any number of types of games to be created quickly, easily, and without modification of the stock engine or tools.
+The Landscape system paves the way for terrains that are orders of magnitude larger than what has been possible in Unreal Engine previously. Because of its powerful **Level of Detail** (**LOD**) system and the way it makes efficient use of memory, heightmaps up to 8192x8192 are now legitimately possible and feasible. The Unreal Engine now supports huge outdoor worlds that allow for any number of types of games to be created quickly, easily, and without modification of the stock engine or tools.
 [/EXCERPT:landscape001]
 
 [REGION:raw]
