#UE4doc #docs:
Added engine version field
Bolded a few terms
Switched some regions from regions to tips and vice versa due to the content of the regions

diff --git a/Engine/Documentation/Source/Engine/Landscape/Creation/Creation.INT.udn b/Engine/Documentation/Source/Engine/Landscape/Creation/Creation.INT.udn
--- a/Engine/Documentation/Source/Engine/Landscape/Creation/Creation.INT.udn
+++ b/Engine/Documentation/Source/Engine/Landscape/Creation/Creation.INT.udn
@@ -2,6 +2,7 @@ Availability: Public
 Title:Creating Landscapes
 Crumbs: %ROOT%, Engine, Engine/Landscape
 Description:Guide to creating new Landscape terrains.
+version: 4.9
 
 [VAR:Topic]
 [OBJECT:Topiccompact]
@@ -26,14 +27,14 @@ Description:Guide to creating new Landscape terrains.
 [TOC(start:2)]
 
 Creating massive terrain based worlds for your projects is something that Unreal Engine 4 (UE4) can do quite well with its built in terrain editing tools. 
-The Landscape tool will allow you to create massive terrain based worlds that can still maintain playable frame rates across a multitude of different devices. 
+The **Landscape** tool will allow you to create massive terrain based worlds that can still maintain playable frame rates across a multitude of different devices. 
 Because of the flexibility of the Landscape system, you can create your Landscape using either of the following methods. 
 
 * You can create a new Landscape heightmap from scratch using the Landscape tools in-engine.
 * You can import a Landscape heightmap created previously in Unreal Editor or created through external tools. To read more about how to do this, please check out
 	the [Creating Custom Heightmaps and Layers](Engine/Landscape/Custom) document.
 
-[REGION:tip]
+[REGION:note]
 For examples of the Landscape tool in action, see the [](Resources\ContentExamples\Landscapes) Content Example.
 [/REGION]
 
@@ -43,12 +44,12 @@ For more information about customizing the Level Editor's panels, see [Level Edi
 
 ## Opening the Landscape Tool
 
-Before you can create a Landscape, you must first open the Landscape tool which you can find on the **Modes toolbar** panel. 
-To open the Landscape tools click on the **Mountain icon** along the top of the **Modes** panel.
+Before you can create a Landscape, you must first open the Landscape tool which you can find on the **Modes** panel. 
+To open the Landscape tools click on the mountain icon along the top of the **Modes** panel.
 
 ![](Landscape_ModesPanel.png)
 
-[REGION:note]
+[REGION:tip]
 You can also press **Shift + 3** on your keyboard to switch to the Landscape tool at any time. 
 [/REGION]
 
@@ -98,7 +99,7 @@ For this example, we are going to leave all of the Landscape settings at their d
 | **Number of Components:**| 8 x 8|
 | **Overall Resolution:**| 505 x 505|
 
-[REGION:tip]
+[REGION:note]
 Be careful when choosing **Section Size** because as the number of components increase, build times and performance can be dramatically affected. 63x63 quads per section is generally recommended as it gives good performance and size.
 [/REGION]
 
@@ -131,7 +132,7 @@ If you like, you can assign a Material to your Landscape before you create it. T
 
 ![](Landscape_Assign_Material.png)
 
-[REGION:tip]
+[REGION:note]
 For more information about Landscape Materials, see [](Engine/Landscape\Materials).
 [/REGION]
 
