#UE4doc #docs:
Added engine version field
Changed font settings for multiple terms
Added a new image to illustrate how to add Layer Info
Set a picture that needed to be a .png to convert:false

diff --git a/Engine/Documentation/Source/Engine/Landscape/Custom/CreatingCustomHeightmapsAndLayers.INT.udn b/Engine/Documentation/Source/Engine/Landscape/Custom/CreatingCustomHeightmapsAndLayers.INT.udn
--- a/Engine/Documentation/Source/Engine/Landscape/Custom/CreatingCustomHeightmapsAndLayers.INT.udn
+++ b/Engine/Documentation/Source/Engine/Landscape/Custom/CreatingCustomHeightmapsAndLayers.INT.udn
@@ -2,6 +2,7 @@ Availability: Public
 Title:Creating and Using Custom Heightmaps and Layers
 Crumbs: %ROOT%, Engine, Engine/Landscape
 Description:Creating and using custom Heightmaps and layers for your landscapes.
+version: 4.9
 
 
 [TOC(start:2 end:3)]
@@ -10,7 +11,7 @@ Description:Creating and using custom Heightmaps and layers for your landscapes.
 
 Sometimes your Landscape will require that you use external programs to create both the height map and layers that are needed. 
 UE4 does accommodate this style of workflow and everything that you need to be aware of when using a workflow like this is in this guide. 
-If this is your first time using the Landscape tools, you might want to check out the [Landscape Overview](Engine/Landscape/Editing) first as what is going to be covered here is quite advanced for new users. 
+If this is your first time using the Landscape tools, you might want to check out the [Landscape Overview](Engine/Landscape/Editing) first as what is going to be covered here may be too advanced for new users. 
 
 
 ## Layers
@@ -54,19 +55,26 @@ If your layers are not output at the correct size you will receive the following
 Importing Layers that were made in an external application is a simple process but you first need to make sure that you have a few things set up in order to get everything to work smoothly. 
 
 1. You first need to make sure that you have a Landscape to work with. If you do not have one now, please create one. 
-1. Once your Landscape has been created, you need to make a new material for it. For this example, we will be making a very simple Material but this Material could easily be expanded upon if needed. The Material that you need to setup should look something like this. 
-![](Landscape_Simple_Shader.png)
+1. Once your Landscape has been created, you need to make a new material for it. For this example, we will be making a very simple Material but this Material could easily be expanded upon if needed. The Material that you need to setup should look something like this.
+
+	![](Landscape_Simple_Shader.png)
 1. Once the material is created, apply it to the Landscape Actor so that you have something that looks like this. 
-![](Landscape_Applied_Material.png)
-1. Now with the material applied, it is now time to add some  **Layer Info** to our Landscape Actor. Add a **Layer Info** object for each layer you are going to have. 
-For this example, we are going to have three. To read more about **Layer Info** objects, check out the [Layer Info Objects](Engine/Landscape/Editing/PaintMode#LayerInfOobjects) section of the [Paint Mode](Engine/Landscape/Editing/PaintMode) document. When completed, your **Target Layers** section should look like this. 
-![](Landscape_Target_Layers.png)
-1. Now it is time to import our layer. To do this, all you need to do is **Right-Click** on the **Target Layer** you want to import a layer for and then select the layer that you want to import from the dialogue box that pops up. 
-![](Landscape_Import_Layer_Option.png)
+
+	![](Landscape_Applied_Material.png)
+1. Now with the material applied, it is now time to add some  **Layer Info** to our Landscape Actor. Add a Layer Info object for each layer you are going to have. For this example, we are going to have three. To read more about Layer Info objects, check out the [Layer Info Objects](Engine/Landscape/Editing/PaintMode#LayerInfOobjects) section of the [Paint Mode](Engine/Landscape/Editing/PaintMode) document.
+
+	![](Landscape_Create_Layer.png)
+1. When completed, your **Target Layers** section should look like this. 
+
+	![](Landscape_Target_Layers.png)
+1. Now it is time to import our layer. To do this, all you need to do is **Right-Click** on the Target Layer you want to import a layer for and then select the layer that you want to import from the dialogue box that pops up. 
+
+	![](Landscape_Import_Layer_Option.png)
 1. If your layers are not output at the correct size, you will receive the following warning.  
-![](Landscape_Layer_Import_Error.png)
+
+	![](Landscape_Layer_Import_Error.png)
  
-	 To fix this issue, you need to determine how big the Landscape layer you are trying to import over is supposed to be. To find this out, you can export any layer by **Right-Clicking** on the **Target Layer** you want to look at then select **Export** from the dialogue box that pops up.
+	 To fix this issue, you need to determine how big the Landscape layer you are trying to import over is supposed to be. To find this out, you can export any layer by **Right-Clicking** on the Target Layer you want to look at then select **Export** from the dialogue box that pops up.
 	 Once you do that, you will be asked to save the layer file somewhere on your PC. After the file has been saved, you can then open it up to determine what size your layers should be. 
 
 
@@ -98,7 +106,7 @@ Importing custom Heightmap is quite simple. The process is initiated from the **
 ![](Engine\Landscape\Creation\Landscape_Import_Section.png)
 
 
-To select a Heightmap to import, press the button next the box that says **Please specify a heightmap** and then select your Heightmap using the file browser. 
+To select a Heightmap to import, press the button next the box that says "Please specify a heightmap" and then select your Heightmap using the file browser. 
 
 ![](Engine\Landscape\Creation\Landscape_Import_Section_Button.png)
 
@@ -108,10 +116,10 @@ Now let's try importing a custom height map. For this example, we will be using
 Make sure you right-click Save as on the image to download it to your PC.
 [/REGION]
 
-![](Engine\Landscape\Creation\Landscape_Test_HM.png)
+![](Engine\Landscape\Creation\Landscape_Test_HM.png)(convert:false)
 
 Once you have the heightmap saved to your PC, it is now time to use it in the Landscape tools. 
-1. Open up the Landscape tool and go into [**Manage Mode.**](Engine/Landscape/Editing/ManageMode)
+1. Open up the Landscape tool and go into [Manage Mode.](Engine/Landscape/Editing/ManageMode)
 	
 	![](Landscape_Open_Manage_Mode.png)
 
