#UE4doc #docs:
Added engine version field
Changed font settings for multiple terms

diff --git a/Engine/Documentation/Source/Engine/Landscape/Editing/Brushes/Brushes.INT.udn b/Engine/Documentation/Source/Engine/Landscape/Editing/Brushes/Brushes.INT.udn
--- a/Engine/Documentation/Source/Engine/Landscape/Editing/Brushes/Brushes.INT.udn
+++ b/Engine/Documentation/Source/Engine/Landscape/Editing/Brushes/Brushes.INT.udn
@@ -2,6 +2,7 @@ Availability: Public
 Title:Landscape Brushes
 Crumbs: %ROOT%, Engine, Engine/Landscape, Engine/Landscape/Editing
 Description:Guide to the brushes available for modifying Landscapes.
+version: 4.9
 
 [VAR:Topic]
 [OBJECT:Topiccompact]
@@ -26,7 +27,7 @@ Description:Guide to the brushes available for modifying Landscapes.
 
 [TOC(start:2 end:2)]
 
-The Landscape tool's brush defines the size and shape of the area of the Landscape that will be affected by either sculpting or painting. Brushes can have different shapes, sizes, and falloffs. Brushes should be a familiar concept to anyone who has experience using Photoshop or a similar image-editing application. 
+The **Landscape** tool's brush defines the size and shape of the area of the Landscape that will be affected by either sculpting or painting. Brushes can have different shapes, sizes, and falloffs. Brushes should be a familiar concept to anyone who has experience using Photoshop or a similar image-editing application. 
 
 You can set the brush type and falloff in either **Sculpt** or **Paint** mode.
 
@@ -52,7 +53,7 @@ The size and falloff of the current brush is displayed in the viewport as a flas
 
 ![Circle Brush](Landscape_CircleBrush.png)
 
-The Circle brush applies the current tool in a circular area, with a falloff defined both numerically and by type.
+The **Circle** brush applies the current tool in a circular area, with a falloff defined both numerically and by type.
 
 ![circle.jpg](circle.jpg)(w:486 h:486)
 
@@ -77,7 +78,7 @@ Below is an example of each of these falloff types when operating on the heightm
 
 ![Alpha Brush](Landscape_AlphaBrush.png)
 
-The Alpha brush is similar to the Pattern brush, but instead of tiling the texture across the Landscape as you paint, it orients the brush texture in the direction of your painting and drags the shape as you move the cursor.
+The **Alpha** brush is similar to the Pattern brush, but instead of tiling the texture across the Landscape as you paint, it orients the brush texture in the direction of your painting and drags the shape as you move the cursor.
 
 ![brush_alpha_drag.jpg](brush_alpha_drag.jpg)(w:273 h:475)
 
@@ -98,7 +99,7 @@ The Alpha brush is similar to the Pattern brush, but instead of tiling the textu
 
 ![Pattern Brush](Landscape_PatternBrush.png)
 
-The Pattern brush allows you to use an arbitrary shape for the brush by sampling a single color channel from a texture to use as the alpha for the brush. The texture pattern is tiled as the brush is painted.
+The **Pattern** brush allows you to use an arbitrary shape for the brush by sampling a single color channel from a texture to use as the alpha for the brush. The texture pattern is tiled as the brush is painted.
 
 For example, the textures below could be used as the alpha:
 
@@ -131,7 +132,7 @@ Those would result in the following brushes:
 
 ![Component Brush](Landscape_ComponentBrush.png)
 
-The Component brush is used for operating on individual components. The cursor becomes limited to a single component at a time:
+The **Component** brush is used for operating on individual components. The cursor becomes limited to a single component at a time:
 
 ![selection_cursor.jpg](selection_cursor.jpg)(w:900 h:453)
 
@@ -143,7 +144,7 @@ The Component brush is the only brush available when you are using tools that op
 
 ![Gizmo Brush](Landscape_GizmoBrush.png)
 
-The Gizmo brush is used for modifying your Landscape with Landscape Gizmos, which are tools that you use to perform actions to a specific localized area of the Landscape.
+The **Gizmo** brush is used for modifying your Landscape with Landscape Gizmos, which are tools that you use to perform actions to a specific localized area of the Landscape.
 
 [REGION:note]
 The Gizmo brush is available only when you are using the Copy/Paste tool in Sculpt mode.
