#UE4doc #docs:
Added engine version field
Changed font settings for a few terms
Replaced name of a field as it has been updated

diff --git a/Engine/Documentation/Source/Engine/Landscape/Editing/Gizmos/Gizmos.INT.udn b/Engine/Documentation/Source/Engine/Landscape/Editing/Gizmos/Gizmos.INT.udn
--- a/Engine/Documentation/Source/Engine/Landscape/Editing/Gizmos/Gizmos.INT.udn
+++ b/Engine/Documentation/Source/Engine/Landscape/Editing/Gizmos/Gizmos.INT.udn
@@ -2,6 +2,7 @@ Availability: Public
 Title:Landscape Gizmos
 Crumbs: %ROOT%, Engine, Engine/Landscape, Engine/Landscape/Editing
 Description:Guide to the Gizmos available for modifying Landscapes.
+version: 4.9
 
 [VAR:Topic]
 [OBJECT:Topiccompact]
@@ -26,7 +27,7 @@ Description:Guide to the Gizmos available for modifying Landscapes.
 
 [TOC(start:2 end:3)]
 
-Landscape Gizmos are editor-only Actors that are similar to volumes in that they define a specified area. Their purpose is to hold height and layer data for an area of a Landscape so that it can be copied to another location on the Landscape or exported for use in a different Landscape or heightmap generator (World Machine, Terresculptor,...).
+**Landscape Gizmos** are editor-only Actors that are similar to volumes in that they define a specified area. Their purpose is to hold height and layer data for an area of a Landscape so that it can be copied to another location on the Landscape or exported for use in a different Landscape or heightmap generator (World Machine, Terresculptor,...).
 
 [OBJECT:EmbeddedVideo]
 [PARAMLITERAL:width]
@@ -55,7 +56,7 @@ IdzHGD7iWdg
 
 Selecting the Gizmo brush creates a temporary Gizmo Actor. You can use the transformation tools to manipulate this Actor as you would any other, to define the area of the Landscape you want to copy. 
 
-[region:tip]
+[region:note]
 For more information about the transformation tools, see [](Engine/Actors/Transform).
 [/region]
 
@@ -78,7 +79,7 @@ You can also modify the Gizmo Actor's properties in the **Details** panel.
 
 ## Copying to the Gizmo
 
-In order to copy a portion of a Landscape, the data for that area must be copied to a Gizmo. Then, that data can later be [pasted](#PastingData) to another location.
+In order to copy a portion of a Landscape, the data for that area must be copied to a Gizmo. Then, that data can later be pasted to another location.
 
 
 
@@ -129,7 +130,7 @@ In order to copy a portion of a Landscape, the data for that area must be copied
 
 ## Pasting from the Gizmo
 
-Pasting data from a Gizmo makes it possible to transfer parts of a Landscape from one location to another. The data can either be pasted completely to create an identical feature, or it can be painted to the new location using brush strokes and strength settings to partially transfer the feature.
+Pasting data from a Gizmo makes it possible to transfer parts of a Landscape from one location to another. The data can either be [pasted](#pastingdata) completely to create an identical feature, or it can be painted to the new location using brush strokes and strength settings to partially transfer the feature.
 
 Before data can be pasted from a Gizmo, it must first be [copied](#CopyingData) to a Gizmo.
 
@@ -192,7 +193,7 @@ Heightmap data can be imported from and exported to the Gizmo through the **Gizm
 (#ExportingGizmo)
 **To Export Gizmo Data:**
 
-1. With the Gizmo populated with data (see [Copying to the Gizmo](#CopyingData)), press the **Export Gizmo Data** button to export the Gizmo data to a file. Enabling the **Gizmo applies to all layers** checkbox at the top of the Gizmo options will export the heightmap and all layers' weight data to files.  
+1. With the Gizmo populated with data (see [Copying to the Gizmo](#CopyingData)), press the **Export Gizmo Data** button to export the Gizmo data to a file. Enabling the **Gizmo copy/paste all layers** checkbox at the top of the Gizmo options will export the heightmap and all layers' weight data to files.  
 1. Choose a location and name for the heightmap file.  
     ![export_file.jpg](export_file.jpg)(w:625 h:436)
 1. If exporting layers, choose locations and filenames for each exported layer.  
