#UE4doc #docs:
Added engine version field
Changed font settings for multiple terms

diff --git a/Engine/Documentation/Source/Engine/Landscape/Editing/ManageMode/ManageMode.INT.udn b/Engine/Documentation/Source/Engine/Landscape/Editing/ManageMode/ManageMode.INT.udn
--- a/Engine/Documentation/Source/Engine/Landscape/Editing/ManageMode/ManageMode.INT.udn
+++ b/Engine/Documentation/Source/Engine/Landscape/Editing/ManageMode/ManageMode.INT.udn
@@ -2,6 +2,7 @@ Availability: Public
 Title:Manage Mode
 Crumbs: %ROOT%, Engine, Engine/Landscape, Engine/Landscape/Editing
 Description:Guide to the Landscape tool's Manage mode.
+version: 4.9
 
 [VAR:Topic]
 [OBJECT:Topiccompact]
@@ -32,7 +33,7 @@ For information about creating Landscapes, see [](Engine/Landscape\Creation).
 
 ## Landscape Selection
 
-In the Landscape tool, in any mode, you can select which Landscape you want to modify from a list of all the Landscapes currently in your level.
+In the **Landscape** tool, in any mode, you can select which Landscape you want to modify from a list of all the Landscapes currently in your level.
 
 ![](Landscape_Select_Different_LS.png)
 
@@ -44,7 +45,7 @@ In the Landscape tool, in any mode, you can select which Landscape you want to m
 Selecting a Landscape from the list makes that the active Landscape. All subsequent editing will then be applied to that Landscape.
 ## Components
 
-Components are the building blocks of Landscapes.  You can work with them with Component tools, and you can also adjust their properties individually in the Level Editor's **Details** panel.
+**Components** are the building blocks of Landscapes.  You can work with them with Component tools, and you can also adjust their properties individually in the Level Editor's **Details** panel.
 
 ### Component Tools
 
@@ -54,7 +55,7 @@ The tools described below act on the component level of the current Landscape.
 
 ![Landscape Selection](Landscape_Selection.png)
 
-The Selection tool selects Landscape components, one at a time, for subsequent use with other tools, such as moving the components to streaming levels and deleting components.
+The **Selection** tool selects Landscape components, one at a time, for subsequent use with other tools, such as moving the components to streaming levels and deleting components.
 
 | **Control** | **Operation** |
 | --- | --- |
@@ -79,7 +80,7 @@ When components have been selected, they become shaded red:
 
 ![Add Component](Landscape_Add.png)
 
-The Add tool creates new components for the current Landscape, one at a time.
+The **Add** tool creates new components for the current Landscape, one at a time.
 
 | **Control** | **Operation** |
 | --- | --- |
@@ -98,7 +99,7 @@ Using the tool adds a new component where the cursor is located:
 
 ![Delete Component](Landscape_Delete.png)
 
-The Delete tool deletes selected components.
+The **Delete** tool deletes selected components.
 
 | **Control** | **Operation** |
 | --- | --- |
@@ -111,7 +112,7 @@ Using the tool removes the selected components:
 
 ![selection_remove.jpg](selection_remove.jpg)(w:900 h:453)
 
-[REGION:tip]
+[REGION:note]
 If you want more control over what areas of your Landscape the player can get through. Check out the [Landscape Hole Material](Engine/Landscape/Materials#landscapeholematerials).
 [/REGION]
 
diff --git a/Engine/Documentation/Source/Shared/LandscapeEditor/LandscapeEditor.INT.udn b/Engine/Documentation/Source/Shared/LandscapeEditor/LandscapeEditor.INT.udn
--- a/Engine/Documentation/Source/Shared/LandscapeEditor/LandscapeEditor.INT.udn
+++ b/Engine/Documentation/Source/Shared/LandscapeEditor/LandscapeEditor.INT.udn
@@ -61,7 +61,7 @@ Engine/Landscape/Editing/ManageMode
 [/EXCERPT:MoveToLevelTool]
 
 [EXCERPT:ResizeLandscape]
-The Change Component Size Landscape tool allows you to change the size of the entire Landscape. Keep in mind that this does not work on individual components but on the entire Landscape as a whole. With this tool you can change the size of the Landscape components and adjust how many sections each component has in the Change Component Size section.
+The **Change Component Size** Landscape tool allows you to change the size of the entire Landscape. Keep in mind that this does not work on individual components but on the entire Landscape as a whole. With this tool you can change the size of the Landscape components and adjust how many sections each component has in the Change Component Size section.
 [VAR:ToolTipFullLink]
 Engine/Landscape/Editing/ManageMode
 [/VAR]
