#UE4doc #docs:
Added engine version
Changed font settings for multiple terms

diff --git a/Engine/Documentation/Source/Engine/Landscape/Editing/SculptMode/SculptMode.INT.udn b/Engine/Documentation/Source/Engine/Landscape/Editing/SculptMode/SculptMode.INT.udn
--- a/Engine/Documentation/Source/Engine/Landscape/Editing/SculptMode/SculptMode.INT.udn
+++ b/Engine/Documentation/Source/Engine/Landscape/Editing/SculptMode/SculptMode.INT.udn
@@ -2,6 +2,7 @@ Availability: Public
 Title:Sculpt Mode
 Crumbs: %ROOT%, Engine, Engine/Landscape, Engine/Landscape/Editing
 Description:Guide to the Landscape tool's Sculpt mode.
+version: 4.9
 
 [VAR:Topic]
 [OBJECT:Topiccompact]
@@ -46,7 +47,7 @@ You can use the Sculpting Tools to modify the shape of your Landscape in various
 
 ![Sculpt Tool](Landscape_Sculpt.png)
 
-The Sculpt tool increases or decreases the heightmap height in the shape of the currently selected brush and falloff.
+The **Sculpt** tool increases or decreases the heightmap height in the shape of the currently selected brush and falloff.
 
 | **Alternate Controls** | **Operation** |
 | --- | --- |
@@ -66,7 +67,7 @@ The brush strength determines the amount of raising or lowering.
 
 ![Smooth Tool](Landscape_Smooth.png)
 
-The Smooth tool smooths the heightmap. The strength determines the amount of smoothing.
+The **Smooth** tool smooths the heightmap. The strength determines the amount of smoothing.
 
 **Heightmap Smoothing**
 
@@ -93,7 +94,7 @@ The Smooth tool smooths the heightmap. The strength determines the amount of smo
 
 ![Flatten Tool](Landscape_FlattenTool.png)
 
-The Flatten tool flattens the Landscape to the level of the Landscape under the mouse cursor when you first activate the tool.
+The **Flatten** tool flattens the Landscape to the level of the Landscape under the mouse cursor when you first activate the tool.
 
 ![Landscape_Original.jpg](Landscape_Original.jpg)(w:900 h:325)
 
@@ -112,11 +113,11 @@ The Flatten tool flattens the Landscape to the level of the Landscape under the
 
 ![Ramp Tool](Landscape_Ramp.png)
 
-The Ramp tool enables you to select two locations on your Landscape and create a flat ramp between them, with a falloff you specify on the sides.
+The **Ramp** tool enables you to select two locations on your Landscape and create a flat ramp between them, with a falloff you specify on the sides.
 
 **To add a ramp to your Landscape:**
 
-1. **Ctrl + click** in two locations to mark the beginning and end points of the ramp.
+1. **Ctrl + Left-click** in two locations to mark the beginning and end points of the ramp.
 
 	![Ramp Setup](Landscape_RampSetup.png)(w:700)
 
@@ -149,7 +150,7 @@ If you change your mind about adding a ramp after starting to set the beginning
 
 ![Erosion Tool](Landscape_Erosion.png)
 
-The Erosion tool uses a thermal erosion simulation to adjust the height of the heightmap. This simulates the transfer of soil from higher elevations to lower elevations. The larger the difference in elevation, the more erosion will occur. This tool also applies a noise effect on top of the erosion, if desired, to provide a more natural random appearance.
+The **Erosion** tool uses a thermal erosion simulation to adjust the height of the heightmap. This simulates the transfer of soil from higher elevations to lower elevations. The larger the difference in elevation, the more erosion will occur. This tool also applies a noise effect on top of the erosion, if desired, to provide a more natural random appearance.
 
 
 ![Landscape_Original.jpg](Landscape_Original.jpg)(w:900 h:325)
@@ -172,7 +173,7 @@ The Erosion tool uses a thermal erosion simulation to adjust the height of the h
 
 ![Hydro Erosion Tool](Landscape_HydroErosion.png)
 
-The Hydro Erosion tool uses a hydraulic erosion simulation, i.e. erosion from water, to adjust the height of the heightmap. A noise filter is used to determine where the initial rain is distributed. Then the simulation is calculated to determine water flow from that initial rain as well as dissolving, water transfer, and evaporation. The result of that calculation provides the actual value used to lower the heightmap.
+The **Hydro Erosion** tool uses a hydraulic erosion simulation, i.e. erosion from water, to adjust the height of the heightmap. A noise filter is used to determine where the initial rain is distributed. Then the simulation is calculated to determine water flow from that initial rain as well as dissolving, water transfer, and evaporation. The result of that calculation provides the actual value used to lower the heightmap.
 
 ![Landscape_Original.jpg](Landscape_Original.jpg)(w:900 h:325)
 
@@ -195,7 +196,7 @@ The Hydro Erosion tool uses a hydraulic erosion simulation, i.e. erosion from wa
 
 ![Noise Tool](Landscape_Noise.png)
 
-The Noise tool applies a noise filter to the heightmap or layer weight. The strength determines the amount of noise.
+The **Noise** tool applies a noise filter to the heightmap or layer weight. The strength determines the amount of noise.
 
 | **Option** | **Description** |
 | --- | --- |
@@ -210,7 +211,7 @@ The Noise tool applies a noise filter to the heightmap or layer weight. The stre
 
 ![Retopologize Tool](Landscape_Retopologize.png)
 
-The Retopologize tool automatically adjusts Landscape vertices with an X/Y offset map to improve vertex density on cliffs, reducing texture stretching.
+The **Retopologize** tool automatically adjusts Landscape vertices with an X/Y offset map to improve vertex density on cliffs, reducing texture stretching.
 
 [REGION:note]
 An X/Y offset map makes the Landscape slower to render and paint on with other tools, so only use the Retopologize tool if needed.
@@ -221,11 +222,11 @@ An X/Y offset map makes the Landscape slower to render and paint on with other t
 
 ![Visibility Tool](Landscape_Visibility.png)
 
-The Visibility tool enables you to create holes in your Landscape, such as for caves. 
+The **Visibility** tool enables you to create holes in your Landscape, such as for caves. 
 
 | **Alternate Controls** | **Operation** |
 | --- | --- |
-| **Shift + Left Mouse Button** | Removes the visibility mask, making the Landscape components visible again. |
+| **Shift + Left-click** | Removes the visibility mask, making the Landscape components visible again. |
 
 #### Using Opacity Masking to Create a Hole
 
@@ -243,11 +244,11 @@ After you have set up your Landscape Hole Material, you can use the painting too
 
 **To create a Landscape hole:**
 
-1. Make sure you have a Landscape Hole Material assigned to your Landscape.
+1. Make sure you have a **Landscape Hole Material** assigned to your Landscape.
 
 	![](Engine/Landscape/Materials/Landscape_HoleSettings.png)
 
-1. In the Landscape tool, in **Sculpt** mode, select the Visibility tool.
+1. In the Landscape tool, in **Sculpt** mode, select the **Visibility** tool.
 
 	![Visibility Tool](Landscape_Visibility.png)
 
@@ -259,13 +260,13 @@ After you have set up your Landscape Hole Material, you can use the painting too
 
 	![](Landscape_MountainCaveMask.png)(w:640)
 
-1. Press **Left Mouse Button** to create the hole and **Shift + Left Mouse Button** to erase a hole that was created. 
+1. **Left-click** to create the hole and **Shift + Left-click** to erase a hole that was created. 
 
 	![](Landscape_MountainHole.png)(w:640)
 
 	You can now fit Static Mesh Actors in the hole in the Landscape to create a cave.
 
-	[REGION:tip]
+	[REGION:note]
 	To test out the collision of the new hole using Play In Editor (PIE), you may have to switch from Landscape mode to Place mode.
 	[/REGION]
 
@@ -280,7 +281,7 @@ You can use the Region tools to perform actions on specific regions of the Lands
 
 ![Region Selection Tool](Landscape_RegionSelect.png)
 
-The Region Selection tool selects regions of the Landscape using the current brush settings and tool strength to be used to fit a Landscape [gizmo](Engine\Landscape\Editing\Gizmos) to a specific area or to act as a mask for copying data to, or pasting data from, a gizmo.
+The **Region Selection** tool selects regions of the Landscape using the current brush settings and tool strength to be used to fit a Landscape [gizmo](Engine\Landscape\Editing\Gizmos) to a specific area or to act as a mask for copying data to, or pasting data from, a gizmo.
 
 ![region_selection.jpg](region_selection.jpg)(w:640 h:322)
 
