#UE4doc #docs:
Added engine version field
Changed font settings for multiple terms
Replaced mentions of ctrl + click with ctrl + Left-click

diff --git a/Engine/Documentation/Source/Engine/Landscape/Editing/Splines/Splines.INT.udn b/Engine/Documentation/Source/Engine/Landscape/Editing/Splines/Splines.INT.udn
--- a/Engine/Documentation/Source/Engine/Landscape/Editing/Splines/Splines.INT.udn
+++ b/Engine/Documentation/Source/Engine/Landscape/Editing/Splines/Splines.INT.udn
@@ -2,6 +2,7 @@ Availability:Public
 Title:Landscape Splines
 Crumbs: %ROOT%, Engine, Engine/Landscape, Engine/Landscape/Editing
 Description: Splines for carving Landscape, modifying blendmask and deforming meshes into roads.
+version: 4.9
 
 [VAR:Topic]
 [OBJECT:Topiccompact]
@@ -25,7 +26,7 @@ Description: Splines for carving Landscape, modifying blendmask and deforming me
 
 [TOC(start:2)]
 
-The Spline tool creates a spline Actor that deforms a terrain mesh along it, such as for creating flat roads through mountainous terrain.
+The **Spline** tool creates a spline Actor that deforms a terrain mesh along it, such as for creating flat roads through mountainous terrain.
 
 [EXCERPT:Splines001]
 
@@ -45,23 +46,23 @@ The Spline tool is best used to make streets or paths in the environment, but in
 
 	![Spline Tool](Landscape_SplineTool.png)
 
-1. On the currently selected Landscape, press **Ctrl + click** to set your spline's first control point. The control point's sprite looks like a bumpy, gray-and-white boulder.
+1. On the currently selected Landscape, press **Ctrl + Left-click** to set your spline's first control point. The control point's sprite looks like a bumpy, brown-and-white boulder.
 
 	![A Spline control point](control_point.jpg)(w:785 h:516 type:jpg)
 
-1. Press **Ctrl + click** again to create another control point, connected to the first by a spline segment. The selection changes to the new control point so that you can repeat this to add more control points.
+1. Press **Ctrl + Left-click** again to create another control point, connected to the first by a spline segment. The selection changes to the new control point so that you can repeat this to add more control points.
 
 	![A second spline control point](short_spline.jpg)(w:785 h:516 type:jpg)
 
 
 **To join two splines:**
 
-* Select one spline and then **Ctrl + click** on the other.
+* Select one spline and then **Ctrl + Left-click** on the other.
 
 
 **To split a spline:**
 
-* Press **Ctrl + click** on a segment. This splits the segment at that point and inserts a new control point.
+* Press **Ctrl + Left-click** on a segment. This splits the segment at that point and inserts a new control point.
 
 You can assign a Static Mesh Actor to follow along your spline's segments.
 
