#DocFix

Forums: Videos on MediaFramework page were not YouTube videos. Converted into YouTube videos. Fixed grammatical error on page.

diff --git a/Engine/Documentation/Source/Engine/MediaFramework/MediaFramework.INT.udn b/Engine/Documentation/Source/Engine/MediaFramework/MediaFramework.INT.udn
--- a/Engine/Documentation/Source/Engine/MediaFramework/MediaFramework.INT.udn
+++ b/Engine/Documentation/Source/Engine/MediaFramework/MediaFramework.INT.udn
@@ -162,7 +162,7 @@ Inside the Details panel, you can inspect the properties of the media source.
 ![](MediaPlayerDetails.png)
 
 Under the **Playback** options, you can set the media to **Auto Play** (which will automatically play the media when loaded) and **Auto Play Rate** (to define the playback 
-rate when Auto Play is enabled), as well as the ability to the media to looping when playback reaches the end. 
+rate when Auto Play is enabled), as well as the ability to enable looping of the media when playback reaches the end. 
 
 For the **Source** section, you can specify the path to the media locally or through the Internet via URL as well as set the Streaming Mode. The last section, displays 
 information about the media itself. 
@@ -214,7 +214,17 @@ This section will show you how to quickly get started with using the Media Frame
 
 Here is an example of a MediaTexture playing on a TV (click to play). 
 
-[![](video1_image.png)](video1.mp4)
+[OBJECT:EmbeddedVideo]
+[PARAMLITERAL:width]
+640
+[/PARAMLITERAL]
+[PARAMLITERAL:height]
+360
+[/PARAMLITERAL]
+[PARAMLITERAL:videoid]
+IdX1gEH2i6c
+[/PARAMLITERAL]
+[/OBJECT]
 
 _For this example, we **Migrated** the **SM_TV** asset from the **Realistic Rendering** project which can be found on the Learn Tab._
 
@@ -253,7 +263,17 @@ to it. Our MediaPlayer asset is also not set to Auto Play.
 
 When the player enters the trigger, we will play the MediaTexture applied to the TV, when they exit the trigger we will pause it (click to play). 
 
-[![](video2_image.png)](video2.mp4)
+[OBJECT:EmbeddedVideo]
+[PARAMLITERAL:width]
+640
+[/PARAMLITERAL]
+[PARAMLITERAL:height]
+360
+[/PARAMLITERAL]
+[PARAMLITERAL:videoid]
+nmb2Fj_5PjA
+[/PARAMLITERAL]
+[/OBJECT]
 
 Inside our **Level Blueprint** we have created the following script. 
 
@@ -281,7 +301,17 @@ a material that is created from a MediaTexture (which in turn is created from a
 Below we have created a setup so that when the player press the M key in-game, a movie player is shown on screen with controls to Play, Pause, or Rewind the video that 
 is being played (click to play). 
 
-[![](video3_image.png)](video3.mp4)
+[OBJECT:EmbeddedVideo]
+[PARAMLITERAL:width]
+640
+[/PARAMLITERAL]
+[PARAMLITERAL:height]
+360
+[/PARAMLITERAL]
+[PARAMLITERAL:videoid]
+HuF6btgkfxE
+[/PARAMLITERAL]
+[/OBJECT]
 
 We did this by first creating a **[Widget Blueprint](Engine\UMG\UserGuide\WidgetBlueprints)** as indicated below. 
 
