#UE4doc #docs:
Fixed a capitalization error

diff --git a/Engine/Documentation/Source/Shared/Glossary/S/Glossary_S.INT.udn b/Engine/Documentation/Source/Shared/Glossary/S/Glossary_S.INT.udn
--- a/Engine/Documentation/Source/Shared/Glossary/S/Glossary_S.INT.udn
+++ b/Engine/Documentation/Source/Shared/Glossary/S/Glossary_S.INT.udn
@@ -67,7 +67,7 @@ socket
 	[PARAM:Definition]
 
 [EXCERPT:Glossary_socket]
-Commonly in games, you will want to attach an object to the bone of a character. This might be a weapon attached to a hand, or a hat attached to the head. For this purpose, Unreal Engine allows you to create **sockets** in its animation toolset - [Persona](Engine/Animation/Persona) - which are offset from a bone within the Skeletal Mesh. Sockets can then be translated, rotated, and scaled relative to the bone. Static Meshes and/or Skeletal Meshes can also be previewed attached to sockets. This allows content creators to easily set up sockets for a Skeletal Mesh, and then tell the programmer the names of the sockets to attach objects to.
+Commonly in games, you will want to attach an object to the bone of a character. This might be a weapon attached to a hand, or a hat attached to the head. For this purpose, Unreal Engine allows you to create **Sockets** in its animation toolset - [Persona](Engine/Animation/Persona) - which are offset from a bone within the Skeletal Mesh. Sockets can then be translated, rotated, and scaled relative to the bone. Static Meshes and/or Skeletal Meshes can also be previewed attached to sockets. This allows content creators to easily set up sockets for a Skeletal Mesh, and then tell the programmer the names of the sockets to attach objects to.
 [/EXCERPT:Glossary_socket]
 
 	[/PARAM]
