#UE4 Docs: Jira UEDOC-2324 (Gamma Correction, Color Picker Updates)

diff --git a/Engine/Documentation/Source/Engine/UI/ColorPicker/ColorPicker.INT.udn b/Engine/Documentation/Source/Engine/UI/ColorPicker/ColorPicker.INT.udn
--- a/Engine/Documentation/Source/Engine/UI/ColorPicker/ColorPicker.INT.udn
+++ b/Engine/Documentation/Source/Engine/UI/ColorPicker/ColorPicker.INT.udn
@@ -2,6 +2,9 @@ Availability: Public
 Title:Color Picker
 Crumbs:%ROOT%, Engine, Engine/UI
 Description:Tool for interactively choosing colors to assign to color properties of actors.
+Version: 4.9
+
+
 
 [VAR:Topic]
 [OBJECT:Topic]
@@ -46,29 +49,32 @@ Description:Tool for interactively choosing colors to assign to color properties
 The **Color Picker** allows you to easily adjust a color value for a color property in the Unreal Editor.  You can adjust colors by RGBA (red, green, blue, alpha), HSV (hue, saturation, value), and Hex (ARGB). Alternatively you can select a color by clicking anywhere in the color wheel or by capturing the color under your mouse, anywhere on your monitors.
 [/EXCERPT:Intro]
 
-[REGION:imagetable]
-| ![ColorPicker.jpg](ColorPicker.png) | [INCLUDE:#Properties] |
-| ------ | ------ |
-|  |  |
-[/REGION]
+![ColorPicker.jpg](ColorPicker.png)
 
-<!--
-[EXCERPT:Properties]
 | Item | Description |
 | ---- | ----------- |
-| ![Color Wheel](color_wheel.png)(w:96) | A Color wheel that allows a color to be selected graphically by clicking and dragging with the mouse. |
-| ![Color Swatches](color_swatches.png) | Displays the color currently assigned to the property and the color currently selected in the **Color Picker** [Top: solid color, Bottom: color with alpha]. Colors can be dragged from here to the **Themes** shelf to save them for reuse later on. See [Themes](#Themes) for more information. |
-| ![Red Channel](channel_r.png) | Drag or type value to set the **red** channel. |
-| ![Green Channel](channel_g.png) | Drag or type value to set the **green** channel. |
-| ![Blue Channel](channel_b.png) | Drag or type value to set the **blue** channel. |
-| ![Hue Channel](channel_h.png) | Drag or type value to set the **hue**. |
-| ![Saturation Channel](channel_s.png) | Drag or type value to set the **saturation**. |
-| ![Value Channel](channel_v.png) | Drag or type value to set the **value** (color brightness). |
-| ![Alpha Channel](channel_a.png) | Drag or type value to set the **alpha** channel. |
-| ![OK Button](button_ok.png) | Closes the picker and assigns the color in the picker to the property. |
-| ![Cancel Button](button_cancel.png) | Closes the picker without modifying the property's value. |
-[/EXCERPT:Properties]
--->
+| ![Themes & sRGB](ThemeSRGB.png) | Allows for [Theme](#Themes) selection and **sRGB** previews (see note below).|
+| ![Color Wheel](ColorWheel.png) | A Color wheel that allows a color to be selected graphically by clicking and dragging with the mouse. Additional adjustments can be made by using the two vertical bars. Instead of Color Wheel, a Color Spectrum may be used if the option is toggled on.  |
+| ![Color Swatches](ColorSwatches.png) | Displays the color currently assigned (Old) to the property and the color currently selected in the Color Picker (New). Colors can be dragged from either the _Old_ or _New_ sections to the **Themes** bar to save them for reuse later on. |
+| ![Picker Options](ColorOptions.png) | Toggles between the Color Wheel or Color Spectrum (left image). The **Eye Dropper** (right image) can be used for selecting a colored pixel from any window inside or outside of Unreal Engine. |
+| ![Red Channel](RedChannel.png) | Drag or type value to set the **Red** channel. |
+| ![Green Channel](GreenChannel.png) | Drag or type value to set the **Green** channel. |
+| ![Blue Channel](BlueChannel.png) | Drag or type value to set the **Blue** channel. |
+| ![Hue Channel](HueChannel.png) | Drag or type value to set the **Hue**. |
+| ![Saturation Channel](SaturationChannel.png) | Drag or type value to set the **Saturation**. |
+| ![Value Channel](ValueChannel.png) | Drag or type value to set the **Value** (color brightness). |
+| ![Alpha Channel](AlphaChannel.png) | Drag or type value to set the **Alpha** channel. |
+| ![Hex Values](HexValues.png) | Enter in Hexadecimal Linear or Hexadecimal sRGB values manually. |
+
+When working with the Color Picker, it does not assume a straight inverse gamma conversion from colors sampled on the screen using the Eye Dropper. Instead, sRGB is the assumed space all colors are sampled in, making sampling colors from Photoshop the identical version of the color converted into linear space. When transformed back into sRGB space, it will be the same color you saw in Photoshop.
+
+[REGION:note]
+`FColor` and `FLinearColor` can also default to the sRGB space using a UPROPERTY meta tag of **sRGB=true** (or **false**), you can force a specific default checking of the sRGB checkbox.
+
+Several places in [Unreal Motion Graphics (UMG)](Engine/UMG) now force this default since sRGB is the space UMG will be rendered in (on PC). Additional work will be needed on Mac where the final rendering space is gamma 2.2. 
+[/REGION]
+
+
 
 ## Themes
 
@@ -124,7 +130,7 @@ To change themes, select the theme you want to be active from the **Theme Menu**
 
 ![Select Theme](theme_select.png)
 
-To use a color from the current theme, simply **double-click** the color to be applied. This will update the color in the picker to match the saved color from the theme.
+To use a color from the current theme, **Double-click** the color to be applied. This will update the currently selected color in the Color Picker. 
 
 ![Use Color](theme_use_color.png)
 
