#UE4 Docs: Jira UEDOC-2224 (Animation Debug section updated, added ViewActor command description).

diff --git a/Engine/Documentation/Source/Engine/Animation/Overview/AnimationOverview.INT.udn b/Engine/Documentation/Source/Engine/Animation/Overview/AnimationOverview.INT.udn
--- a/Engine/Documentation/Source/Engine/Animation/Overview/AnimationOverview.INT.udn
+++ b/Engine/Documentation/Source/Engine/Animation/Overview/AnimationOverview.INT.udn
@@ -223,8 +223,16 @@ Debug features can be accessed by pressing the **`** key to bring up the console
 
 To display the Skeletal Mesh Bones, enter the `ShowDebug Bones` command inside the console. As seen in the image above the bones are represented by individual white lines. 
 
+The `ShowDebug Bones` command will display the information for the playable character during a play session. To view the debug animation for another character (other than the player character), you can use the `viewactor ExampleCharacterID` command, where _ExampleCharacterID_ refers to the Character ID assigned to the character you wish to view. You can find this ID by hovering over a character in the [World Outliner](Engine/UI/LevelEditor/SceneOutliner) which will display its ID name. 
+
+[REGION:note]
+At this time when using the `viewactor` command, the camera will be moved to and locked on the Actor you are currently viewing. If you need to manipulate the camera while debugging still, a work around for this is to create a [Play in Editor](Engine/UI/LevelEditor/InEditorTesting) session with more than one player from the dropdown menu using one window for debugging and the other for controlling a player character. 
+[/REGION]
+
 An alternative look, matching the bones displayed in Persona, can be enabled by entering the `ShowDebugToggleSubCategory 3DBones` command inside the console. 
 
+
+
 ### Animation Debug Output
 
 ![](ShowDebugAnimation.png)
