#UE4doc #docs:
-WIP-
Changed font settings for multiple terms
Replaced 'Right-Click Menu' with 'Context Menu'

diff --git a/Engine/Documentation/Source/Engine/Landscape/Materials/Materials.INT.udn b/Engine/Documentation/Source/Engine/Landscape/Materials/Materials.INT.udn
--- a/Engine/Documentation/Source/Engine/Landscape/Materials/Materials.INT.udn
+++ b/Engine/Documentation/Source/Engine/Landscape/Materials/Materials.INT.udn
@@ -25,7 +25,7 @@ Description:Setting up materials to work with Landscape terrains.
 
 [TOC(start:2 end:3)]
 
-Although any random Material can be used with a Landscape Actor, the Material system inside Unreal Engine 4 (UE4) does provide some special Landscape specific material nodes that can greatly help you to better texture your Landscape. In the following document, we will be taking an in-depth look at these nodes, how they function, and how you can use them in your Landscapes materials. 
+Although any random **Material** can be used with a **Landscape** Actor, the Material system inside Unreal Engine 4 (UE4) does provide some special Landscape specific material nodes that can greatly help you to better texture your Landscape. In the following document, we will be taking an in-depth look at these nodes, how they function, and how you can use them in your Landscapes materials. 
 
 [REGION:note]
 You modify Materials for use with Landscapes the same way you modify other Materials, in the [Material Editor](Engine/Rendering/Materials\Editor).
@@ -37,9 +37,9 @@ Landscapes use weight blending rather than alpha blending, so the blend factors
 
 ## Landscape Specific Material Nodes
 
-Inside of the Material editor, you will find that there are six special nodes that can be used with the Landscape system. You can find all the nodes in either the **Right-Click** menu or in the **Palette** menu under the Landscape category. 
+Inside of the Material editor, you will find that there are six special nodes that can be used with the Landscape system. You can find all the nodes in either the **Context** menu or in the **Palette** menu under the Landscape category. 
 
-**Right-Click Menu**
+**Context Menu**
 
 ![Material Editor Context Menu](Landscape_Mat_Context.png)
 
@@ -51,7 +51,7 @@ Inside of the Material editor, you will find that there are six special nodes th
 ### Landscape Layer Blend Node
 [EXCERPT:LandscapeLayerBlend]
 
-The **LandscapeLayerBlend** node enables you to blend together multiple Textures or Material networks so that they can be used as Landscape layers. The **LandscapeLayerBlend** uses an array to store information about the Landscape layers. To add layers to this array clicking the plus sign icon that you see next to the word **0 elements**.
+The **LandscapeLayerBlend** node enables you to blend together multiple Textures or Material networks so that they can be used as Landscape layers. The LandscapeLayerBlend uses an array to store information about the Landscape layers. To add layers to this array clicking the plus sign icon that you see next to the entry: "0 elements".
 
 ![](Landscape_Add_To_Array.png)
 
@@ -59,7 +59,7 @@ When you add multiple layers to the **LandscapeLayerBlend** node, you get someth
 
 ![LandscapeLayerBlend Node](Landscape_LayerBlend.png)
 
-You can set the properties of the **LandscapeLayerBlend** node in the **Details** panel of the Material Editor. 
+You can set the properties of the **LandscapeLayerBlend** node in the **Details** panel of the **Material Editor**. 
 
 
 |Number| Property | Description |
@@ -83,25 +83,22 @@ The following table details the **Layer Blend node's** inputs and outputs.
 |3:| **Unlabeled Output** | The result of the layers blended together. |
 
 
-[REGION:note]
+
 When using certain combinations of layer blend modes you could end up with black spots all over your Landscape where different layers meet. This problem is especially prevalent when using the **LB_HeightBlend** mode for multiple Landscape layers. **LB_HeightBlend** works by modulating the blend factor, or weight, for the layer using the specified height value. When you have multiple layers painted on an area and they all are set to **LB_HeightBlend**, it is possible that all the layers painted in a particular area will simultaneously have a 0 height value, so the desired blend factor for each layer becomes 0. Because there is no implicit or explicit ordering, the result will be black spots because no layers will have any contribution to that area. The situation is worse when you are blending a Normal map, because it results in a Normal value of (0,0,0) which is not valid and will cause rendering problems with the lighting. The solution to this problem is to use **LB_AlphaBlend** for one of the layers like in the following example. 
 ![LB_HeightBlend Problem](HeightBlendProblem.jpg)(w:800 h:238)
 
 In the left image, all layers are **LB_HeightBlend**, causing some areas to be black. On the right, the red "1" layer has been changed to use **LB_AlphaBlend**, which solves the problem.
 
 
-[/REGION]
-
-Below is an example of the properties of the **Landscape Layer Blend** node for a number of layers all being blended together. Make sure to note how the **Soil** layer has its blend mode set to **LB_AlphaBlend** while the other layers have theirs set to **LB_HeightBlend**. This is to stop the issue mention above, having black spots where layers meet, from happening. 
+Below is an example of the properties of the **Landscape Layer Blend** node for a number of layers all being blended together. Make sure to note how the **Soil** layer has its blend mode set to **LB_AlphaBlend** while the other layers have theirs set to **LB_HeightBlend**. This is to stop the issue mentioned above (having black spots where layers meet) from happening. 
 
 ![Layer Blend Properties](LayerBlendProperties.png)
 
 
-[REGION:tip]
 To delete a layer, click the **Drop-Down Arrow** to the right of the layer's element number to open the context menu, and then click **Delete**.
 
 ![Delete Layer](Landscape_DeleteBlendLayer2.png)
-[/REGION]
+
 
 [/EXCERPT:LandscapeLayerBlend]
 
@@ -169,7 +166,7 @@ You can set the properties of the LandscapeLayerWeight node in the **Details** p
 
 [EXCERPT:LandscapeVisibilityMask]
 
-The LandscapeVisibilityMask node is used for removing the visibility of parts of your Landscape, so you can create holes, for example, for creating caves. 
+The **LandscapeVisibilityMask** node is used for removing the visibility of parts of your Landscape, so you can create holes, for example, for creating caves. 
 
 ![LandscapeVisibilityMask Node](Landscape_VisibilityMask.png)
 
@@ -194,12 +191,12 @@ The ability to blend multiple textures and materials networks into paintable Lan
 
 Setting up and using **LandscapeLayerWeight** nodes is a quite easy process for which will be outlined below. 
 
-[REGION:tip]
+[REGION:note]
 Because Layer names must be unique, it is a best practice to give each layer a descriptive name as possible.
 [/REGION]
 
 
-1. In the Material Editor, add a LandscapeLayerWeight node to the Material Graph. By default, it will have the name **Parameter Name** "None."
+1. In the **Material Editor**, add a **LandscapeLayerWeight** node to the Material Graph. By default, it will have the **Parameter Name** "None."
 	[REGION:fullwidth]
 	![](Landscape_Layer_Weight_Base.png)
 	[/REGION]
@@ -214,11 +211,11 @@ Because Layer names must be unique, it is a best practice to give each layer a d
 
 1. Add and connect **Texture Samples**, or other material network expressions, to your **LandscapeLayerWeight** nodes.
 
-1. Add a **LandscapeLayerCoords** node, set its UV titling to what you want, and connect it to the Texture Sample nodes.
+1. Add a **LandscapeLayerCoords** node, set its UV titling to what you want, and connect it to the **Texture Sample** nodes.
 
-1. Drag from each layer node's output pin to the **Base** pin of the next layer node to connect them with a wire (leaving the first layer node's **Base** pin unconnected).
+1. Drag from each **Layer** node's output pin to the **Base** pin of the next layer node to connect them with a wire (leaving the first layer node's **Base** pin unconnected).
 
-1. Drag from the final layer node's output pin to the **Base Color** input of the Material's base node.
+1. Drag from the final **Layer** node's output pin to the **Base Color** input of the Material's **Base** node.
 
 
 When completed, you should have something that looks like the following example using **LandscapeLayerWeight** nodes:
