#UE4doc #docs - Added Engine Version field. Updated image to match UI changes. Replaced lightboxed image with Comparison slider.

diff --git a/Engine/Documentation/Source/Engine/Rendering/Materials/LightingModels/SubSurfaceProfile/SSSProfileLightingModel.INT.udn b/Engine/Documentation/Source/Engine/Rendering/Materials/LightingModels/SubSurfaceProfile/SSSProfileLightingModel.INT.udn
--- a/Engine/Documentation/Source/Engine/Rendering/Materials/LightingModels/SubSurfaceProfile/SSSProfileLightingModel.INT.udn
+++ b/Engine/Documentation/Source/Engine/Rendering/Materials/LightingModels/SubSurfaceProfile/SSSProfileLightingModel.INT.udn
@@ -2,6 +2,7 @@ Availability:Public
 Title:Subsurface Profile Shading Model
 Crumbs: %ROOT%, Engine, Engine/Rendering/Materials
 Description:Description and technical details of the subsurface profile shading model available in materials.
+Version: 4.9
 
 [TOC(start:2)]
 
@@ -12,11 +13,15 @@ While the Subsurface Profile shading model has similar properties to the Subsurf
 SubsurfaceProfile uses a rendering approach that is based in Screen Space because this helps to better display the subtle subsurface effects seen in human skin where backscattering is a secondary effect only seen in few cases like ears for example.
 The following document will cover what a Subsurface Profile is and how you can use them in your work. 
 
-[REGION:lightbox]
-[![](Results.png)(w:920 convert:false)](Results.png)
-[/REGION]
+[OBJECT:ComparisonSlider]
+	[PARAM:before]
+	![Not using SubsurfaceProfile](Results_1.png)(w:700 convert:false)
+	[/PARAM]
+	[PARAM:after]
+	![Using SubsurfaceProfile](Results_2.png)(w:700 convert:false)
+	[/PARAM]
+[/OBJECT]
 
-The Left Side shows a material not using the SubsurfaceProfile shading model. The Right Side shows an asset using the SubsurfaceProfile shading model. 
 
 [Region:note]
 Special thanks to [REDACTED_PERSON] and his company [Infinite Realities] (http://ir-ltd.net ) for providing the 3D scanned head model and assistance. 
