#UE4doc #docs - Added Engine Version field. Updated multiple images to mach latest UI changes. Bolded terms.

diff --git a/Engine/Documentation/Source/Engine/Rendering/Materials/MaterialInputs/MaterialInputs.INT.udn b/Engine/Documentation/Source/Engine/Rendering/Materials/MaterialInputs/MaterialInputs.INT.udn
--- a/Engine/Documentation/Source/Engine/Rendering/Materials/MaterialInputs/MaterialInputs.INT.udn
+++ b/Engine/Documentation/Source/Engine/Rendering/Materials/MaterialInputs/MaterialInputs.INT.udn
@@ -2,6 +2,7 @@ Availability:Public
 Title:Material Inputs			
 Crumbs: %ROOT%, Engine, Engine/Rendering/Materials, Engine/Rendering/Materials/IntroductionToMaterials
 Description:An in-depth look at the various inputs available on a Material and how they are used.
+Version: 4.9
 
 Here we take a look at the available inputs when making a Material. By feeding values (by way of constants, parameters, and textures) into these inputs, you can define just about any physical surface you can imagine.
 
diff --git a/Engine/Documentation/Source/Engine/Rendering/Materials/MaterialInstances/MaterialInstances.INT.udn b/Engine/Documentation/Source/Engine/Rendering/Materials/MaterialInstances/MaterialInstances.INT.udn
--- a/Engine/Documentation/Source/Engine/Rendering/Materials/MaterialInstances/MaterialInstances.INT.udn
+++ b/Engine/Documentation/Source/Engine/Rendering/Materials/MaterialInstances/MaterialInstances.INT.udn
@@ -2,6 +2,7 @@ Availability: Public
 Title:Instanced Materials
 Crumbs: %ROOT%, Engine, Engine/Rendering/Materials
 Description:Hierarchical material relationships that allow modification of properties to alter behavior and appearance of child materials without shader recompilation.
+Version: 4.9
 
 
 ![](MICEditor.png)
@@ -39,7 +40,7 @@ There are 2 types of Material Instances available in Unreal Engine 4:
 
 ### Material Instance Constant
 
-A Material Instance Constant is an instanced Material that calculates only once, prior to runtime. This means that it cannot change during gameplay. However, although they remain _constant_ throughout your game, they still have the performance advantage of not requiring compilation. For instance, if you have a game with a variety of cars, all with individual paint jobs, but whose colors will _not_ be changing during gameplay, the best practice approach would be to create a master Material that represented the base aspects of a generic car paint, and then create Material Instance Constants to represent the variations for different types of car, such as different colors, varying levels of metallicity, glossiness, etc.
+A **Material Instance Constant** is an instanced Material that calculates only once, prior to runtime. This means that it cannot change during gameplay. However, although they remain _constant_ throughout your game, they still have the performance advantage of not requiring compilation. For instance, if you have a game with a variety of cars, all with individual paint jobs, but whose colors will _not_ be changing during gameplay, the best practice approach would be to create a master Material that represented the base aspects of a generic car paint, and then create Material Instance Constants to represent the variations for different types of car, such as different colors, varying levels of metallicity, glossiness, etc.
 
 Material Instance Constants are created within the [Content Browser](Engine/Content/Browser) and are edited via the [Material Instance Editor](Engine/Rendering/Materials/MaterialInstances/Editor).
 
@@ -48,7 +49,7 @@ Material Instance Constants are created within the [Content Browser](Engine/Cont
 
 ### Material Instance Dynamic
 
-A Material Instance Dynamic (MID) is an instanced Material that can be calculated during gameplay (at runtime). This means that as you play, you can use script (either compiled code or Blueprint visual script) to change the parameters of your Material, thereby altering your Material throughout the game. The possible applications for this are endless, from showing different levels of damage to changing a paint job to blending in different skin textures in response to facial expressions.
+A **Material Instance Dynamic** (MID) is an instanced Material that can be calculated during gameplay (at runtime). This means that as you play, you can use script (either compiled code or Blueprint visual script) to change the parameters of your Material, thereby altering your Material throughout the game. The possible applications for this are endless, from showing different levels of damage to changing a paint job to blending in different skin textures in response to facial expressions.
 
 MIDs are created within script, either from a parameterized Material or a Material Instance Constant. In Blueprint, one would take a given Material that had parameterized properties, and feed it through a _Create Dynamic Material Instance_ node. The result of that node is then applied to the object in question via a _Set Material_ node, resulting in a new Material that can be changed during gameplay.
 
@@ -111,7 +112,7 @@ Give the parameter a unique name, assign it to a [group](#ParameterGroups), and
 
 ## Scalar Parameters
 
-A ScalarParameter is a parameter that contains a single floating-point value. These can be used to drive effects based on single values such as specular power, the alpha of a linear interpolation, opacity, etc.
+A **ScalarParameter** is a parameter that contains a single floating-point value. These can be used to drive effects based on single values such as specular power, the alpha of a linear interpolation, opacity, etc.
 
 See the [](Engine/Rendering/Materials/ExpressionReference) for a complete listing and descriptions of all expressions.
 
@@ -119,7 +120,7 @@ See the [](Engine/Rendering/Materials/ExpressionReference) for a complete listin
 
 ## Vector Parameters
 
-A VectorParameter is a parameter that contains a 4-channel vector value, or 4 floating-point values. These are generally used to provide configurable colors, but could also be used to represent positional data or drive any effect that requires multiple values.
+A **VectorParameter** is a parameter that contains a 4-channel vector value, or 4 floating-point values. These are generally used to provide configurable colors, but could also be used to represent positional data or drive any effect that requires multiple values.
 
 See the [](Engine/Rendering/Materials/ExpressionReference) for a complete listing and descriptions of all expressions.
 
@@ -146,11 +147,15 @@ See the [](Engine/Rendering/Materials/ExpressionReference) for a complete listin
 
 ## Static Parameters
 
-Static parameters are applied at compile-time, so they produce more optimal code as entire branches of the material that are masked out by a static parameter will be compiled out and not executed at runtime.  Because they are applied at compile-time, they can only be changed from within the [MaterialInstanceEditor](Engine/Rendering/Materials/MaterialInstances/Editor) and not from script.
-
+**Static** parameters are applied at compile-time, so they produce more optimal code as entire branches of the material that are masked out by a static parameter will be compiled out and not executed at runtime.  Because they are applied at compile-time, they can only be changed from within the [MaterialInstanceEditor](Engine/Rendering/Materials/MaterialInstances/Editor) and not from script.
+[Region:note]
 *Warning: A new material will be compiled out for every combination of static parameters in the base material that are used by instances!
 
-This can lead to an excessive amount of shaders being compiled.  Try to minimize the number of static parameters in the material and the number of permutations of those static parameters that are actually used.  See [Static Switch Parameter](Engine/Rendering/Materials/ExpressionReference/Parameters#StaticSwitchParameter) and [Static Component Mask Parameter](Engine/Rendering/Materials/ExpressionReference/Parameters#StaticComponentMaskParameter) for information on the specific static parameter types.
+This can lead to an excessive amount of shaders being compiled.  Try to minimize the number of static parameters in the material and the number of permutations of those static parameters that are actually used.  
+[/Region]
+
+
+See [Static Switch Parameter](Engine/Rendering/Materials/ExpressionReference/Parameters#StaticSwitchParameter) and [Static Component Mask Parameter](Engine/Rendering/Materials/ExpressionReference/Parameters#StaticComponentMaskParameter) for information on the specific static parameter types.
 
 See the [](Engine/Rendering/Materials/ExpressionReference) for a complete listing and descriptions of all expressions.
 
