#UE4doc #doc - Added engine version field. Added missing options that have been added to the Project Settings.

diff --git a/Engine/Documentation/Source/Engine/UI/ProjectSettings/ProjectSettings.INT.udn b/Engine/Documentation/Source/Engine/UI/ProjectSettings/ProjectSettings.INT.udn
--- a/Engine/Documentation/Source/Engine/UI/ProjectSettings/ProjectSettings.INT.udn
+++ b/Engine/Documentation/Source/Engine/UI/ProjectSettings/ProjectSettings.INT.udn
@@ -2,6 +2,7 @@ Availability:Public
 Title:Project Settings
 Crumbs:%ROOT%, Engine, Engine/UI
 Description:Dialog for configuring options specific to the current project.
+Version: 4.9
 
 [VAR:Topic]
 [OBJECT:Topic]
@@ -88,8 +89,9 @@ Contains options for specific systems and sections of the engine.
 | Section | Description |
 | ------- | ----------- |
 | **AI System** | Contains basic options for the AI system. |
-| **Animation** |  |
+| **Animation** | Contains settings for animations used in the editor. |
 | **Audio** | Contains options for setting default and adding quality levels.  |
+| **Collision** | Allows you to setup and modify various collision settings. |
 | **Console** | Allows you to set your own in-game inputs for the console. |
 | **Cooker** | Contains options for various cooker settings. These will handle things like compression quality for textures. |
 | **Crowd Manager** | Contains a set of configurable options for managing AI crowds. |
@@ -114,7 +116,8 @@ Contains options for how the editor should look and feel on a per-project basis.
 | ------- | ----------- |
 | **2D** | Contains configuration options for how the 2D level editor should work.  |
 | **Appearance** | Contains options for the appearance of the editor. There is options to change how units are displayed in the project, such as, choosing to display inches vs centimeters or  celcius vs farenheit. |
-| **Paper2D** | Allows you to set how 2D sprites should be handled when they are imported.  |
+| **Paper2D - Import** | Allows you to set how 2D sprites should be handled when they are imported.  |
+| **Sequencer** | Allows you to configure settings for the Sequencer plugin |
 
 
 ### Platforms Category
