#UE4doc #docs - Added engine version field. Updated image to match latest UI changes. Fixed step instructing where to locate in context menu.

diff --git a/Engine/Documentation/Source/Engine/UI/PropertyMatrix/PropertyMatrix.INT.udn b/Engine/Documentation/Source/Engine/UI/PropertyMatrix/PropertyMatrix.INT.udn
--- a/Engine/Documentation/Source/Engine/UI/PropertyMatrix/PropertyMatrix.INT.udn
+++ b/Engine/Documentation/Source/Engine/UI/PropertyMatrix/PropertyMatrix.INT.udn
@@ -2,6 +2,7 @@ Availability:Public
 Title:Property Matrix
 Crumbs:%ROOT%, Engine, Engine/UI
 Description:Tool for viewing and editing multiple properties of multiple actors at the same time.
+Version: 4.9
 
 [VAR:Topic]
 [OBJECT:Topic]
@@ -94,7 +95,7 @@ the **Search** box will launch a **Property Matrix** bound with your current sel
 
 	![](open_matrix_details.png)
 
-* In the [](Engine/Content/Browser), make a selection of assets, **right-click**, and choose **Property Matrix** 
+* In the [](Engine/Content/Browser), make a selection of assets, **right-click**, choose **Asset Actions**, then choose **Bulk Edit via Property Matrix**. 
 from the context menu.
 
 	![](open_matrix_browser.png)
