#UE4doc #docs - Added Engine version field. Bolded Terms.

diff --git a/Engine/Documentation/Source/Gameplay/Framework/Camera/Animations/CameraAnim.INT.udn b/Engine/Documentation/Source/Gameplay/Framework/Camera/Animations/CameraAnim.INT.udn
--- a/Engine/Documentation/Source/Gameplay/Framework/Camera/Animations/CameraAnim.INT.udn
+++ b/Engine/Documentation/Source/Gameplay/Framework/Camera/Animations/CameraAnim.INT.udn
@@ -2,8 +2,9 @@ Availability:Public
 Title:CameraAnims
 Crumbs:%ROOT%, Gameplay, Gameplay/Framework, Gameplay/Framework/Camera
 Description:CameraAnims allow you to layer on animations to your camera to simulate impact, motion in your environment, and other effects.
+Version: 4.9
 
-Conceptually, a CameraAnim is simply an animation that can be layered onto the in-game camera.  You can animate the camera position and rotation, FOV, and post process settings.
+Conceptually, a **CameraAnim** is simply an animation that can be layered onto the in-game camera.  You can animate the camera position and rotation, FOV, and post process settings.
 A CameraAnim's position, rotation, and FOV changes are applied additively to the game camera.  These keyframes are treated as "relative to initial", which means that only the delta 
 from time 0.0 in the animation will be applied.
 
@@ -20,7 +21,7 @@ To edit a CameraAnim, simply **double-click** the asset in the **Content Browser
 CameraAnim, a temporary camera is placed in the world that you can manipulate to create keys and use to preview your animation.
 
 ## Playing a CameraAnim
-The PlayerCameraManager contains functions that allow you to control the application of CameraAnims to the game camera from either C++ code or Blueprint script.  When you play a CameraAnim, 
+The **PlayerCameraManager** contains functions that allow you to control the application of CameraAnims to the game camera from either C++ code or Blueprint script.  When you play a CameraAnim, 
 a CameraAnimInst is created which contains the playback information for an active instance of the animation.  You can use this object to stop the animation manually or modify its playback parameters.
 
 Multiple CameraAnimInsts (up to 8 currently) can be active at once, all blending and contributing to the final camera settings.
