#UE4doc #docs - Added Engine version field.

diff --git a/Engine/Documentation/Source/Gameplay/Framework/Controller/Controller.INT.udn b/Engine/Documentation/Source/Gameplay/Framework/Controller/Controller.INT.udn
--- a/Engine/Documentation/Source/Gameplay/Framework/Controller/Controller.INT.udn
+++ b/Engine/Documentation/Source/Gameplay/Framework/Controller/Controller.INT.udn
@@ -6,6 +6,7 @@ Related: Gameplay\Framework\Controller\PlayerController
 Related: Gameplay\Framework\Controller\AIController
 Related: Gameplay\Framework\Pawn
 Related: Gameplay\Framework\Pawn\Character
+Version: 4.9
 
 
 [INCLUDE:Shared/Glossary/C/#Glossary_Controller_Class]
