#UE4doc #docs - added engine version field.

diff --git a/Engine/Documentation/Source/Gameplay/Framework/Controller/AIController/AIController.INT.udn b/Engine/Documentation/Source/Gameplay/Framework/Controller/AIController/AIController.INT.udn
--- a/Engine/Documentation/Source/Gameplay/Framework/Controller/AIController/AIController.INT.udn
+++ b/Engine/Documentation/Source/Gameplay/Framework/Controller/AIController/AIController.INT.udn
@@ -6,6 +6,7 @@ Related: Gameplay\Framework\Controller\PlayerController
 Related: Gameplay\Framework\Controller
 Related: Gameplay\Framework\Pawn
 Related: Gameplay\Framework\Pawn\Character
+Version: 4.9
 
 
 While the **PlayerController** relies on the human player to make decisions on what to do, the **AIController** is more focused on responding to input from the environment and game world. The job of the AIController is to observe the world around it and make decisions and react accordingly without explicit input from a human player. 
