#UE4doc #docs - Added Engine version field. Bolded terms.

diff --git a/Engine/Documentation/Source/Gameplay/Framework/Controller/PlayerController/PlayerController.INT.udn b/Engine/Documentation/Source/Gameplay/Framework/Controller/PlayerController/PlayerController.INT.udn
--- a/Engine/Documentation/Source/Gameplay/Framework/Controller/PlayerController/PlayerController.INT.udn
+++ b/Engine/Documentation/Source/Gameplay/Framework/Controller/PlayerController/PlayerController.INT.udn
@@ -6,11 +6,11 @@ Related: Gameplay\Framework\Controller\AIController
 Related: Gameplay\Framework\Controller
 Related: Gameplay\Framework\Pawn
 Related: Gameplay\Framework\Pawn\Character
+Version: 4.9
 
+A **PlayerController** is the interface between the Pawn and the human player controlling it.  The PlayerController essentially represents the human player's will.
 
-A PlayerController is the interface between the Pawn and the human player controlling it.  The PlayerController essentially represents the human player's will.
-
-One thing to consider when setting up your PlayerController is what functionality should be in the PlayerController, and what should be in your Pawn.  It is possible to handle all input in the 
+One thing to consider when setting up your PlayerController is what functionality should be in the PlayerController, and what should be in your **Pawn**.  It is possible to handle all input in the 
 Pawn, especially for less complex cases.  However, if you have more complex needs, like multiple players on one game client, or the ability to change characters dynamically at runtime, it might be
 better to handle input in the PlayerController. In this case, the PlayerController decides what to do and then issues commands to the Pawn (e.g. "start crouching", "jump").
 
