#UE4doc #docs - Added Engine version field. Bolded term. Updated Image to match latest UI changes.

diff --git a/Engine/Documentation/Source/Gameplay/Framework/GameMode/GameMode.INT.udn b/Engine/Documentation/Source/Gameplay/Framework/GameMode/GameMode.INT.udn
--- a/Engine/Documentation/Source/Gameplay/Framework/GameMode/GameMode.INT.udn
+++ b/Engine/Documentation/Source/Gameplay/Framework/GameMode/GameMode.INT.udn
@@ -2,6 +2,7 @@ Availability:Public
 Title:GameMode
 Crumbs:%ROOT%, Gameplay/Framework
 Description:The GameMode is responsible for setting the rules of the game.
+Version: 4.9
 
 Even the most open-ended game has an underpinning of rules, and these rules make up a gametype.  On the most basic level, these rules include:
 
