#UE4doc #docs - Added Engine version field. Bolded term.

diff --git a/Engine/Documentation/Source/Gameplay/Framework/GameState/GameState.INT.udn b/Engine/Documentation/Source/Gameplay/Framework/GameState/GameState.INT.udn
--- a/Engine/Documentation/Source/Gameplay/Framework/GameState/GameState.INT.udn
+++ b/Engine/Documentation/Source/Gameplay/Framework/GameState/GameState.INT.udn
@@ -2,9 +2,10 @@ Availability:Public
 Title:GameState
 Crumbs:%ROOT%, Gameplay/Framework
 Description:The GameState is responsible for monitoring the state of the game on the client.
+Version: 4.9
 
 
-The GameState is responsible for monitoring the state of the game on the client.
+The **GameState** is responsible for monitoring the state of the game on the client.
 
 Conceptually, think of the GameState as the state of the game. It can keep track of properties like score, list of connected players, number of caps in a CTF game, where the pieces are in a chess game, 
 what missions you have completed in an open world game, etc. In general, the GameState should track properties that change during gameplay. The GameState exists on the server and all clients and can
