#UE4doc #docs  - Added engine version field. Bolded terms.

diff --git a/Engine/Documentation/Source/Gameplay/Framework/Pawn/Pawn.INT.udn b/Engine/Documentation/Source/Gameplay/Framework/Pawn/Pawn.INT.udn
--- a/Engine/Documentation/Source/Gameplay/Framework/Pawn/Pawn.INT.udn
+++ b/Engine/Documentation/Source/Gameplay/Framework/Pawn/Pawn.INT.udn
@@ -5,6 +5,7 @@ Description:The Pawn is the physical representation of a player within the world
 Related: Gameplay/Framework\Pawn\Character
 Related: Gameplay/Framework/Controller
 Related: Gameplay/Framework/GameMode
+Version: 4.9
 
 
 ![](pawn_ball.png)(w:300)
@@ -18,7 +19,7 @@ represents the physical location, rotation, etc. of a player or entity within th
 By default, there is a one-to-one relationship between [Controllers](Gameplay/Framework/Controller) and Pawns; meaning, each Controller controls only one Pawn at any given time. Also, Pawns spawned during gameplay 
 are not automatically possessed by a Controller.
 
-In **Blueprints**, the best way to add movement to your Pawn-derived class is with **SetActorLocation**. With **SetActorLocation**, you can decide if you want to teleport to the location, or sweep. 
+In **Blueprints**, the best way to add movement to your Pawn-derived class is with **SetActorLocation**. With SetActorLocation, you can decide if you want to teleport to the location, or sweep. 
 When sweeping, your Pawn will move along the direction and stop if it hits something.
 
 
