#UE4doc #docs - Added Engine version field.

diff --git a/Engine/Documentation/Source/Gameplay/Framework/Pawn/Character/Character.INT.udn b/Engine/Documentation/Source/Gameplay/Framework/Pawn/Character/Character.INT.udn
--- a/Engine/Documentation/Source/Gameplay/Framework/Pawn/Character/Character.INT.udn
+++ b/Engine/Documentation/Source/Gameplay/Framework/Pawn/Character/Character.INT.udn
@@ -8,6 +8,7 @@ Related: Gameplay/Framework/GameMode
 Related: Engine/Animation/RootMotion
 Related: Engine/Animation/CharacterSetupOverview
 Related: Gameplay/HowTo/CharacterMovement
+Version: 4.9
 
 With the addition of a `CharacterMovementComponent`, a `CapsuleComponent`, and a **Skeletal Mesh**, the **Pawn** class is extended into the highly-featured **Character** class.
 A Character is designed for a vertically-oriented player representation that can walk, run, jump, fly, and swim through the world. This class also contains 
