#UE4doc #docs - Added Engine version field. Bolded terms.

diff --git a/Engine/Documentation/Source/Gameplay/Framework/QuickReference/QuickReference.INT.udn b/Engine/Documentation/Source/Gameplay/Framework/QuickReference/QuickReference.INT.udn
--- a/Engine/Documentation/Source/Gameplay/Framework/QuickReference/QuickReference.INT.udn
+++ b/Engine/Documentation/Source/Gameplay/Framework/QuickReference/QuickReference.INT.udn
@@ -2,9 +2,10 @@ Availability:Public
 Title:Gameplay Framework Quick Reference
 Crumbs:%ROOT%, Gameplay, Gameplay/Framework
 Description:Short overview of classes for game rules, characters, controllers, user interfaces, etc., that make up the framework for the game.
+Version: 4.9
 
 The basic gameplay classes include functionality for representing players, allies, and enemies, as well as for controlling these avatars with player input or AI logic. There are also classes for creating 
-heads-up displays and cameras for players. Finally, gameplay classes such as GameMode, GameState, and PlayerState set the rules of the game, and track how the game and the players are progressing. 
+heads-up displays and cameras for players. Finally, gameplay classes such as **GameMode**, **GameState**, and **PlayerState** set the rules of the game, and track how the game and the players are progressing. 
 
 These classes all create types of Actors, which can either be placed into your level or spawned when needed.
 
@@ -90,7 +91,7 @@ These classes all create types of Actors, which can either be placed into your l
 
 		**PlayerState** 
 		
-		A PlayerState is the state of a participant in the game, such as a human player or a bot that is simulating a player. Non-player AI that exists as part of the game would not have a PlayerState. Example data that would be appropriate in a PlayerState include player name, score, in-match level for something like a MOBA, or whether the player is currently carrying the flag in a CTF game. PlayerStates for all players exist on all machines (unlike PlayerControllers) and can replicate freely to keep things in sync.
+		A **PlayerState** is the state of a participant in the game, such as a human player or a bot that is simulating a player. Non-player AI that exists as part of the game would not have a PlayerState. Example data that would be appropriate in a PlayerState include player name, score, in-match level for something like a MOBA, or whether the player is currently carrying the flag in a CTF game. PlayerStates for all players exist on all machines (unlike PlayerControllers) and can replicate freely to keep things in sync.
 
 	[/PARAM]
 [/OBJECT]
