#UE4doc #docs - Added engine version field. Bolded terms.

diff --git a/Engine/Documentation/Source/Gameplay/Framework/UIAndHUD/UIAndHUD.INT.udn b/Engine/Documentation/Source/Gameplay/Framework/UIAndHUD/UIAndHUD.INT.udn
--- a/Engine/Documentation/Source/Gameplay/Framework/UIAndHUD/UIAndHUD.INT.udn
+++ b/Engine/Documentation/Source/Gameplay/Framework/UIAndHUD/UIAndHUD.INT.udn
@@ -2,6 +2,7 @@ Availability: Public
 Title:User Interfaces & HUDs
 Crumbs:%ROOT%, Gameplay, Gameplay/Framework
 Description:Guides and information for artists and programmers creating user interfaces such as menus and HUDs.
+Version: 4.9
 
 <!--
 [REGION:header_img]
@@ -23,9 +24,9 @@ Description:Guides and information for artists and programmers creating user int
 [/OBJECT]
 -->
 
-The way in which the game communicates and interacts with the player is extremely important. User interfaces (UIs) and heads-up displays (HUDs) are the games' way of providing information about the game to the player and in some cases allowing the player to interact with the game.
+The way in which the game communicates and interacts with the player is extremely important. **User Interfaces** (UIs) and **Heads-up Displays** (HUDs) are the games' way of providing information about the game to the player and in some cases allowing the player to interact with the game.
 
-Unreal Engine 4 provides multiple means of creating UIs and HUDs. The Canvas class can be used to draw directly to the screen at a low level, overlaid onto the world.
+Unreal Engine 4 provides multiple means of creating UIs and HUDs. The `Canvas` class can be used to draw directly to the screen at a low level, overlaid onto the world.
 
 
 
