#UE4doc #docs - Added Engine version field.

diff --git a/Engine/Documentation/Source/Gameplay/HowTo/ActorInput/ActorInputOverview.INT.udn b/Engine/Documentation/Source/Gameplay/HowTo/ActorInput/ActorInputOverview.INT.udn
--- a/Engine/Documentation/Source/Gameplay/HowTo/ActorInput/ActorInputOverview.INT.udn
+++ b/Engine/Documentation/Source/Gameplay/HowTo/ActorInput/ActorInputOverview.INT.udn
@@ -2,6 +2,7 @@ Availability: Public
 Title:Setting Up Input on an Actor
 Crumbs: %ROOT%, Gameplay
 Description:A How To Guide to Setting up Input for an Actor in Unreal Engine 4. 
+Version: 4.9
 
 When developing your game, there may be instances when you want to allow the player to perform some form of **Input** on an Actor in your level. For example, perhaps you have a treasure chest that you want the player to be able to open when they approach it, or a light you want them to be able to turn on/off, or some other form of interactable object you want to affect in some way when the player presses a button. 
 
