#UE4doc #docs -  Added engine version field. Added skill level field. Bolded terms.

diff --git a/Engine/Documentation/Source/Gameplay/HowTo/ActorInput/Blueprints/ActorInput_BP.INT.udn b/Engine/Documentation/Source/Gameplay/HowTo/ActorInput/Blueprints/ActorInput_BP.INT.udn
--- a/Engine/Documentation/Source/Gameplay/HowTo/ActorInput/Blueprints/ActorInput_BP.INT.udn
+++ b/Engine/Documentation/Source/Gameplay/HowTo/ActorInput/Blueprints/ActorInput_BP.INT.udn
@@ -2,12 +2,14 @@ Availability: Public
 Title:Setting Up Input on an Actor in Blueprints
 Crumbs: %ROOT%, Gameplay, Gameplay\HowTo\ActorInput
 Description:A How To Guide for setting up input on an Actor in Blueprints. 
+Version: 4.9
+SkillLevel: Intermediate
 
 [TOC (start:2 end:2)]
 
 This page covers the basic methods on how to **Enable Input** or **Disable Input** for an Actor. By Enabling Input on an Actor, you can allow a player to press a button or key and execute events that affect the Actor in some way (be it turn a light on or off, open or close a door, activate something, etc.). 
 
-The Enable/Disable Input Usage section also covers an approach to Enabling/Disabling Input using a Trigger Volume to determine if a player is near the Actor in the level. This is useful when you want a player to affect the Actor only when they are near or in range of the Actor to conceivably affect it in a real-world scenario. 
+The Enable/Disable Input Usage section also covers an approach to Enabling/Disabling Input using a **Trigger Volume** to determine if a player is near the Actor in the level. This is useful when you want a player to affect the Actor only when they are near or in range of the Actor to conceivably affect it in a real-world scenario. 
  
 
 ## Enabling Actor Input
