#UE4doc #docs - Added engine version field. Added skill level field.

diff --git a/Engine/Documentation/Source/Gameplay/HowTo/AddingComponents/Blueprints/Components_BP.INT.udn b/Engine/Documentation/Source/Gameplay/HowTo/AddingComponents/Blueprints/Components_BP.INT.udn
--- a/Engine/Documentation/Source/Gameplay/HowTo/AddingComponents/Blueprints/Components_BP.INT.udn
+++ b/Engine/Documentation/Source/Gameplay/HowTo/AddingComponents/Blueprints/Components_BP.INT.udn
@@ -2,6 +2,8 @@ Availability: Public
 Title:Adding Components in Blueprints
 Crumbs: %ROOT%, Gameplay, Gameplay\HowTo\AddingComponents
 Description:A How To Guide for Adding Components in Blueprints. 
+Version: 4.9
+SkillLevel: Intermediate
 
 [TOC (start:2 end:2)]
 
