#UE4 Docs: Adding a new section about Area Shadows with images.

diff --git a/Engine/Documentation/Source/Engine/Rendering/LightingAndShadows/LightMobility/StationaryLights/StationaryLights.INT.udn b/Engine/Documentation/Source/Engine/Rendering/LightingAndShadows/LightMobility/StationaryLights/StationaryLights.INT.udn
--- a/Engine/Documentation/Source/Engine/Rendering/LightingAndShadows/LightMobility/StationaryLights/StationaryLights.INT.udn
+++ b/Engine/Documentation/Source/Engine/Rendering/LightingAndShadows/LightMobility/StationaryLights/StationaryLights.INT.udn
@@ -170,3 +170,24 @@ light has **Visible** unchecked, its indirect lighting will be put in the lightm
 **IndirectLightingIntensity** on the light can be used to scale or disable the indirect lighting from a given light at lighting build time.  
 
 However there's a post process volume setting called **IndirectLightingIntensity** which lets you scale the contribution of the lightmap for all lights, which can be changed at runtime from a blueprint. 
+
+
+## Area Shadows
+With the 4.9 release of Unreal Engine, Stationary Directional Lights have a new shadowing option located in the Lightmass options called **Use Area Shadows for Stationary Lights**. 
+
+![](Area_Shadows_Options.png)
+
+When the Use Area Shadows for Stationary Lights option is enabled, the Stationary Light will use area shadows for the precomputed shadow maps. 
+Area shadows are shadows that get softer the further they are from shadow casters. 
+In the following image you see the difference between the two shadowing methods.
+
+[OBJECT:ComparisonSlider]
+	[PARAM:before]
+	![Area Shadows Enabled](Area_Shadows_On.png)(w:900)
+	[/PARAM]
+	[PARAM:after]
+	![Area Shadows Disabled](Area_Shadows_Off.png)(w:900)
+	[/PARAM]
+[/OBJECT]
+
+Note that Area Shadows will only work on Stationary Lights and you might have to increase some objects lightmap resolution to get the same shadow quality and sharpness.
\ No newline at end of file
