Fixing up the flow of the area lights description and re-doing the title to say Use Area Shadows for Stationary Lights instaed of just Area Lights.

diff --git a/Engine/Documentation/Source/Engine/Rendering/LightingAndShadows/LightMobility/StationaryLights/StationaryLights.INT.udn b/Engine/Documentation/Source/Engine/Rendering/LightingAndShadows/LightMobility/StationaryLights/StationaryLights.INT.udn
--- a/Engine/Documentation/Source/Engine/Rendering/LightingAndShadows/LightMobility/StationaryLights/StationaryLights.INT.udn
+++ b/Engine/Documentation/Source/Engine/Rendering/LightingAndShadows/LightMobility/StationaryLights/StationaryLights.INT.udn
@@ -172,13 +172,15 @@ light has **Visible** unchecked, its indirect lighting will be put in the lightm
 However there's a post process volume setting called **IndirectLightingIntensity** which lets you scale the contribution of the lightmap for all lights, which can be changed at runtime from a blueprint. 
 
 
-## Area Shadows
+## Use Area Shadows for Stationary Lights
 With the 4.9 release of Unreal Engine, Stationary Directional Lights have a new shadowing option located in the Lightmass options called **Use Area Shadows for Stationary Lights**. 
 
 ![](Area_Shadows_Options.png)
 
-When the Use Area Shadows for Stationary Lights option is enabled, the Stationary Light will use area shadows for the precomputed shadow maps. 
-Area shadows are shadows that get softer the further they are from shadow casters. 
+To Enable the Use Area Shadows for Stationary Lights option, first select the **Directional Light** in your scene and make sure that it's **Mobility** is set to **Stationary**.
+Then in the **Lightmass** section of the Directional light,  click on the box next the Use Area Shadows for Stationary Lights option to enable it. 
+When the Use Area Shadows for Stationary Lights option is enabled, the Stationary Light will use area shadows for the precomputed shadow maps.
+Area shadows are shadows that get softer the further they are from shadow casters.
 In the following image you see the difference between the two shadowing methods.
 
 [OBJECT:ComparisonSlider]
