#UE4doc #docs - Added Engine Version field. Added Skill Level Field. Updated image. Reported issue with UEDOC-2454

diff --git a/Engine/Documentation/Source/Gameplay/HowTo/CharacterMovement/Blueprints/Overview.INT.udn b/Engine/Documentation/Source/Gameplay/HowTo/CharacterMovement/Blueprints/Overview.INT.udn
--- a/Engine/Documentation/Source/Gameplay/HowTo/CharacterMovement/Blueprints/Overview.INT.udn
+++ b/Engine/Documentation/Source/Gameplay/HowTo/CharacterMovement/Blueprints/Overview.INT.udn
@@ -2,6 +2,8 @@ Availability: Public
 Title:Setting Up Character Movement in Blueprints
 Crumbs: %ROOT%, Gameplay, Gameplay\HowTo\CharacterMovement
 Description:A How To Guide for Setting Up Character Movement in Blueprints. 
+Version: 4.9
+SkillLevel: Advanced
 
  
 This "How To" series covers all the steps needed to create a fully playable and animated character through **Blueprints**. By the end of this How To series, you will have a character that can Walk, Run, Crouch and Crouch Walk, Go Prone, Jump from a stationary standing position, or Jump while moving. 
@@ -22,7 +24,7 @@ EhspjHjDmtI
  
 For this series, we are using the **Animation Starter Pack**, which is available for free on the **Marketplace**, as it contains several animations that can be used to create character movement. Once you have downloaded the Animation Pack, you can add the pack to an existing project (preferably a New Blank Project) via the **Add to project** button in the Launcher then specifying which project you want to add the assets to. 
 
-![](Characters_1.png)
+![](Characters_2.png)
 
 To get started, follow the steps below. 
 
