#UE4doc #docs - Added Engine Version Field. Unbolded some terms.

diff --git a/Engine/Documentation/Source/Gameplay/HowTo/FindingActors/FindingActorsOverview.INT.udn b/Engine/Documentation/Source/Gameplay/HowTo/FindingActors/FindingActorsOverview.INT.udn
--- a/Engine/Documentation/Source/Gameplay/HowTo/FindingActors/FindingActorsOverview.INT.udn
+++ b/Engine/Documentation/Source/Gameplay/HowTo/FindingActors/FindingActorsOverview.INT.udn
@@ -2,6 +2,8 @@ Availability: Public
 Title:Finding Actors
 Crumbs: %ROOT%, Gameplay
 Description:A How To Guide for Finding Actors in your Scenes in Unreal Engine 4. 
+Version: 4.9
+
 
 Once you start constructing your game and its levels, you may find that locating an individual Actor in your scene may be difficult. It is not uncommon to have hundreds of Actors in your scene ranging from level set pieces or geometry, to NPCs or enemies, to interactable objects or pick-ups and locating one Actor out of all those is like the old saying of "looking for a needle in a haystack". 
 
@@ -9,15 +11,15 @@ You can however use the **World Outliner** to help narrow your search for an Act
 
 ![](Actors1.png)(w:800)
 
-In the image above, the **World Outliner** (right) represents all the Actors that are placed inside your level. Clicking on one of the Actors in the **World Outliner** highlights it inside the level viewport as well. You can sort the **World Outliner** by Actor Name or by Actor Type by clicking on the **Actor** or **Type** headers (you can change the secondary category from Type to other filters by clicking the down arrow in the Type Box). 
+In the image above, the World Outliner (right) represents all the Actors that are placed inside your level. Clicking on one of the Actors in the World Outliner highlights it inside the level viewport as well. You can sort the World Outliner by Actor Name or by Actor Type by clicking on the **Actor** or **Type** headers (you can change the secondary category from Type to other filters by clicking the down arrow in the Type Box). 
 
 [REGION:tip]
-With an Actor selected in the **World Outliner**, and while navigating inside the Level Viewport, pressing the **F** key will move the camera to the location of the Actor selected to **Focus** on it. 
+With an Actor selected in the World Outliner, and while navigating inside the Level Viewport, pressing the **F** key will move the camera to the location of the Actor selected to **Focus** on it. 
 [/REGION]
 
 Refer to the [](Engine/UI/LevelEditor/SceneOutliner) documentation for more information on locating an Actor in your scenes. 
 
-While finding an Actor in the Editor can be accomplished with the help of the **World Outliner**, how do you go about finding an Actor through Blueprints so that you may modify it in some way? This "How To" series will show you some ways in which you can search your scene for an Actor based on varying criteria at which point you can then access and modify it if you wish. 
+While finding an Actor in the Editor can be accomplished with the help of the World Outliner, how do you go about finding an Actor through Blueprints so that you may modify it in some way? This "How To" series will show you some ways in which you can search your scene for an Actor based on varying criteria at which point you can then access and modify it if you wish. 
 
 
 ## Setup
