#UE4doc #docs - Added Engine Version field. Added Skill Level Field. Fixed grammatical error in sentence for Get Specific Actor

diff --git a/Engine/Documentation/Source/Gameplay/HowTo/FindingActors/Blueprints/FindingActors_BP.INT.udn b/Engine/Documentation/Source/Gameplay/HowTo/FindingActors/Blueprints/FindingActors_BP.INT.udn
--- a/Engine/Documentation/Source/Gameplay/HowTo/FindingActors/Blueprints/FindingActors_BP.INT.udn
+++ b/Engine/Documentation/Source/Gameplay/HowTo/FindingActors/Blueprints/FindingActors_BP.INT.udn
@@ -2,6 +2,8 @@ Availability: Public
 Title:Finding Actors in Blueprints
 Crumbs: %ROOT%, Gameplay, Gameplay\HowTo\FindingActors
 Description:A How To Guide for Finding Actors in Blueprints. 
+Version: 4.9
+SkillLevel: Intermediate
 
 [TOC (start:2 end:2)]
 
@@ -9,8 +11,9 @@ This page covers examples on how to find Actors in a level by using the Get All
  
 
 ## Get All Actors of Class
-
+<!---
 The steps below will show you how you can use the **Get All Actors of Class** node to find all Actors in your scene of a particular Class. 
+--->
 
  _For this example, we are using the **Blueprint Third Person Template with Starter Content**._
 
@@ -107,7 +110,7 @@ One method you can take is by using **Tags** on the Actor, for example:
 
 	_Click image for full view._
 
-	Here we are stating that after we have collected all Actors of the Class specified and put them in an Array, we then check each Actor to see if it has the Tag we specified and ONLY if it does, do we deactivate it. 
+	Here we are stating that after we have collected all Actors of the Class specified and put them in an Array, we then check each Actor to see if it has the Tag we specified and ONLY if it does, we deactivate it. 
 
 Here is another example of filtering results: 
 
