#UE4doc #docs - Added Engine Version Field. Bolded Term.

diff --git a/Engine/Documentation/Source/Gameplay/HowTo/Networking/ReplicateActor/ReplicateActorOverview.INT.udn b/Engine/Documentation/Source/Gameplay/HowTo/Networking/ReplicateActor/ReplicateActorOverview.INT.udn
--- a/Engine/Documentation/Source/Gameplay/HowTo/Networking/ReplicateActor/ReplicateActorOverview.INT.udn
+++ b/Engine/Documentation/Source/Gameplay/HowTo/Networking/ReplicateActor/ReplicateActorOverview.INT.udn
@@ -2,8 +2,9 @@ Availability: Public
 Title:How to Replicate Actors
 Crumbs: %ROOT%, Gameplay
 Description:Guide for Replicating Actors in Multiplayer Games. 
+version: 4.9
 
-The process of passing data or commands between Server and Clients is referred to as Replication, and Actor Replication is where the Server identifies and acknowledges if an Actor is set to use Replication or not. If the Actor is set to Replicate, the Server will tell the Client to create and maintain a Replicated copy of that Actor on the Client's machine. The Server always has the authoritative version of that Actor, while a Client (or multiple Clients) may have the approximate, replicated versions of that Actor. 
+The process of passing data or commands between Server and Clients is referred to as **Replication**, and Actor Replication is where the Server identifies and acknowledges if an Actor is set to use Replication or not. If the Actor is set to Replicate, the Server will tell the Client to create and maintain a Replicated copy of that Actor on the Client's machine. The Server always has the authoritative version of that Actor, while a Client (or multiple Clients) may have the approximate, replicated versions of that Actor. 
 
 Options can be set to define how Replicated Actors behave on a Client, from whether or not to Replicate Movement (the location/rotation of the Actor) to setting its Network Relevancy (or when to send Replicated Data to the Client, as the Actor set to Replicate may or may not be relevant depending upon gameplay situations). 
 
