#UE4 Docs: removing unused page

diff --git a/Engine/Documentation/Source/Engine/Animation/Animation.INT.udn b/Engine/Documentation/Source/Engine/Animation/Animation.INT.udn
--- a/Engine/Documentation/Source/Engine/Animation/Animation.INT.udn
+++ b/Engine/Documentation/Source/Engine/Animation/Animation.INT.udn
@@ -133,7 +133,7 @@ Navigation:topic
 		* [](Engine/Content/FBX/Animations "%Engine/Content/FBX/Animations:description%")
 		* [](Engine/Content/FBX/MorphTargets "%Engine/Content/FBX/MorphTargets:description%")
 		* **Programming**		
-			* [](Programming/Animation "%Programming/Animation:description%")
+			* [](Programming/Animation/Nodes "%Programming/Animation/Nodes:description%")
 		* **Tutorials**		
 			* [](Engine/Animation/CharacterSetupOverview "%Engine/Animation/CharacterSetupOverview:description%")
 			* [](Engine/Animation/IKSetups "%Engine/Animation/IKSetups:description%")
diff --git a/Engine/Documentation/Source/Programming/Animation/AnimationProgramming.INT.udn b/Engine/Documentation/Source/Programming/Animation/AnimationProgramming.INT.udn
deleted file mode 100644
diff --git a/Engine/Documentation/Source/Programming/ProgrammingGuide.INT.udn b/Engine/Documentation/Source/Programming/ProgrammingGuide.INT.udn
--- a/Engine/Documentation/Source/Programming/ProgrammingGuide.INT.udn
+++ b/Engine/Documentation/Source/Programming/ProgrammingGuide.INT.udn
@@ -207,7 +207,7 @@ C++ classes can be used as a base for Blueprint classes, and in this way program
 	[PARAM:description]
 	[/PARAM]
 	[PARAM:links]
-		* [](Programming/Animation "%Programming/Animation:description%")
+		* [](Programming/Animation/Nodes "%Programming/Animation/Nodes:description%")
 		* [](Programming/Automation "%Programming/Automation:description%")
 		* [](Engine/Blueprints/TechnicalGuide "%Engine/Blueprints/TechnicalGuide:description%")
 		* [](Programming/Online "%Programming/Online:description%")
