#UE4doc #docs - Added Engine Version FIeld. Added Skill Level Field.

diff --git a/Engine/Documentation/Source/Gameplay/HowTo/Networking/ReplicateActor/Blueprints/ReplicateActor_BP.INT.udn b/Engine/Documentation/Source/Gameplay/HowTo/Networking/ReplicateActor/Blueprints/ReplicateActor_BP.INT.udn
--- a/Engine/Documentation/Source/Gameplay/HowTo/Networking/ReplicateActor/Blueprints/ReplicateActor_BP.INT.udn
+++ b/Engine/Documentation/Source/Gameplay/HowTo/Networking/ReplicateActor/Blueprints/ReplicateActor_BP.INT.udn
@@ -4,6 +4,8 @@ Crumbs: %ROOT%, Gameplay, Gameplay\HowTo\Networking/ReplicateActor
 Description:Guide to using the Replicates checkbox for Actors in Blueprints. 
 Related: Engine/Blueprints
 Related: Resources/ContentExamples/Networking
+Version: 4.9
+SkillLevel: Advanced
 
 [TOC (start:2 end:2)]
 
