#UE4doc #docs - Added Engine Version field.

diff --git a/Engine/Documentation/Source/Gameplay/HowTo/Networking/ReplicateFunction/ReplicateFunctionOverview.INT.udn b/Engine/Documentation/Source/Gameplay/HowTo/Networking/ReplicateFunction/ReplicateFunctionOverview.INT.udn
--- a/Engine/Documentation/Source/Gameplay/HowTo/Networking/ReplicateFunction/ReplicateFunctionOverview.INT.udn
+++ b/Engine/Documentation/Source/Gameplay/HowTo/Networking/ReplicateFunction/ReplicateFunctionOverview.INT.udn
@@ -2,6 +2,7 @@ Availability: Public
 Title:Replicating Functions
 Crumbs: %ROOT%, Gameplay
 Description:Guide for Replicating Functions in Multiplayer Games. 
+Version: 4.9
 
 **Function Replication** (Remote Procedure Calls or RPCs for short) are functions that are called locally, but executed remotely on another machine (separate from the machine performing the call). RPC functions can be very useful and allow either the Client or the Server to send messages to each other over a network connection. Replicated function calls can be set to either **Reliable** or **Unreliable** where Reliable calls are guaranteed to occur, while Unreliable calls may be dropped during heavy network traffic. Most Replicated Functions that handle cosmetic visuals are unreliable to avoid saturating the network. 
 
