#UE4doc #docs - Added Skill Level field. Added Engine Version field.

diff --git a/Engine/Documentation/Source/Gameplay/HowTo/Networking/ReplicateFunction/Blueprints/ReplicateFunction_BP.INT.udn b/Engine/Documentation/Source/Gameplay/HowTo/Networking/ReplicateFunction/Blueprints/ReplicateFunction_BP.INT.udn
--- a/Engine/Documentation/Source/Gameplay/HowTo/Networking/ReplicateFunction/Blueprints/ReplicateFunction_BP.INT.udn
+++ b/Engine/Documentation/Source/Gameplay/HowTo/Networking/ReplicateFunction/Blueprints/ReplicateFunction_BP.INT.udn
@@ -4,7 +4,8 @@ Crumbs: %ROOT%, Gameplay, Gameplay\HowTo\Networking\ReplicateFunction
 Description:Guide to Replicating Functions in multiplayer games with Blueprints. 
 Related: Engine/Blueprints
 Related: Resources/ContentExamples/Networking
-
+Version: 4.9
+SkillLevel: Advanced
 
 [TOC (start:2 end:2)]
 
