#UE4doc #docs - Added Engine Version Field. Added Skill Level Field.

diff --git a/Engine/Documentation/Source/Gameplay/HowTo/Networking/ReplicateVariable/Blueprints/ReplicateVariable_BP.INT.udn b/Engine/Documentation/Source/Gameplay/HowTo/Networking/ReplicateVariable/Blueprints/ReplicateVariable_BP.INT.udn
--- a/Engine/Documentation/Source/Gameplay/HowTo/Networking/ReplicateVariable/Blueprints/ReplicateVariable_BP.INT.udn
+++ b/Engine/Documentation/Source/Gameplay/HowTo/Networking/ReplicateVariable/Blueprints/ReplicateVariable_BP.INT.udn
@@ -4,6 +4,8 @@ Crumbs: %ROOT%, Gameplay, Gameplay\HowTo\Networking\ReplicateVariable
 Description:Guide to using Replicated and RepNotify Variables in Blueprints.
 Related: Engine/Blueprints
 Related: Resources/ContentExamples/Networking
+Version: 4.9
+SkillLevel: Advanced
 
 
 [TOC (start:2 end:2)]
