#UE4doc #docs - Added Engine Version Field.

diff --git a/Engine/Documentation/Source/Gameplay/HowTo/Networking/TestMultiplayer/TestMultiplayer.INT.udn b/Engine/Documentation/Source/Gameplay/HowTo/Networking/TestMultiplayer/TestMultiplayer.INT.udn
--- a/Engine/Documentation/Source/Gameplay/HowTo/Networking/TestMultiplayer/TestMultiplayer.INT.udn
+++ b/Engine/Documentation/Source/Gameplay/HowTo/Networking/TestMultiplayer/TestMultiplayer.INT.udn
@@ -3,6 +3,7 @@ Title:Testing Multiplayer
 Crumbs: %ROOT%, Gameplay
 Description:Demonstrates how you can set up the Unreal Editor for testing multiplayer games. 
 Related: Resources/ContentExamples/Networking
+Version: 4.9
 
 [TOC (start:2 end:2)]
 
