#UE4doc #docs - Added Engine Version Field. Bolded term.

diff --git a/Engine/Documentation/Source/Gameplay/HowTo/PossessPawns/GHT_4_Overview.INT.udn b/Engine/Documentation/Source/Gameplay/HowTo/PossessPawns/GHT_4_Overview.INT.udn
--- a/Engine/Documentation/Source/Gameplay/HowTo/PossessPawns/GHT_4_Overview.INT.udn
+++ b/Engine/Documentation/Source/Gameplay/HowTo/PossessPawns/GHT_4_Overview.INT.udn
@@ -2,10 +2,11 @@ Availability: Public
 Title:Possessing Pawns
 Crumbs: %ROOT%, Gameplay
 Description:A How To Guide for possessing different pawns in Unreal Engine 4. 
+Version: 4.9
 
 If you want to switch between multiple Playable Characters in your game or you have a situation where you want a Character to enter and control a Vehicle (car, plane, tank, etc.) or other Pawn that has unique controls (like a mounted machine gun a player can control), you will want to take Possession of the Character or Pawn in order to pass input information to it. 
 
-The Player Controller is what is responsible for directing a Pawn and Possession of a Pawn requires a Player Controller to be specified. Pawns themselves do not need to be a humanoid character and can be anything that you want to apply basic movement to and allow a player to control. Characters on the other hand are a form of Pawn that includes Collision and a CharacterMovementComponent by default which allows it to do basic human-like movement. 
+The **PlayerController** is what is responsible for directing a Pawn and Possession of a Pawn requires a PlayerController to be specified. Pawns themselves do not need to be a humanoid character and can be anything that you want to apply basic movement to and allow a player to control. Characters on the other hand are a form of Pawn that includes Collision and a CharacterMovementComponent by default which allows it to do basic human-like movement. 
 
 The link below in the Setup section will give you examples on Pawn Possession using Blueprints.
 
