#UE4doc #doc - Added Skill Level field

diff --git a/Engine/Documentation/Source/Gameplay/HowTo/PossessPawns/Blueprints/GHT_4B.INT.udn b/Engine/Documentation/Source/Gameplay/HowTo/PossessPawns/Blueprints/GHT_4B.INT.udn
--- a/Engine/Documentation/Source/Gameplay/HowTo/PossessPawns/Blueprints/GHT_4B.INT.udn
+++ b/Engine/Documentation/Source/Gameplay/HowTo/PossessPawns/Blueprints/GHT_4B.INT.udn
@@ -3,7 +3,7 @@ Title:Possessing Pawns in Blueprint
 Crumbs: %ROOT%, Gameplay, Gameplay/HowTo/PossessPawns
 Description:Illustrates using the Possess node in Blueprints to switch between and control multiple characters. 
 Version: 4.8
-SkillLevel: Beginner
+SkillLevel: Intermediate
 Related: Engine/Blueprints
 Related: Gameplay/Framework/Pawn
 Related: Gameplay/Framework/Pawn/Character
