#UE4doc #docs - Added Engine Version Field. Bolded Terms.

diff --git a/Engine/Documentation/Source/Gameplay/HowTo/ReferenceAssets/GHT_5_Overview.INT.udn b/Engine/Documentation/Source/Gameplay/HowTo/ReferenceAssets/GHT_5_Overview.INT.udn
--- a/Engine/Documentation/Source/Gameplay/HowTo/ReferenceAssets/GHT_5_Overview.INT.udn
+++ b/Engine/Documentation/Source/Gameplay/HowTo/ReferenceAssets/GHT_5_Overview.INT.udn
@@ -2,9 +2,10 @@ Availability: Public
 Title:Referencing Actors
 Crumbs: %ROOT%, Gameplay
 Description:A How To Guide for Referencing Actors in Unreal Engine 4. 
+Version: 4.9
 
 [EXCERPT:ReferenceOverview]
-By establishing reference to an Actor, you can access its Properties, Variables, Events, or Functions (if applicable) during gameplay through Blueprints or C++ and alter them to suit your gameplay needs. 
+By establishing reference to an **Actor**, you can access its **Properties**, **Variables**, **Events**, or **Functions** (if applicable) during gameplay through Blueprints or C++ and alter them to suit your gameplay needs. 
 
 Say for example you want a light to turn on when a player enters a Trigger Box, having reference to the light and the Trigger Box will allow you to do so by tying the light Actor to the OnComponentBeginOverlap Event of the Trigger Box (in which case you could turn on the light). 
 
