#UE4doc #docs - Added Engine Version Field.

diff --git a/Engine/Documentation/Source/Programming/UnrealArchitecture/Reference/Properties/Specifiers/Interp/PropertySpecifier_Interp.INT.udn b/Engine/Documentation/Source/Programming/UnrealArchitecture/Reference/Properties/Specifiers/Interp/PropertySpecifier_Interp.INT.udn
--- a/Engine/Documentation/Source/Programming/UnrealArchitecture/Reference/Properties/Specifiers/Interp/PropertySpecifier_Interp.INT.udn
+++ b/Engine/Documentation/Source/Programming/UnrealArchitecture/Reference/Properties/Specifiers/Interp/PropertySpecifier_Interp.INT.udn
@@ -2,5 +2,6 @@ Title: Interp
 Crumbs: %ROOT%, Programming, Programming/UnrealArchitecture/Reference, Programming/UnrealArchitecture/Reference/Properties, Programming/UnrealArchitecture/Reference/Properties/Specifiers
 Description: 
 Availability:Public
+Version: 4.9
 
 Indicates that the value can be driven over time by a track in Matinee. 
