#UE4doc #docs:
Added engine version field
Changed font settings for multiple terms
Changed multiple Tip regions to Note regions due to the nature of the content
Replaced a couple images due to outdated UI

diff --git a/Engine/Documentation/Source/Engine/Landscape/Materials/Materials.INT.udn b/Engine/Documentation/Source/Engine/Landscape/Materials/Materials.INT.udn
--- a/Engine/Documentation/Source/Engine/Landscape/Materials/Materials.INT.udn
+++ b/Engine/Documentation/Source/Engine/Landscape/Materials/Materials.INT.udn
@@ -2,6 +2,7 @@ Availability: Public
 Title:Landscape Materials
 Crumbs: %ROOT%, Engine, Engine/Landscape
 Description:Setting up materials to work with Landscape terrains.
+version: 4.9
 
 [VAR:Topic]
 [OBJECT:Topiccompact]
@@ -84,18 +85,18 @@ The following table details the **Layer Blend node's** inputs and outputs.
 
 
 
-When using certain combinations of layer blend modes you could end up with black spots all over your Landscape where different layers meet. This problem is especially prevalent when using the **LB_HeightBlend** mode for multiple Landscape layers. **LB_HeightBlend** works by modulating the blend factor, or weight, for the layer using the specified height value. When you have multiple layers painted on an area and they all are set to **LB_HeightBlend**, it is possible that all the layers painted in a particular area will simultaneously have a 0 height value, so the desired blend factor for each layer becomes 0. Because there is no implicit or explicit ordering, the result will be black spots because no layers will have any contribution to that area. The situation is worse when you are blending a Normal map, because it results in a Normal value of (0,0,0) which is not valid and will cause rendering problems with the lighting. The solution to this problem is to use **LB_AlphaBlend** for one of the layers like in the following example. 
+When using certain combinations of layer blend modes you could end up with black spots all over your Landscape where different layers meet. This problem is especially prevalent when using the LB_HeightBlend mode for multiple Landscape layers. LB_HeightBlend works by modulating the blend factor, or weight, for the layer using the specified height value. When you have multiple layers painted on an area and they all are set to LB_HeightBlend, it is possible that all the layers painted in a particular area will simultaneously have a 0 height value, so the desired blend factor for each layer becomes 0. Because there is no implicit or explicit ordering, the result will be black spots because no layers will have any contribution to that area. The situation is worse when you are blending a Normal map, because it results in a Normal value of (0,0,0) which is not valid and will cause rendering problems with the lighting. The solution to this problem is to use **LB_AlphaBlend** for one of the layers like in the following example. 
 ![LB_HeightBlend Problem](HeightBlendProblem.jpg)(w:800 h:238)
 
-In the left image, all layers are **LB_HeightBlend**, causing some areas to be black. On the right, the red "1" layer has been changed to use **LB_AlphaBlend**, which solves the problem.
+In the left image, all layers are LB_HeightBlend, causing some areas to be black. On the right, the red "1" layer has been changed to use LB_AlphaBlend, which solves the problem.
 
 
-Below is an example of the properties of the **Landscape Layer Blend** node for a number of layers all being blended together. Make sure to note how the **Soil** layer has its blend mode set to **LB_AlphaBlend** while the other layers have theirs set to **LB_HeightBlend**. This is to stop the issue mentioned above (having black spots where layers meet) from happening. 
+Below is an example of the properties of the **Landscape Layer Blend** node for a number of layers all being blended together. Make sure to note how the **Soil** layer has its blend mode set to LB_AlphaBlend while the other layers have theirs set to LB_HeightBlend. This is to stop the issue mentioned above (having black spots where layers meet) from happening. 
 
 ![Layer Blend Properties](LayerBlendProperties.png)
 
 
-To delete a layer, click the **Drop-Down Arrow** to the right of the layer's element number to open the context menu, and then click **Delete**.
+To delete a layer, click the Drop-Down Arrow to the right of the layer's element number to open the context menu, and then click **Delete**.
 
 ![Delete Layer](Landscape_DeleteBlendLayer2.png)
 
@@ -229,6 +230,7 @@ You can preview the effects of different weights on the Material by changing the
 ### Using Landscape Layer Blend Nodes
 
 As an alternative to manually blending layers together with the **LandscapeLayerWeight** node, the **LandscapeLayerBlend** node blends multiple layers together using either alpha blending or alpha blending with a height-based offset. The height-based offset allows a layer to blend with other layers based on a heightmap input. 
+
 **To use a LandscapeLayerBlend node:**
 
 1. In the Material Editor, add a LandscapeLayerBlend node.
@@ -288,7 +290,7 @@ If a particular layer is not used on a particular Landscape component, the sub-t
 This reduces overall Material complexity, allowing the Material applied to the Landscape to contain any number of texture samples as long as the number of samples being used on any individual component never exceeds the maximum set forth by the shader model specifications (16 for SM 3.0). 
 This means it is possible to set up a master Material that contains every texture or network used anywhere on the Landscape, which can end up being a very complex collection of networks, while still keeping the final Materials applied to the components of the Landscape within the parameters allowed by the hardware.
 
-[REGION:tip]
+[REGION:note]
 Any network of material expressions can be connected to the **Layer** inputs in place of a simple **TextureSample**. This makes it possible to do more complex effects such as transitioning from detail textures to larger macro textures depending on the distance the layer is being viewed from. 
 [/region]
 
