#UE4doc #docs:
Added engine version field
Fixed capitalization on a few terms

diff --git a/Engine/Documentation/Source/Engine/Landscape/QuickStart/LandscapeQuickStart.INT.udn b/Engine/Documentation/Source/Engine/Landscape/QuickStart/LandscapeQuickStart.INT.udn
--- a/Engine/Documentation/Source/Engine/Landscape/QuickStart/LandscapeQuickStart.INT.udn
+++ b/Engine/Documentation/Source/Engine/Landscape/QuickStart/LandscapeQuickStart.INT.udn
@@ -2,6 +2,7 @@ Availability:Public
 Title: Landscape Quick Start Guide
 Crumbs: %ROOT%, GettingStarted
 Description: Getting up and running with the basics of the Landscape System in Unreal Engine.
+version: 4.9
 
 [VAR:Topic]
 [OBJECT:Topic]
@@ -37,7 +38,7 @@ Description: Getting up and running with the basics of the Landscape System in U
 [/OBJECT]
 [/VAR]
 
-The Unreal Editor Landscape Quick Start Guide walks you through creating a new Landscape, Sculpting the Landscape, Creating new Materials for the Landscape, and Painting those materials on the Landscape. 
+The Unreal Editor Landscape Quick Start Guide walks you through creating a new Landscape, sculpting the Landscape, creating new Materials for the Landscape, and painting those Materials on the Landscape. 
 
 [OBJECT:TopicButtonList]
 	[PARAM:icon]
