#UE4doc #docs:
Added engine version field
Bolded a few terms

diff --git a/Engine/Documentation/Source/Engine/Landscape/QuickStart/1/LSQS_1.INT.udn b/Engine/Documentation/Source/Engine/Landscape/QuickStart/1/LSQS_1.INT.udn
--- a/Engine/Documentation/Source/Engine/Landscape/QuickStart/1/LSQS_1.INT.udn
+++ b/Engine/Documentation/Source/Engine/Landscape/QuickStart/1/LSQS_1.INT.udn
@@ -2,7 +2,7 @@ Availability: Public
 Title:1. Working with the Landscape tools
 Crumbs:%ROOT%, GettingStarted, Engine/Landscape/QuickStart
 Description:Working with the Landscape tools.
-
+version: 4.9
 
 [VAR:Steps]
 [OBJECT:Navigation]
@@ -21,8 +21,7 @@ Description:Working with the Landscape tools.
 
 %Steps%
 
-## Overview
-Creating a Landscape inside of Unreal Engine 4 (UE4) can be done in a few simple steps. But before we dive into creating our first Landscape, let us first familiarize ourselves with some of the tools and keyboard inputs that are most commonly used to interact with the Landscape system. 
+Creating a **Landscape** inside of Unreal Engine 4 (UE4) can be done in a few simple steps. But before we dive into creating our first Landscape, let us first familiarize ourselves with some of the tools and keyboard inputs that are most commonly used to interact with the Landscape system. 
 
 
 ## Opening the Landscape Tool
@@ -32,7 +31,7 @@ All of the tools that are used to interact with the Landscape system can be foun
 
 ## Manage, Sculpt, and Paint Mode
 
-The Landscape tool has three modes, Manage, Sculpt, and Paint that are accessible by clicking on their icons at the top of the Landscape's toolbar window. Each mode will allow you to interact with the Landscape in a different manner. Here is a very quick rundown of what each mode allow you to do.
+The Landscape tool has three modes, **Manage**, **Sculpt**, and **Paint** that are accessible by clicking on their icons at the top of the Landscape's toolbar window. Each mode will allow you to interact with the Landscape in a different manner. Here is a very quick rundown of what each mode allow you to do.
 
 [INCLUDE:Engine/Landscape#landscapetoolmodes1]
 
