#UE4doc #docs:
Added engine version field
Changed font settings for a few terms
Fixed capitalization for a few terms
Replaced an image due to outdated UI

diff --git a/Engine/Documentation/Source/Engine/Landscape/QuickStart/2/LSQS_2.int.udn b/Engine/Documentation/Source/Engine/Landscape/QuickStart/2/LSQS_2.int.udn
--- a/Engine/Documentation/Source/Engine/Landscape/QuickStart/2/LSQS_2.int.udn
+++ b/Engine/Documentation/Source/Engine/Landscape/QuickStart/2/LSQS_2.int.udn
@@ -1,7 +1,8 @@
 Availability: Public
 Title:2. Creating a new Landscape
 Crumbs:%ROOT%, GettingStarted, Engine/QuickStart
-Description:How to create a new Landscape for your level. 
+Description:How to create a new Landscape for your level.
+version: 4.9
 
 
 [VAR:Steps]
@@ -25,19 +26,19 @@ Description:How to create a new Landscape for your level.
 ## Creating a New FPS Blueprint Project
 Creating a new Landscape is actually a simple process that we will be going over in the next few steps. 
 However before we begin to create our Landscape, lets first create a new project **Blueprint First Person** Project. 
-If you are unfamiliar with how to create a new project, please check out the following page on how to [**Create a New Project**](Engine/QuickStart/1) 
+If you are unfamiliar with how to create a new project, please check out the following page on how to [Create a New Project](Engine/QuickStart/1) 
 
 ## Creating a Landscape
 1. First, create a new **Blueprint First Person** project if you have not done so already. 
-While you can use other templates for this tutorial, the **Blue Print First Person** will make it a little easier to inspect your Landscape.
+While you can use other templates for this tutorial, the Blueprint First Person will make it a little easier to inspect your Landscape.
 	[REGION:fullwidth]
 	![](T_Creating_a_new_FPS_BP_Project.png)
 	[/REGION]
 
-1. Once you have created your new project and the editor has been loaded, create a new level using **File->New Level** and select the Default Level from the New Level Template.  
+1. Once you have created your new project and the editor has been loaded, create a new level using **File->New Level** and select the **Default** Level from the New Level Template.  
 	![](T_Creating_A_New_Level.png)
 
-1. With your new level now created, select the **SM_Templae_Floor_Map_Floor** from the level and press the **Delete Key** to remove it from the level.
+1. With your new level now created, select the **Floor** from the level and press the **Delete** key to remove it from the level.
 	[REGION:note]
 	Make sure that you select your player start and move it up slightly in the Z axis. This will make sure that your player does not start under your newly created Landscape.
 	[/REGION]
@@ -55,8 +56,8 @@ While you can use other templates for this tutorial, the **Blue Print First Pers
 	![](T_Landscape_Tools.png)
 
 1. For this tutorial, we will just focus on creating a Landscape using the default settings. 
-If you would like to know more about the various settings in the Mange mode of the Landscape tool, please refer to the following document on  [**Creating Landscapes**](Engine/Landscape/Creation). 
-For now, just press the **_1: Fill World Button_** to make your Landscape Actor fill a majority of the world and then press the **_ 2: Create Button_** to actually create the Landscape. 
+If you would like to know more about the various settings in the Mange mode of the Landscape tool, please refer to the following document on  [Creating Landscapes](Engine/Landscape/Creation). 
+For now, just press the **1: Fill World Button** to make your Landscape Actor fill a majority of the world and then press the **2: Create Button** to actually create the Landscape. 
 	
 	![](T_Fillworld_Create_Landscape.png)
 
