#UE4doc #docs:
Added engine version field
Bolded a few terms
Reworded a term as it wasn't the actual name of a tool

diff --git a/Engine/Documentation/Source/Engine/Landscape/QuickStart/3/LSQS_3.int.udn b/Engine/Documentation/Source/Engine/Landscape/QuickStart/3/LSQS_3.int.udn
--- a/Engine/Documentation/Source/Engine/Landscape/QuickStart/3/LSQS_3.int.udn
+++ b/Engine/Documentation/Source/Engine/Landscape/QuickStart/3/LSQS_3.int.udn
@@ -2,6 +2,7 @@ Availability: Public
 Title:3. Sculpting the Landscape
 Crumbs:%ROOT%, GettingStarted, Engine/Landscape/QuickStart
 Description:How to sculpt your Landscape.
+version: 4.9
 
 
 [VAR:Steps]
@@ -24,7 +25,7 @@ Description:How to sculpt your Landscape.
 
 
 ## Sculpting the Landscape
-Sculpting the Landscape is a rather easy but time consuming process. All of the tools used for sculpting can be found under the Sculpting Tools Icon in the Landscape tools. If you would like to know more about what each of the Sculpting Tools does in detail, you can read the [Sculpt Mode](Engine/Landscape/Editing/SculptMode) page. For quick reference, here is a list of the most common key and mouse interactions that are used when sculpting the Landscape.   
+Sculpting the **Landscape** is a rather easy but time consuming process. All of the tools used for sculpting can be found under the **Sculpt** icon in the Landscape tools. If you would like to know more about what each of the Sculpting Tools does in detail, you can read the [Sculpt Mode](Engine/Landscape/Editing/SculptMode) page. For quick reference, here is a list of the most common key and mouse interactions that are used when sculpting the Landscape.   
 
 | **Common Controls** | **Operation** |
 | --- | --- |
@@ -43,7 +44,7 @@ There could be a lot of various reasons as to why what you do in this tutorial d
 
 1. To begin, first find a section of the Landscape that you would like to work with. 
 For this tutorial, we are not going to be filling in the entire Landscape but just a section of it. 
-For ease of use, you should set a camera bookmark by pressing **CRTL + 1** on the keyboard. 
+For ease of use, you should set a camera bookmark by pressing **Ctrl + 1** on the keyboard. 
 This will set a camera bookmark which will make it easier for you to gauge how your Landscape is coming along by giving you a camera view to always come back to. 
 At any time during your editor session, if you press the 1 Key, your camera will be returned to the exact same position that you set.
 	[REGION:fullwidth] 
