#UE4doc #docs:
Added engine version field
Changed font settings for multiple terms
Added multiple steps as it was not clear how to complete a certain section
Added a note section
Replaced multiple images due to out of date UI

diff --git a/Engine/Documentation/Source/Engine/Landscape/QuickStart/4/LSQS_4.INT.udn b/Engine/Documentation/Source/Engine/Landscape/QuickStart/4/LSQS_4.INT.udn
--- a/Engine/Documentation/Source/Engine/Landscape/QuickStart/4/LSQS_4.INT.udn
+++ b/Engine/Documentation/Source/Engine/Landscape/QuickStart/4/LSQS_4.INT.udn
@@ -28,20 +28,20 @@ But before we do this, let us first setup some folders so that the content that
 For this section of the tutorial, we are going to be making the following four folders.
 
 [REGION:note]
-If you would like to know more about how to setup folders in Unreal Engine 4, please refer to the following page about [**Folders**](Engine/Content/Browser/UserGuide/WorkingWithSources).
+If you would like to know more about how to setup folders in Unreal Engine 4, please refer to the following page about [Folders](Engine/Content/Browser/UserGuide/WorkingWithSources).
 [/REGION]
 
-1. **Landscape** 
-2. **Materials**
-3. **Resources**
-4. **Textures**
+1. Landscape 
+2. Materials
+3. Resources
+4. Textures
 
 When completed, you should have something that looks like the following.
 
 ![](T_Landscape_Folders.png)
 
 ## Migrating Textures
-Now that we have the folders in place, it is time to migrate some textures from the **Landscape Content Example** project so that we have some Textures to work with. If you would like to know more about how to Migrate Content from one project to another one, please refer to the following page on how to [**Migrate Content**](Engine/Content/Browser/UserGuide/Migrate)  
+Now that we have the folders in place, it is time to Migrate some Textures from the **Landscape Content Example** project so that we have some Textures to work with. If you would like to know more about how to Migrate content from one project to another one, please refer to the following page on how to [Migrate Content](Engine/Content/Browser/UserGuide/Migrate)  
 [REGION:note]
 When Migrating content between projects, you could possibly end up with additional folders that you do not want. 
 To fix this, select the Textures that you want inside of the **Content Browser** and then drag them from their current location into the folder that you want them to be placed in. 
@@ -67,18 +67,18 @@ Once that has been completed, it is now time to create our Material for the Land
 Creating a Material for our Landscape can be done in the following steps.
 
 1. Navigate to the **Materials** folder.
-1. **Right-Click** in the **Asset Management Area** and select **Material** from the **Basic Assets** list. 
-1. Name the newly created **Material** something that will allow you to easily find it like **Landscape_Material** for example. 
+1. **Right-click** in the **Content Browser** and select **Material** from the **Basic Assets** list. 
+1. Name the newly created Material something that will allow you to easily find it like **Landscape_Material** for example. 
 
 [REGION:note]
-If you have not done so already, please check out the [**Materials**](Engine/Rendering/Materials/) pages to gain a more in-depth understanding of how materials work inside of Unreal Engine 4.
+If you have not done so already, please check out the [Materials](Engine/Rendering/Materials/) pages to gain a more in-depth understanding of how materials work inside of Unreal Engine 4.
 [/REGION]
 
 Once this has been completed, you should have something that looks like this.
 
 ![](T_Landscape_Create_New_Material.png)
 
-With our new Landscape Material created, open up the Material by **double-clicking** on it inside of the **Generic Browser**. When you do, you should see something like this come up on the screen.
+With our new Landscape Material created, open up the Material by **Double-clicking** on it inside of the **Content Browser**. When you do, you should see something like this come up on the screen.
 [REGION:fullwidth] 
 ![](T_Landscape_Blank_Material.png)
 [/REGION]
@@ -91,27 +91,34 @@ This will enable this Material to be used with Landscapes and save you additiona
 [/REGION]
 
 Now with the **Used with Landscape** Usage flag set, it is now time to start laying down nodes inside of the Material Editor. 
-The first node that you are going to want to lay down is a **LandscapeCoords** UV node. This node will help to generate UV cords that can be used to map the Material network to the Landscape Terrain.
+The first node that you are going to want to lay down is a **LandscapeLayerCoords** UV node. This node will help to generate UV cords that can be used to map the Material network to the Landscape Terrain.
 
 ![](T_Landscape_UVCords.png)
 
 [REGION:note]
-The easiest way to find nodes specific to Landscape is to search for them in the **Materials Palette** box using Landscape as the key word.
+The easiest way to find nodes specific to Landscape is to search for them in the Materials **Palette** box using Landscape as the key word.
 
 ![](T_Landscape_Material_Search.png)
 [/REGION]
 
-The next Material nodes that we are going to lay down are going to be for the textures for the ground's **Base Color** and **Normal** maps. For the snow, we are just going to use a **Vector Parameter** that uses an off White color. To make sure that no **Metallic** information is used, a **Scalar Value** of 0 is used and plugged into the **Metallic** input. Finally, for the **Roughness**, we set another **Scalar Value** but unlike with the **Metallic** one, this will be a **Scalar Parameter** so that this value can be tweaked via a **Material Instance** later on. Finally make sure that you hook up the **LandscapeCoords** to the UV's of each of the **Texture Samples**. When completed, you should have something that looks like this.  
+The next Material nodes that we are going to lay down are going to be for the textures for the ground's **Base Color** and **Normal** maps. For the snow, we are just going to use a **Vector Parameter** (**V + Left-click**) that uses an off White color. To make sure that no **Metallic** information is used, a **Constant** (**1 + Left-click**) of 0 is used and plugged into the **Metallic** input. Finally, for the **Roughness**, we set a **Scalar Parameter** (**S + Left-click**) so that this value can be tweaked via a **Material Instance** later on. Finally make sure that you hook up the **LandscapeCoords** to the UV's of each of the **Texture Samples**. When completed, you should have something that looks like this.
 [REGION:fullwidth] 
 ![](T_Landscape_Material_00.png)
 [/REGION]
+
+To add the **Texture Sample** nodes for the various textures, first select the desired texture in the **Content Browser** and then press **T + Left-click** in the **Material Editor**'s graph to create the node.
+
+[REGION:note]
+To find more information about these keybindings, please navigate to the **Edit > Editor Preferences > Keyboard Shortcuts** window and select the **Material Editor - Spawn Nodes** section.
+[/REGION]
+
 | **Number** | **Texture Name** |
 | --- | --- | 
 | **1**|T_ForestGround_D| 
 | **2**|T_Ground_Grass_D| 
 | **3**|T_ForestGround_N|
 | **4**|T_Ground_Grass_N|
-| **5**|T_Noise_N| 
+| **5**|T_Noise_N|
 
 After the Material nodes have been laid down and the **LandscapeCoords** connected to the textures UVs, it is now time to add and setup the **Landscape Layer Blend** node. This node is used to blend all of the textures that are used for the Landscape so that they can be painted on the Landscape individually using the Landscape Painting tools. The following table shows what **Textures** are associated with which **Layer Name** and what **Blend Type** they use.
 
@@ -139,7 +146,7 @@ When you first place down a **Landscape Layer Blend** node, it will be blank lik
 ![](T_Landscape_Layer_Blend_N.png)
 
 [REGION:note]
-For more in-depth information about using the **Landscape Layer Blend** node, please read the [**Terrain Expressions**](Engine/Rendering/Materials/ExpressionReference/Landscape) page.
+For more in-depth information about using the **Landscape Layer Blend** node, please read the [Terrain Expressions](Engine/Rendering/Materials/ExpressionReference/Landscape) page.
 [/REGION]
 
 Once the **Layer Blend** nodes have been set up, it is time to connect the Texture maps to them. When completed, you should have something that looks like the following. 
