#UE4doc #docs:
Resubmitting as not all changes were saved

diff --git a/Engine/Documentation/Source/Engine/Landscape/QuickStart/4/LSQS_4.INT.udn b/Engine/Documentation/Source/Engine/Landscape/QuickStart/4/LSQS_4.INT.udn
--- a/Engine/Documentation/Source/Engine/Landscape/QuickStart/4/LSQS_4.INT.udn
+++ b/Engine/Documentation/Source/Engine/Landscape/QuickStart/4/LSQS_4.INT.udn
@@ -2,6 +2,7 @@ Availability: Public
 Title:4. Creating Landscape Materials
 Crumbs:%ROOT%, GettingStarted, Engine/Landscape/QuickStart
 Description:How to create Materials to be used for your Landscape.
+version: 4.9
 
 [VAR:Steps]
 [OBJECT:Navigation]
