#UE4doc #docs:
Added engine version field
Changed font settings for a few terms

diff --git a/Engine/Documentation/Source/Engine/Landscape/QuickStart/5/LSQS_5.INT.udn b/Engine/Documentation/Source/Engine/Landscape/QuickStart/5/LSQS_5.INT.udn
--- a/Engine/Documentation/Source/Engine/Landscape/QuickStart/5/LSQS_5.INT.udn
+++ b/Engine/Documentation/Source/Engine/Landscape/QuickStart/5/LSQS_5.INT.udn
@@ -2,6 +2,7 @@ Availability: Public
 Title:5. Painting Landscape Materials
 Crumbs:%ROOT%, GettingStarted, Engine/Landscape/QuickStart
 Description:With some Actors in the Level, it's now time to add some variation.
+version: 4.9
 
 [VAR:Steps]
 [OBJECT:Navigation]
@@ -83,11 +84,11 @@ Before you begin to paint the Landscape, here is a re-cap of some of the most us
 | **Ctrl+Z** | Undo last stroke. |
 | **Ctrl+Y** | Redo last undone stroke. |
 
-The main tool that you will use for actually applying textures to your Landscape is going to be the **Paint Tool**. To find out more about all of the tools that you can use to paint on the Landscape or if you want to know more about how they work, please check out the [**Paint Mode**](Engine/Landscape/Editing/PaintMode) documentation. 
+The main tool that you will use for actually applying textures to your Landscape is going to be the **Paint Tool**. To find out more about all of the tools that you can use to paint on the Landscape or if you want to know more about how they work, please check out the [Paint Mode](Engine/Landscape/Editing/PaintMode) documentation. 
 
 To apply a Material to the Landscape, hold down the **Left Mouse Button** to apply whatever Material you have selected to the area under the mouse. 
 
-To select a new texture to paint, first make sure you are in **Landscape Painting Mode** then on the **Modes** panel, under the **Target Layers** section, select which texture you want to paint with by clicking on it in the list. Whichever texture is highlighted in **Orange** will be the texture that is painted to the Landscape. In the image below, the **Soil** is what is highlighted in **Orange** so this means that this is the texture that will be painted to the Landscape. To select another texture to be painted, all you have to do is select it in the list and begin painting.
+To select a new texture to paint, first make sure you are in **Landscape Painting Mode** then on the **Modes** panel, under the **Target Layers** section, select which texture you want to paint with by clicking on it in the list. Whichever texture is highlighted in orange will be the texture that is painted to the Landscape. In the image below, the **Soil** is what is highlighted in orange so this means that this is the texture that will be painted to the Landscape. To select another texture to be painted, all you have to do is select it in the list and begin painting.
 
 [REGION:fullwidth] 
 ![](T_Landscape_Picking_Layers_To_Paint.png)
