#UE4doc #docs:
Added engine version field
Unbolded multiple links
Replaced a few mentions of terms with links to their respective pages

diff --git a/Engine/Documentation/Source/Engine/Landscape/QuickStart/6/LSQS_6.INT.udn b/Engine/Documentation/Source/Engine/Landscape/QuickStart/6/LSQS_6.INT.udn
--- a/Engine/Documentation/Source/Engine/Landscape/QuickStart/6/LSQS_6.INT.udn
+++ b/Engine/Documentation/Source/Engine/Landscape/QuickStart/6/LSQS_6.INT.udn
@@ -2,6 +2,7 @@ Availability: Public
 Title:6. Landscape Tips and Tricks
 Crumbs:%ROOT%, GettingStarted, Engine/Landscape/QuickStart
 Description:Landscape Tips and Tricks.
+version: 4.9
 
 
 [VAR:Steps]
@@ -40,38 +41,38 @@ While the previous quick start tutorial for Landscape will get you up and runnin
 * You can change the LOD Distance Factor for individual Landscape components so they will simplify at closer or further distance thresholds. Things like mountain peaks or anything with a distinct silhouette will LOD the most noticeably as you move further away, so you can reduce the LOD bias for those components to preserve their shape. You can also raise the LOD bias for low-detail areas like flat plains that will not look noticeably different with less tessellation.
 
 ##World Composition 
-Unreal Engine 4(UE4) now offers the ability to have massive worlds made out of Landscape that can easily be managed by using the **World Composition** tool. **World Composition** was designed to help simplify the management of large worlds, especially if those worlds are made using the Landscape system. To find out more about the **World Composition** tool, please refer to the official document that you can find here. 
+Unreal Engine 4 (UE4) now offers the ability to have massive worlds made out of Landscape that can easily be managed by using the **World Composition** tool. World Composition was designed to help simplify the management of large worlds, especially if those worlds are made using the Landscape system. To find out more about the World Composition tool, please refer to the official document that you can find here. 
 
-[**World Composition Documents**](Engine/LevelStreaming/WorldBrowser)
+[World Composition Documents](Engine/LevelStreaming/WorldBrowser)
 
 ##External Creation Tools 
 While the default Landscape tools do have the ability to meet all of your sculpting and painting needs, there could be some situations where you might want some extra control over your Landscape's look and feel. The following is a list of software packages that could help you obtain the results you are looking for if you cannot get them from using the Landscape tools. 
 
-**World Machine**
+[World Machine](http://www.world-machine.com/)
 
 World Machine is a powerful procedural terrain creation software. It can be used to build, texture, and export both heightmaps and textures for your Landscape. You can read more about how World Machine might be able to help you out with your Landscape by checking out their website. 
 
-[**World Machine**](http://www.world-machine.com/)
 
-**Terragen**
+
+[Terragen](http://planetside.co.uk/)
 
 Terragen is another powerful fully procedural terrain creation software. Much like World Machine, it can be used to build, texture, and export both heightmaps and textures for your Landscape. You can read more about how Terragen might be able to help you out with your Landscape by checking out their website.  
 
-[**Terragen**](http://planetside.co.uk/)
 
-**Mudbox**
+
+[Mudbox](http://www.autodesk.com/products/mudbox/overview)
 
 While primarily a tool for digital sculpting and painting 3D meshes, Mudbox can also be used to generate heightmap data for your Landscape. You can read more about how Mudbox might be able to help you out with your Landscape by checking out their website.  
 
-[**Mudbox**](http://www.autodesk.com/products/mudbox/overview)
 
 
-**ZBrush**
+
+[ZBrush](http://pixologic.com/)
 
 Zbrush is another digital sculpting and painting tool that can be used to generate heightmap data for your Landscape. You can read more about how Zbrush might be able to help you out with your Landscape by checking out their website.  
 
 
-[**ZBrush**](http://pixologic.com/)
+
 
 
 
