#UE4doc #docs:
Added engine version field
Changed font settings for multiple terms
Fixed capitalization for multiple terms

diff --git a/Engine/Documentation/Source/Engine/Landscape/TechnicalGuide/LansdscapeTechnicalGuide.INT.udn b/Engine/Documentation/Source/Engine/Landscape/TechnicalGuide/LansdscapeTechnicalGuide.INT.udn
--- a/Engine/Documentation/Source/Engine/Landscape/TechnicalGuide/LansdscapeTechnicalGuide.INT.udn
+++ b/Engine/Documentation/Source/Engine/Landscape/TechnicalGuide/LansdscapeTechnicalGuide.INT.udn
@@ -2,6 +2,7 @@ Availability: Public
 Title:Landscape Technical Guide
 Crumbs: %ROOT%, Engine, Engine/Landscape
 Description:Technical Settings for Landscape.
+version: 4.9
 
 [VAR:TopicCompact]
 [OBJECT:TopicCompact]
@@ -26,7 +27,6 @@ Description:Technical Settings for Landscape.
 [TOC(start:2 end:3)]
 
 
-## Overview
 
 To get the best performance possible from the Landscape system, there are certain technical restrictions that have to be adhered to. The following document aims to bring these restrictions to your attention as well as provide you with valuable information so that you can strike the best balance between beauty and performance with your Landscape. 
 
@@ -36,7 +36,7 @@ To get the best performance possible from the Landscape system, there are certai
 
 Valid dimensions for Landscape heightmaps are not always immediately obvious when getting started. Being able to determine what dimensions for heightmaps are valid and, beyond that, which are optimal requires a thorough understanding of the underlying architecture of the Landscape. In order to create a system that allows for huge terrains while still being efficient in terms of memory and performance, the architecture implicitly applies restrictions on the dimensions of the heightmap, meaning there are certain dimensions that are valid and others that are not. In previous terrain systems in Unreal Engine, there were either no restrictions (i.e., any dimensions were valid and would work) or the restrictions were fairly simple (i.e., only square power of two heightmaps were allowed). The restrictions on heightmaps for Landscapes are much more complex and rigid.
 
-A Landscape Actor is color coded so that it is easier to tell what each section does. The edges of the Landscape are highlighted in **Yellow**, the edge of each component is in **Light Green**, section edges (if set to 2x2 sections) in a **Medium Green**, and the individual Landscape quads in a **Dark Green**.
+A Landscape Actor is color coded so that it is easier to tell what each section does. The edges of the Landscape are highlighted in yellow, the edge of each component is in light green, section edges (if set to 2x2 sections) in a medium green, and the individual Landscape quads in a dark green.
 
 | Color: |Landscape Component|
 |---|---|
@@ -119,9 +119,9 @@ In order to make things easier, here are a number of sizes that maximize the are
 
 ## Layer Debug Mode
 
-**Layer Debug** mode allows the weight of specific layers to be visualized on the Landscape in the viewport. You can enable **Layer Debug** mode in the viewport's View menu, under **Landscape Visualizers**. For more information about the viewport's Landscape-related view modes, see the "[Landscape Visualizers](Engine/UI/LevelEditor/Viewports/ViewModes\#LandscapeVisualizers)".
+**Layer Debug** mode allows the weight of specific layers to be visualized on the Landscape in the viewport. You can enable Layer Debug mode in the viewport's View menu, under **Landscape Visualizers**. For more information about the viewport's Landscape-related view modes, see the [Landscape Visualizers](Engine/UI/LevelEditor/Viewports/ViewModes\#LandscapeVisualizers).
 
-When **Layer Debug** mode is enabled, radio buttons allowing selection of individual color channels are shown with the target layers in the list.
+When Layer Debug mode is enabled, radio buttons allowing selection of individual color channels are shown with the target layers in the list.
 
 ![Landscape_Target_Debug.png](Landscape_Target_Debug.png)(convert:false)
 
