#UE4doc #docs:
Added engine version field
Changed font settings for multiple terms
Fixed capitalization for a term

diff --git a/Engine/Documentation/Source/Engine/Landscape/UI/InterfaceReference.INT.udn b/Engine/Documentation/Source/Engine/Landscape/UI/InterfaceReference.INT.udn
--- a/Engine/Documentation/Source/Engine/Landscape/UI/InterfaceReference.INT.udn
+++ b/Engine/Documentation/Source/Engine/Landscape/UI/InterfaceReference.INT.udn
@@ -2,6 +2,7 @@ Availability: Public
 Title:Landscape UI Reference
 Crumbs:%ROOT%, Engine, Engine/Landscape
 Description:Description of the Landscape tool's user interface.
+version: 4.9
 
 [VAR:Topic]
 [OBJECT:Topiccompact]
@@ -72,7 +73,7 @@ This document contains references to all of the icons for the **Landscape** syst
 [/VAR]
 
 [VAR:SculptTool]
-![](Engine/Landscape\icon_Landscape_Tool_Sculpt_40x.png)(convert:false) | **Sculpt Tool** | Raises and lowers the Landscape within the brushes influence. Use **Left Mouse Button** to raise or **Shift + Left Mouse button** to lower the Landscape.  |
+![](Engine/Landscape\icon_Landscape_Tool_Sculpt_40x.png)(convert:false) | **Sculpt Tool** | Raises and lowers the Landscape within the brushes influence. Use **Left Mouse Button** to raise or **Shift + Left Mouse Button** to lower the Landscape.  |
 [/VAR]
 
 [VAR:SmoothTool]
@@ -108,11 +109,11 @@ This document contains references to all of the icons for the **Landscape** syst
 [/VAR]
 
 [VAR:RetoptTool]
-![](Engine/Landscape\icon_Landscape_Tool_Retopologize_40x.png)(convert:false) | **Retopologize** | Similar to the smooth too, Retopologize will push and pull triangles to smooth the transition of the terrain, however it will attempt to keep the basic form of the terrain, minimizing change in the Z direction. **_An X/Y offset map makes the Landscape slower to render and paint on with other tools, so only use the Retopologize tool if needed._** |
+![](Engine/Landscape\icon_Landscape_Tool_Retopologize_40x.png)(convert:false) | **Retopologize** | Similar to the smooth too, Retopologize will push and pull triangles to smooth the transition of the terrain, however it will attempt to keep the basic form of the terrain, minimizing change in the Z direction. An X/Y offset map makes the Landscape slower to render and paint on with other tools, so only use the Retopologize tool if needed. |
 [/VAR]
 
 [VAR:VisTool]
-![](Engine/Landscape\icon_Landscape_Tool_Visibility_40x.png)(convert:false) | **Visibility** | This tool will allow you to mask out the visibility and collision of areas of your Landscape when used in conjunction with the [**Landscape Hole Material**](Engine/Landscape/Materials/).|
+![](Engine/Landscape\icon_Landscape_Tool_Visibility_40x.png)(convert:false) | **Visibility** | This tool will allow you to mask out the visibility and collision of areas of your Landscape when used in conjunction with the [Landscape Hole Material](Engine/Landscape/Materials/).|
 [/VAR]
 
 [VAR:SelectTool]
@@ -148,7 +149,7 @@ This document contains references to all of the icons for the **Landscape** syst
 [/VAR]
 
 [VAR:EditSplines]
-![](Engine/Landscape\icon_Landscape_Tool_Splines_40x.png)(convert:false) | **Edit Splines** | While using this tool, **Ctrl + Left Clicking** on the Landscape will place a node for a spline. If a node is selected and another node is placed, the tool will connect the new node to the previous node. If no nodes are selected, a new start node will be placed.  |
+![](Engine/Landscape\icon_Landscape_Tool_Splines_40x.png)(convert:false) | **Edit Splines** | While using this tool, **Ctrl + Left clicking** on the Landscape will place a node for a spline. If a node is selected and another node is placed, the tool will connect the new node to the previous node. If no nodes are selected, a new start node will be placed.  |
 [/VAR]
 
 [VAR:PaintTool]
@@ -173,7 +174,7 @@ This document contains references to all of the icons for the **Landscape** syst
 [EXCERPT:LSREF002]
 [REGION:raw]
 | ----- | ----- | ---- |
-| ![](Engine/Landscape\icon_Landscape_Mode_Sculpt_40x.png)(convert:false) | [**Sculpt**](Engine/Landscape/Editing/SculptMode) | This mode allows you to sculpt the look of your Landscape heightmap so that it fits your project�s needs. 
+| ![](Engine/Landscape\icon_Landscape_Mode_Sculpt_40x.png)(convert:false) | [**Sculpt**](Engine/Landscape/Editing/SculptMode) | This mode allows you to sculpt the look of your Landscape heightmap so that it fits your project's needs. 
 | ^ |  |  |
 %Ind% %SculptTool%
 %Ind% %SmoothTool%
