#UE4 Docs: Adding two new sections for global distance feilds and new images to go along with them.

diff --git a/Engine/Documentation/Source/Engine/Rendering/Materials/ExpressionReference/Utility/UtilityExpressions.INT.udn b/Engine/Documentation/Source/Engine/Rendering/Materials/ExpressionReference/Utility/UtilityExpressions.INT.udn
--- a/Engine/Documentation/Source/Engine/Rendering/Materials/ExpressionReference/Utility/UtilityExpressions.INT.udn
+++ b/Engine/Documentation/Source/Engine/Rendering/Materials/ExpressionReference/Utility/UtilityExpressions.INT.udn
@@ -343,9 +343,54 @@ The technique works best in magnification and with blurred input content. Compre
 
 [/EXCERPT:AntialiasedTextureMask]
 
+## DistanceToNearestSurface
 
+The **Distance To Nearest Surface** Material Expression node allows Materials to sample any point in the levels Global Distance Field. 
+This Material Expression works by outputting the signed distance in world space units from the distance field to the nearest occluders in the scene.
 
+[region:note]
+Please note that this Material Expression node requires that your project have the **Generate Mesh Distance Fields** enabled in order to work correctly. If you have not done so already, enabled this property in your projects settings under the Lighting category before going any further. 
+[/region]
 
+| Item | Description |
+| -------- | ----------- |
+| **Position**  | Defaults to current World Position if nothing is input.  |
+
+Here is an example of the **Distance To Nearest Surface** Material Expression in action. 
+
+![](DistanceNearSurface_00.png)
 
+In this example the Distance To Nearest Surface was fed into the Opacity input on a Material and that Material was applied to a Static Mesh plane that was placed just above the levels floor. 
+What the Distance To Nearest Surface is doing is telling the Material to only color areas red were the Static Meshes plane will start to intersect other Static Meshes placed in the scene.
 
+## DistanceFieldGradient
 
+The **DistanceFieldGradient** Material Expression node, when normalized, outputs the X,Y,Z direction an object would move with in the distance field. 
+This makes the Distance Field Gradient Material Expression node perfect for Materials that need to simulate the flow of liquids.
+
+[region:note]
+Please note that this Material Expression node requires that your project have the **Generate Mesh Distance Fields** enabled in order to work correctly. If you have not done so already, enabled this property in your projects settings under the Lighting category before going any further. 
+[/region]
+
+| Item | Description |
+| -------- | ----------- |
+| **Position**  | Defaults to current World Position if nothing is input.  |
+
+Here is an example of how to use the **DistanceFieldGradient** Material Expression in your Materials.
+In this example below make sure to note that the DistanceFieldGradient was first normalized and then input into a Mask Channel node.
+The reason for this is because without normalizing the DistanceFieldGradient first you can not get directional data. 
+The Mask Channel Parameters was added to allow for easier RGB channel switching with in the Material Instance.
+
+![](DistanceFieldGradient_Material.png)
+
+Here is an example of the DistanceFieldGradient in action. 
+The image below shows what data the DistanceFieldGradient will use when the various RGB are enabled.
+
+![](DistanceFieldGradient_Example.png)
+
+
+| Number | Description |
+| -------- | ----------- |
+| **1**  | Enabling the R channel and disabling all other channels. |
+| **2**  | Enabling the G channel and disabling all other channels. |
+| **3**  | Enabling the B channel and disabling all other channels. |
\ No newline at end of file
