UEDOC-2454

Replaced the video of the Blue guy with the Gray guy

diff --git a/Engine/Documentation/Source/Gameplay/HowTo/CharacterMovement/Blueprints/Overview.INT.udn b/Engine/Documentation/Source/Gameplay/HowTo/CharacterMovement/Blueprints/Overview.INT.udn
--- a/Engine/Documentation/Source/Gameplay/HowTo/CharacterMovement/Blueprints/Overview.INT.udn
+++ b/Engine/Documentation/Source/Gameplay/HowTo/CharacterMovement/Blueprints/Overview.INT.udn
@@ -18,13 +18,13 @@ For an example of the finished product, click the video sample below.
 360
 [/PARAMLITERAL]
 [PARAMLITERAL:videoid]
-EhspjHjDmtI
+f5DQQWQ8GEY
 [/PARAMLITERAL]
 [/OBJECT]
  
 For this series, we are using the **Animation Starter Pack**, which is available for free on the **Marketplace**, as it contains several animations that can be used to create character movement. Once you have downloaded the Animation Pack, you can add the pack to an existing project (preferably a New Blank Project) via the **Add to project** button in the Launcher then specifying which project you want to add the assets to. 
 
-![](Characters_2.png)
+![](AnimationStartePack1.jpg)
 
 To get started, follow the steps below. 
 
