#UE4doc #docs - Bolded term.

diff --git a/Engine/Documentation/Source/Programming/UnrealArchitecture/Reference/Structs/Specifiers/Immutable/StructSpecifier_Immutable.INT.udn b/Engine/Documentation/Source/Programming/UnrealArchitecture/Reference/Structs/Specifiers/Immutable/StructSpecifier_Immutable.INT.udn
--- a/Engine/Documentation/Source/Programming/UnrealArchitecture/Reference/Structs/Specifiers/Immutable/StructSpecifier_Immutable.INT.udn
+++ b/Engine/Documentation/Source/Programming/UnrealArchitecture/Reference/Structs/Specifiers/Immutable/StructSpecifier_Immutable.INT.udn
@@ -4,4 +4,4 @@ Description:
 Availability:Public
 Version: 4.9
 
-Immutable is only legal in Object.h and is being phased out, do not use on new structs!
+**Immutable** is only legal in Object.h and is being phased out, do not use on new structs!
