#UE4doc #docs - Added Engine Version Field.

diff --git a/Engine/Documentation/Source/Programming/UnrealArchitecture/Reference/Structs/Specifiers/StructSpecifiers.INT.udn b/Engine/Documentation/Source/Programming/UnrealArchitecture/Reference/Structs/Specifiers/StructSpecifiers.INT.udn
--- a/Engine/Documentation/Source/Programming/UnrealArchitecture/Reference/Structs/Specifiers/StructSpecifiers.INT.udn
+++ b/Engine/Documentation/Source/Programming/UnrealArchitecture/Reference/Structs/Specifiers/StructSpecifiers.INT.udn
@@ -2,6 +2,7 @@ Title: Struct Specifiers
 Crumbs:%ROOT%, Programming/UnrealArchitecture/Reference,  Programming/UnrealArchitecture/Reference, Programming/UnrealArchitecture/Reference/Structs
 Description: Keywords used when declaring gameplay structs to specify how the struct behaves with various aspects of the engine and editor.
 Availability:Public
+Version: 4.9
 
 
 [EXCERPT:main]
