Reverting streaming docs

diff --git a/Engine/Documentation/Source/Engine/LevelStreaming/StreamingVolumes/StreamingVolumes.INT.udn b/Engine/Documentation/Source/Engine/LevelStreaming/StreamingVolumes/StreamingVolumes.INT.udn
new file mode 100644
--- /dev/null
+++ b/Engine/Documentation/Source/Engine/LevelStreaming/StreamingVolumes/StreamingVolumes.INT.udn
@@ -0,0 +1,129 @@
+Availability: Public
+Title:Level Streaming Volumes
+Crumbs:%ROOT%, Editor, Editor/LevelEditing, Engine/LevelStreaming
+Description:Guide to using volumes to control the streaming of levels based on the player's location.
+
+[REGION:todo]
+	This page is converted from Unreal Engine 3 documentation and currently being reviewed. Information may be outdated.
+[/REGION]
+
+[TOC(start:2)]
+
+
+
+## Overview
+
+This document contains an overview of how to use volume-based level streaming.  It assumes the reader is already familiar with the basics of level streaming described in the [Level Streaming How-To](Engine/LevelStreaming/UserGuide).
+
+Level streaming volumes make controlling level streaming very easy.  The idea is simple: load/unload requests for a streaming level are issued based on whether the viewpoint is inside any of the LevelStreamingVolume volumes associated with a level.
+
+Specifically, level streaming volumes can be used in two ways:
+
+
+* In the **game**, level streaming volumes cause levels to load when the player is inside the volume and unload when the player is outside the volume.
+
+* In the **editor**, level streaming volumes can be used to preview level streaming by automatically hiding/unhiding levels based on the location of the perspective viewport's camera.
+
+
+Volume-based level streaming is simple to use and does not require any scripting, making it an ideal way to control level streaming.  Furthermore, volume-based level streaming is much easier to maintain than scripted streaming: when demands on the loading system change, level loading/unloading behaviour can be modified by simply resizing the streaming volumes.
+
+
+
+## Associating Streaming Volumes With Levels
+
+Volume-based level streaming works as follows: each streaming level can have associated with it a set of volumes of type LevelStreamingVolume.  Each frame, the engine iterates over each level and checks to see if the player viewpoint is inside any of the LevelStreamingVolumes associated with that level.  If the viewpoint is inside at least one LevelStreamingVolume, a request is issued to begin loading that level.  If the viewpoint is outside all LevelStreamingVolumes, the level is marked for unloading.
+
+To associate streaming volumes with a level:
+
+1. Drop a LevelStreamingVolume in the Persistent Level (as you would any other volume by using the "Add Volume" button on the main editor toolbar);
+1. Select the level(s) to associate with this volume in the Level Browser, right-click to bring up the Level Browser context menu, and select "Add Streaming Volumes".
+
+That's it!
+
+The full list of LevelStreamingVolume operations appearing in the Level Browser context menu is as follows:
+
+![LevelBrowserMenu.jpg](LevelBrowserMenu.jpg)
+
+
+* **Add Streaming Volumes**: Adds to the level(s)'s volumes, retaining any existing associations.
+
+* **Set Streaming Volumes**: Sets the level(s)'s volumes to be exactly the selected volumes, destroying any existing associations.
+
+* **Clear Streaming Volumes**: Clears all streaming volume associations from the selected level(s).
+
+* **Select Associated Streaming Volumes**: Selects all volumes associated with the selected level(s).
+
+
+Since it's easy for streaming volumes to clutter the view, LevelStreamingVolumes are only visible for levels that are **1) visible**; and **2) selected** in the level browser.  So, if you deselect the level you've just associated that volume with, you'll see the volume disappear in the editor viewports.  Selecting the level again in the Level Browser will cause the associated volumes to reappear.
+
+
+
+## Important Details
+
+
+* All LevelStreamingVolumes must exist in the persistent level.  LevelStreamingVolumes that live in other levels cannot be used for level streaming, and will generate warnings when the map is checked for errors.
+
+* Only levels that were added to the world using the "Blueprint" streaming method can use volume-based level streaming.
+
+* If a level has any streaming volumes associated with it, other methods of streaming the level (ie distance streaming or Blueprint streaming) will not behave correctly.
+
+* A single LevelStreamingVolume can affect multiple levels.  Similarly, a single level can be affected by multiple LevelStreamingVolumes.
+
+* Volume-based streaming works for split screen.  The viewpoints of all local players are considered before any loading/unloading requests are issued.
+
+
+
+
+## Testing Your Streaming Volume Setup
+
+It is critical that volume-based level streaming be tested in game on the target platform.  Streaming in "Play in Editor" (PIE) will show <i>where</i> the loads/unloads will happen, but **streaming in PIE is not representative of real, in-game loading-unloading**.  This is because in PIE, the levels are already in memory, and so "loading" a level is simply a matter of unhiding it and is thus instantaneous.
+
+Running the level in the standalone game on the target platform is critical in making sure your streaming setups work.  Note that on some platforms, it can take several seconds to stream in a level.  Size your LevelStreamingVolumes appropriately, so that the level is loaded by the time the player can reach it.  Level loading behaviour can be modified by resizing the LevelStreamingVolumes associated with a level.  Growing a volume causes associated levels to load sooner and unload later, while shrinking a volume causes later loads and earlier unloads.
+
+
+
+## Level Streaming Volumes for Previs
+
+![PrevisButton.jpg](PrevisButton.jpg)
+
+Volume-based level streaming can be easily previewed by enabling the "Level Streaming Volume Previs" button on the perspective viewport toolbar.  When this mode is enabled, levels will be hide/unhide based on the location of the perspective viewport camera.  Note that this hiding/unhiding will affect the set of visible levels in the Level Browser.
+
+A LevelStreamingVolume can be marked for editor previs only by setting the bEditorPreVisOnly flag on the LevelStreamingVolume.  In this manner, volume-based level streaming can still be used for editor previs while at the same time using another streaming method for in-game streaming.
+
+
+
+## Cost of Level Streaming Volumes
+
+Each frame, UWorld::ProcessLevelStreamingVolumes iterates over each streaming level and for each level, that level will begin loading if any local player is within any of the volumes associated with that level.  Likewise, the level will be begin unloading if all local players are outside all volumes.
+
+UWorld::ProcessLevelStreamingVolumes exploits coherency in the following manner: for each level, the volume that most recently contained a player is cached.  This cached volume is checked first, so that levels a player is in or returns to are quickly accepted.
+
+Volumes of any shape are fine, although obivously the fewer the better.  An upper bound on the cost of level streaming volumes can be approximated by the sum of the number of level streaming volumes associated with unloaded levels.
+
+Two stats exist under the "Streaming" stats group for monitoring level streaming performance.  The "Streaming Volumes" stat tracks the number of LevelStreamingVolumes tested against player viewpoints per frame, while the "Volume Streaming Tick" stat tracks the amount of time spent in UWorld::ProcessLevelStreamingVolumes per frame.
+
+
+
+## Adding Hysteresis to Unloading Requests
+
+A player moving back and forth across a LevelStreamingVolume boundary causes spurious load/unload requests to be issued.  To address this, hysteresis has been added to unloading requests.   No hysteresis exists for loading requests, because if a level needs to be loaded, we always want it loaded ASAP.
+
+The amount of unloading hystersis can be adjusted by modifying the MinTimeBetweenVolumeUnloadRequests property of the LevelStreaming object (LevelStreaming properties can be accessed from the Level Browser context menu).  The default unloading hysteresis is 2.0 seconds.
+
+
+
+## Disabling Level Streaming Volumes
+
+There is a property on LevelStreamingVolumes called bDisabled.  When set to true, the volume will be ignored by the streaming volume code, both in-game and in the editor.  bDisabled can be used to disable a LevelStreamingVolume without disassociating it from the level.  
+
+As an example of where the bDisabled flag is useful, imagine a door leading to a level whose streaming is controlled with streaming volumes.
+
+![bDisabledUsage.jpg](bDisabledUsage.jpg)
+
+In the image above, the player approaches from the left, and the level to stream in is on the right side of the door.  The streaming volume extends leftwards out past the door, so that the level will be streamed in by the time the player can reach the door and open it.  Initially, however the door is locked, and will become unlocked when the player achieves an objective in another part of the map.  So, even though the streaming volume extends past the door, we don't want the level on the other side of the door to stream in unless the door is actually unlocked ("openable").
+
+ 
+
+
+
+
diff --git a/Engine/Documentation/Source/Engine/LevelStreaming/UserGuide/UserGuide.INT.udn b/Engine/Documentation/Source/Engine/LevelStreaming/UserGuide/UserGuide.INT.udn
new file mode 100644
--- /dev/null
+++ b/Engine/Documentation/Source/Engine/LevelStreaming/UserGuide/UserGuide.INT.udn
@@ -0,0 +1,138 @@
+Availability: Public
+Title:Level Streaming User Guide
+Crumbs:%ROOT%, Editor, Editor/LevelEditing, Engine/LevelStreaming
+Description:Guide to setting up real-time asynchronous streaming of levels.
+
+[REGION:todo]
+	This page is converted from Unreal Engine 3 documentation and currently being reviewed. Information may be outdated.
+[/REGION]
+
+[TOC(start:2)]
+
+
+
+## Overview
+
+By following the directions in this document, Unreal Engine 4 users should be able to successfully set up the realtime streaming of multiple levels together to create the illusion of one large, seamless world.
+
+For information about volume-based level streaming, see the doc on [Level Streaming Volumes](Engine/LevelStreaming/StreamingVolumes).
+
+
+
+## The Persistent Level
+
+The first step towards getting your levels to blend together seamlessly is creating the persistent level. This can be thought of as a master level that is primarily used to govern which levels will be streamed in and out. 
+
+Creating a persistent level and adding streaming levels to it can be done in just a few simple steps: 
+
+1. Create a new map, add a small cube of geometry out of sight somewhere (this simply validates it as a level as far as the engine is concerned, and will likely be a step that is done away with eventually)
+1. Place a playerstart wherever you would like the player to begin. Note that in the persistent level this playerstart will likely be floating in space, but once you have streamed a level in, it should be moved to wherever you want the player to begin in the first streamed level. We will be launching the persistent level when it's time to play the game, but the player will actually be moving through the geometry of one of your streamed levels, not the persistent level.  
+1. Save your map, giving it a name denoting it as the persistent level for your own organizational purposes. Headquarters_P, for example.
+1. Open the generic browser, and click on the "Level" tab. This is where you must list every level that you wish to stream in. If Headquarters consists of three maps, all three should be added here. Do so by selecting Level->New From File from the Level pull-down menu. The order in which you add the streaming levels is not important, but future versions of Editor/UserGuide will allow you to re-order the levels to keep things organized. You should end up with something like what is in the picture below. 
+1. _Note that you will be given the option to stream your levels in by distance or by Blueprint_:
+    * Distance streaming loads and unloads a level based on the player's distance from the level's origin. 
+    * Blueprint streaming allows the designer to load and unload the level at will via Blueprint scripting, and thus is the most common choice here.
+
+![streamingImagenew1.jpg](streamingImagenew1.jpg)
+
+**Example 1: The level browser after three streamed levels have been added**
+
+The level denoted as "CURRENT" is the level that will be modified if changes are made in the editor windows or Blueprints. This allows you to easily work on multiple maps, as long as they are all writeable. If you press the save button, it will save the CURRENT level (or you can save all maps listed here including the persistent by clicking the Save All Levels button).
+
+The visible checkbox allows you to hide and unhide the levels in the editor, although this is for visualization purposes only and has no bearing on whether or not a level will stream into the game when it is run. However, a level that is not visible here will not be affected if you rebuild the level, which can be a great timesaver if you have complicated levels. 
+
+In order to permanently remove a streamed level from your persistent level, select Remove level from world from the Level pulldown in the browser. If you need to change a level's streaming method, remove the level and then re-add it, selecting the method you would prefer it to stream in by. Currently "Load" and "Unload" are not functional, nor is the Lock/Unlock icon or the Draw Bounding Box flag. 
+
+![streamingImagenew2.jpg](streamingImagenew2.jpg)
+
+**Example 2: Right-click and choose "Make Current" to change which level you are working on**
+
+
+### Persistent Level Quick How-To
+
+
+* Create a new map in Editor/UserGuide
+* Add a piece of simple BSP geometry, such as a cube
+* Add a playerstart
+* Add the levels you would like to stream in the Level Browser.
+* Save and name the map
+
+
+
+## Distance-based Streaming
+
+If you have chosen to stream your level in based on its distance from the player, there are some extra steps to take. First, make sure your persistent level is selected in the generic browser, and then expand the WorldSettings category on the far right-hand side. Right-click on the WorldSettings_# that you should now see and select "Edit Properties". Open the WorldSettings property in the property window and then open the StreamingLevels property. You should see each of the streaming levels listed here. You will now be able to set at what distance from the player the level will stream in at in the MaxDistance field. Think of this as a sphere that is centered at 0,0,0 in that level's world space by default. When the player enters this sphere, the level will stream in and when he leaves the sphere, the level will be unloaded. The location of this sphere can be changed, and most likely will have to be, since you'll in most cases want this sphere to completely surround the level you are streaming in, so that it is loaded when you are near the level, and unloaded when you leave the levels immediate area. We are planning on adding ways of visualizing where this sphere exists, but this has not yet been implemented. 
+
+This method of streaming can be useful for situations where you want to load and unload a level numerous times, based on player distance from the level. Note that there may be a noticeable hitch as the level loads, if it is a large map, as the level is actually being streamed in and out, not being hidden and unhidden. 
+
+![streamingImagenew4.jpg](streamingImagenew4.jpg)
+
+**Example 4: Setting up distance-based streaming**
+
+
+### Distance Streaming Quick How-To
+
+
+* Open the Generic Browser
+* Click the Level tab
+* Select Level->New From File from the Generic Browser file menu and select the level you would like to stream in
+* Choose "Distance" streaming
+* Making sure "Persistent Level" is selected, expand the WorldSettings property on the right-hand side. Right-click on the WorldSettings_# and select "Edit Properties"
+* In the property menu that should now be open, expand WorldSettings and then StreamingLevels.
+* Enter the distance from the player in game units at which you would like the level to be streamed into the world, and change the origin to the location in world space you would like that distance to be measured from. 
+
+
+
+
+## Working with the Streamed Levels
+
+
+### Loading and Events
+
+A new Kismet action was added to replace the Level Startup action that typically gets many events rolling when a level is loaded. This action is called Level Loaded And Visible. It basically functions in the same way as a Level Startup, but is designed to work in conjunction with the Change Level Visibility action. Whenever a level is made visible to the player, either by the CLV action or by streaming it in if the bMakeVisibleAfterLoad flag has been checked, this Loaded And Visible action will fire. This combined with Remote Events allows each streamed level to be its own self-contained entity that can be loaded and unloaded as necessary as well as receive and give events to other streamed levels. The example picture below has the player teleporting as soon as the subject level has been made visible, and also spawns a creature when a remote event from another level has been fired. 
+
+![streamingImagenew5.jpg](streamingImagenew5.jpg)
+
+**Example 5: Example script of a streamed level**
+
+
+### Moving Actors Between Levels
+
+As of the DEC 2007 QA-Approved build, you can move Actors between levels by simply copying them from one level (Ctrl + C), changing the current level and then pasting them (Ctrl + V). However, this uses a similar pathway to that of the _MoveSelectedActorsToCurrentLevel_ in that it has to serialize them out to a _scratch_ world in order to resolve potential references between the Actors. This functionality was introduced in changelist #203995 and #206418.
+
+In builds prior to the JAN 2007 QA-Approved build, you may have experiencd _out of memory_ problems when moving a large number of Actors at once; but this was fixed when we improved performance (and memory usage) of editor transaction buffers (changelist #206829).
+
+Also in the JAN 2007 QA-Approved build, you can move Actors into a new level using the _New Level from Selected Actors_ button in the [Level Browser](Editor/Browsers/Levels) (changelist #206324).
+
+
+### Optimization
+
+Most developers will find that the more streaming levels they have the better, both for performance and workflow reasons. However, this should be tempered by a limiting factor currently in the system, namely that actors that have been placed in one level cannot be directly referenced in kismet by another level. Remote events can be fired and received to trigger events across streaming levels, but it is not possible to directly reference actors across levels at this point. In the picture above, Pathnode_0 and Pathnode_1 must be located in the same level as this script, in this case Headquarters_2. If either of those nodes were in HQ_1 or 3, these actions would not function. 
+
+As of the December 2007 QA-Approved build, there is a feature that makes it easier to chop existing levels into streaming chunks.
+
+
+* Select a bunch of actors in the world.
+* Select _New Level from Selected Actors_ from the Level Browser's _Level_.
+* A new streaming level will be created on disk and added to the world, and the selected actors will then be moved to that new level.
+
+
+This should make it easier to create smaller streaming level sections!
+
+
+
+## Streaming Volumes
+
+For information about volume-based level streaming, see the [Level Streaming Volumes](Engine/LevelStreaming/StreamingVolumes) page.
+
+
+
+## Conclusion
+
+That concludes our overview of how to dynamically stream levels in and out of your game world. Since many aspects of this feature are new and constantly being updated, expect this document to be updated regularly as new features and simplifications come online. 
+
+ 
+
+
+
+
