#UE4 Docs: Added a new section that talks about sampling level AO in your Materails.

diff --git a/Engine/Documentation/Source/Engine/Rendering/Materials/ExpressionReference/Constant/ConstantExpressions.INT.udn b/Engine/Documentation/Source/Engine/Rendering/Materials/ExpressionReference/Constant/ConstantExpressions.INT.udn
--- a/Engine/Documentation/Source/Engine/Rendering/Materials/ExpressionReference/Constant/ConstantExpressions.INT.udn
+++ b/Engine/Documentation/Source/Engine/Rendering/Materials/ExpressionReference/Constant/ConstantExpressions.INT.udn
@@ -218,3 +218,27 @@ Here is the same example with a field of view of 90 degrees.
 [REGION:fullwidth]
 ![ViewPropertyExampleFOV90.png](ViewPropertyExampleFOV90.png)
 [/REGION]
+
+## Precomputed AO Mask
+
+The **PrecompitedAOMask** node lets you access Lightmass calculated AO in your material,
+ which can be useful for procedural texturing or for adding in aging effects and dirt in areas where thing of that nature would slowly accumulate over time.
+
+[region:note]
+Before you can see the results of the Precomputed AO Mask you will need to build the lighting of the level you are working in using Lightmass first before you can see the results in the Material.
+[/region]
+
+![](PrecompitedAOMas_00.png)
+
+The above screenshot uses an AO mask to blend a dirt layer automatically into corners of the environment.
+To use AO mask, you'll need to enable both **Use Ambient Occlusion** and **Generate Ambient Occlusion Material Mask** under **World Settings -> Lightmass settings**, and then build lighting. 
+The other AO controls like **Max Occlusion Distance** can be useful to tweak the look of the AO. 
+Also make sure to set both **Direct** and **Indirect Occlusion Fraction** to **0**, so that this AO will not be applied to the actual level lighting.
+
+![](PrecompitedAOSettings_00.png)
+
+You can access the AO in any Material by using the **PrecomputedAOMask** Material Expression node. 
+The PrecomputedAOMask works as a 0 to 1 mask with 1 providing areas that are effected by AO and 0 for areas that are not.
+In the image below you can see how how to setup a Material to make use of the Precomputed AO Mask.
+
+![](PrecomputedAO_Material_00.png)
