#UE4doc #docs - Updated images, and description to replace references to self shadow accuracy to shadow bias. Fixed CE example 6.1. References UEDOC-2015

diff --git a/Engine/Documentation/Source/Resources/ContentExamples/Lighting/6_1/6_1.INT.udn b/Engine/Documentation/Source/Resources/ContentExamples/Lighting/6_1/6_1.INT.udn
--- a/Engine/Documentation/Source/Resources/ContentExamples/Lighting/6_1/6_1.INT.udn
+++ b/Engine/Documentation/Source/Resources/ContentExamples/Lighting/6_1/6_1.INT.udn
@@ -1,5 +1,5 @@
 Availability: Public
-Title:6.1 - Self Shadowing Accuracy
+Title:6.1 - Shadow Bias
 Crumbs:%ROOT%, Resources, Resources/ContentExamples, Resources/ContentExamples/Lighting
 Description:An overview of the Lighting example level, example 6.1: Self Shadowing Accuracy
 Related: Engine/Rendering/LightingAndShadows/Basics
@@ -51,15 +51,15 @@ Related: Engine/Rendering/LightingAndShadows/LightMobility
 
 
 [REGION:fullwidth]
-![%Resources\ContentExamples\Lighting\6_1:title%](Resources\ContentExamples\Lighting\Lighting_6_1.png)
+![%Resources\ContentExamples\Lighting\6_1:title%](Resources\ContentExamples\Lighting\Lighting_6_1.png)(w:700)
 [/REGION]
 
 
-This example shows off the Self Shadowing Accuracy, or **Shadow Bias**. This property exists primarily to help objects feel more grounded along a surface. There is a tradeoff, however, in that increasing the value to make an object appear more grounded can sometimes lead to self-shadowing artifacts, as seen in the image below:
+This example shows off the Self Shadowing Accuracy, or **Shadow Bias**. This property exists primarily to help objects feel more grounded along a surface. There is a tradeoff, however, in that decreasing the value to make an object appear more grounded can sometimes lead to self-shadowing artifacts, as seen in the image below:
 
-![](SelfShadowingArtifacts.png)
+![](SelfShadowingArtifacts.png)(w:700)
 
-The diagonal streaks on the surface are the result of these artifacts, created from a high Self Shadowing Accuracy value. It should be noted that you will generally not need to change this property, as the default will work in most instances. 
+The diagonal streaks on the surface are the result of these artifacts, created from a low Shadow Bias value. It should be noted that you will generally not need to change this property, as the default will work in most instances. 
 
 
 
diff --git a/Engine/Documentation/Source/Resources/ContentExamples/Lighting/Lighting.INT.udn b/Engine/Documentation/Source/Resources/ContentExamples/Lighting/Lighting.INT.udn
--- a/Engine/Documentation/Source/Resources/ContentExamples/Lighting/Lighting.INT.udn
+++ b/Engine/Documentation/Source/Resources/ContentExamples/Lighting/Lighting.INT.udn
@@ -118,7 +118,7 @@ Listed below are each of the sections provided inside the **Lighting** map:
 
 | Example | What is Demonstrated |
 | ------- | ------------ |
-| **[6.1 Self Shadowing Accuracy](Resources\ContentExamples\Lighting\6_1)** | Shows how an object can cast shadows on itself to help it feel more grounded to a surface. |
+| **[6.1 Shadow Bias](Resources\ContentExamples\Lighting\6_1)** | Shows how an object can cast shadows on itself to help it feel more grounded to a surface. |
 | **[6.2 Shadow Filter Sharpen](Resources\ContentExamples\Lighting\6_2)** | A setting used to boost the sharpness at the edges of shadows. |
 
 <!--
