#UE4 Docs: Added a new section for the new Particel UV flipping feature.

diff --git a/Engine/Documentation/Source/Engine/Rendering/ParticleSystems/Reference/Modules/Required/Required.INT.udn b/Engine/Documentation/Source/Engine/Rendering/ParticleSystems/Reference/Modules/Required/Required.INT.udn
--- a/Engine/Documentation/Source/Engine/Rendering/ParticleSystems/Reference/Modules/Required/Required.INT.udn
+++ b/Engine/Documentation/Source/Engine/Rendering/ParticleSystems/Reference/Modules/Required/Required.INT.udn
@@ -51,6 +51,7 @@ Every particle emitter contains this module which contains all the properties th
 | [REGION:tablesection]Rendering[/REGION] ||
 | **Use Max Draw Count** | When _true_, the emitter will never draw more than MaxDrawCount particles. All particles will still be updated when ticking. |
 | **Max Draw Count** | The number of particles to limit rendering to. |
+| **UV Flipping Mode** | Enabled the flipping of Sprite and GPU particles UV's in a number of different ways. [INCLUDE:#UVFlipping] |
 | [REGION:tablesection]Normals[/REGION] ||
 | **Emitter Normals Mode** | The mode used to generate normals for this emitter LOD. [INCLUDE:#NormalsMode] |
 | **Normals Sphere Center** | When _Emitter Normals Mode_ is _ENM_Spherical_, particle normals are created to face away from _Normals Sphere Center_. The value is in local space. |
@@ -62,8 +63,6 @@ Every particle emitter contains this module which contains all the properties th
 
 Every particle emitter contains this module which contains all the properties that determine how particles are spawned. This includes the following:
 
-
-
 | Property | Description |
 | -------- | ----------- |
 | [REGION:tablesection]Spawn[/REGION] ||
@@ -106,6 +105,22 @@ Every particle emitter contains this module which contains all the properties th
 [/REGION]
 [/EXCERPT:SortMode]
 
+[EXCERPT:UVFlipping]
+[REGION:simpletable]
+| Mode | Description |
+| --- | --- |
+|None|No UV flipping occurs. This is the default.|
+|FlipUV|Flips the U and V texture coordinates.|
+|U Only|Flips only the U texture coordinate.|
+|V Only|Flips the V texture coordinate.|
+|Random Flip UV|Randomly flips U and V together at spawn time.|
+|Random Flip U Only|Randomly flips just U at spawn time.|
+|Random Flip V Only|Randomly flips just V at spawn time.|
+|Random Flip UV Independent|Randomly flips U and V independently at spawn time.|
+[/REGION]
+[/EXCERPT:UVFlipping]
+
+
 [EXCERPT:InterpMethod]
 [REGION:simpletable]
 | Method | Description |
