#UE4doc #docs:
Added engine version field
Changed font settings for multiple terms

diff --git a/Engine/Documentation/Source/Engine/Performance/PerformanceProfiling.INT.udn b/Engine/Documentation/Source/Engine/Performance/PerformanceProfiling.INT.udn
--- a/Engine/Documentation/Source/Engine/Performance/PerformanceProfiling.INT.udn
+++ b/Engine/Documentation/Source/Engine/Performance/PerformanceProfiling.INT.udn
@@ -3,6 +3,7 @@ Title: Performance and Profiling
 Crumbs: %ROOT%, Engine
 Description:How to identify performance problems and fix them.
 Navigation:topic
+version: 4.9
 
 [VAR:Topic]
 [OBJECT:Topic]
@@ -66,7 +67,7 @@ You can use the editor output log or the in-game console to:
 * Get the current state (**cvarname**) 
 * Get help for a variable **cvarname ?**. 
 
-If needed, you can put the settings in `ConsoleVariables.ini` with this syntax: **cvarname=value**. To find the right console variable, you can use **DumpConsoleVars** or the auto completion system. 
+If needed, you can put the settings in ConsoleVariables.ini with this syntax: **cvarname=value**. To find the right console variable, you can use **DumpConsoleVars** or the auto completion system. 
 Most rendering variables start with **r.**.
 
 * For more about console variables and toggling other engine options to improve performance, see the [](Engine/Performance\Options) reference page.
@@ -117,18 +118,18 @@ Here we can see the **GPU** is the limiting factor (the largest of the three). I
 
 The engine show flags can be used to toggle a lot of rendering features. The editor lists all show flags in a convenient 2D UI. You can click on the checkbox to toggle multiple checkboxes without closing the menu.
 
-In-game, you can use the show command. Use show to get a list of all show flags and their state. Use `show showflagname` to toggle a feature. Note that this only works in game viewports, so in editor
-viewports, the editor UI needs to be used. You can use the console variables (e.g. showflag.Bloom) to override show flag values in-game or in the editor, but that also disables the UI.
+In-game, you can use the show command. Use show to get a list of all show flags and their state. Use **_show showflagname_** to toggle a feature. Note that this only works in game viewports, so in editor
+viewports, the editor UI needs to be used. You can use the console variables (e.g. **_showflag.Bloom_**) to override show flag values in-game or in the editor, but that also disables the UI.
 
 [REGION:note]
 Some features still consume performance even if the feature is not rendering any more, e.g. show particles hides the particles but simulation time is still
-required to support re-enabling them at a later time. The console variable fx.freezeparticlesimulation allows to you disable the update part as well (editor UI also exists).
+required to support re-enabling them at a later time. The console variable **_fx.freezeparticlesimulation_** allows to you disable the update part as well (editor UI also exists).
 Another one: show tessellation disables triangle amplification but still uses the tessellation shaders.
 [/REGION:note]
 
 A good profiling starting point is with high-level features, such as show StaticMeshes or show tessellation.
 
-All show flags are also exposed as console variables, e.g. console show Bloom, showflag.Bloom 0 or in configuration files: showflag.Bloom = 0
+All show flags are also exposed as console variables, e.g. console show Bloom, **_showflag.Bloom 0_** or in configuration files: showflag.Bloom = 0
 Console variables require more typing but they also override the editor UI showflags and they can be put into configuration files like other console variables.
 
 The most useful ones for profiling:
@@ -167,13 +168,13 @@ The most useful ones for profiling:
 ### View Modes
 
 View modes are just a combination of show flags. The editor UI is separate from the show flags and you can use the ViewMode console command to switch between them.
-The most useful for performance are: Wireframe, LightComplexity, ShaderComplexity, and Lit.
+The most useful for performance are: **Wireframe**, **LightComplexity**, **ShaderComplexity**, and **Lit**.
+
+
+![](ViewModes.png)(w:794)
+
+_Various view modes (in reading order): Lit, LightCompexity (darker is better), Wireframe, Shader Complexity (green is good)_
 
-[REGION:imagetable]
-| ![](ViewModes.png)(w:794) |
-| :----: | 
-| Various view modes (in reading order): Lit, LightCompexity (darker is better), Wireframe, Shader Complexity (green is good) |
-[/REGION]
 
 ## How to Deal With a Wide Range of Hardware
 
