#UE4doc #docs:
Added engine version field

diff --git a/Engine/Documentation/Source/Engine/Performance/CPU/CPU.INT.udn b/Engine/Documentation/Source/Engine/Performance/CPU/CPU.INT.udn
--- a/Engine/Documentation/Source/Engine/Performance/CPU/CPU.INT.udn
+++ b/Engine/Documentation/Source/Engine/Performance/CPU/CPU.INT.udn
@@ -3,6 +3,7 @@ Title: CPU Profiling
 Crumbs: %ROOT%, Engine
 Description:How to optimize the CPU demands of your game.
 Navigation:topic
+version: 4.9
 
 
 If you are CPU bound in the render thread, it is likely because of too many draw calls. This is a common problem and artists often
