#UE4doc #docs:
Added engine version field
Changed font settings for multiple terms
Updated step to opening Profiler

diff --git a/Engine/Documentation/Source/Engine/Performance/Profiler/Profiler.INT.udn b/Engine/Documentation/Source/Engine/Performance/Profiler/Profiler.INT.udn
--- a/Engine/Documentation/Source/Engine/Performance/Profiler/Profiler.INT.udn
+++ b/Engine/Documentation/Source/Engine/Performance/Profiler/Profiler.INT.udn
@@ -2,6 +2,7 @@ Availability: Public
 Title:Profiler
 Crumbs: %ROOT%, Engine, Engine/Performance
 Description:An overview on the Profiler Tool used to monitor game performance.
+version: 4.9
 
 [REGION:fullwidth]
 ![](Profiler_Header.png)
@@ -21,13 +22,13 @@ There are two methods you can use to enable **Profiler** and collect data:
 1. Live Connection
 	* Run the game with the parameter -messaging (for example: **UE4Editor-Win64-Debug.exe -messaging**).
 	* Run the UFE with the parameter -messaging (for example: **UnrealFrontend-Win64-Debug.exe -messaging**).
-	* Select **Session Frontend** from the **Window** menu bar, then select the **Data Preview** tab. 
+	* Select **Session Frontend** from the **Developer Tools** section of the **Window** menu bar, then select the **Profiler** tab. 
 1. Load previously captured data
 	* By pressing the **Load** button on the toolbar from the **Profiler** tab and selecting a **.ue4stats** file. 
 	* By dropping a stats file from explorer or other file manager into the **Profiler** window. 
 
 [REGION:note]
-You can also launch the stand-alone **UnrealFrontEnd.exe** from the **C:\Program Files\Unreal Engine\4.3\Engine\Binaries\Win64** folder to gain access to **Session FrontEnd** and the **Profiler** tab. _(Note: your version number may differ from the path above)_.
+You can also launch the stand-alone **UnrealFrontEnd.exe** from the **C:\Program Files\Unreal Engine\4.3\Engine\Binaries\Win64** folder to gain access to **Session FrontEnd** and the **Profiler** tab. (Note: your version number may differ from the path above).
 [/REGION]
 	
 ## Profiler Window
@@ -129,7 +130,7 @@ The stat's text consists of the following information and is broken down below:
 1. **Max Value** - The maximum number of calls to this stat based on the whole capture (for hierarchical data only).
 
 [REGION:note]
-The values denoted with **for hierarchical data only** are only displayed for hierarchical data and may not be displayed. 
+The values denoted with for hierarchical data only are only displayed for hierarchical data and may not be displayed. 
 [/REGION]
 
 An explanation of the **Data Graph** axes are also outlined below: 
