#UE4doc #docs:
Added engine version field

diff --git a/Engine/Documentation/Source/Engine/Performance/Scalability/Scalability.INT.udn b/Engine/Documentation/Source/Engine/Performance/Scalability/Scalability.INT.udn
--- a/Engine/Documentation/Source/Engine/Performance/Scalability/Scalability.INT.udn
+++ b/Engine/Documentation/Source/Engine/Performance/Scalability/Scalability.INT.udn
@@ -2,6 +2,7 @@ Availability:Public
 Title: Scalability
 Crumbs: %ROOT%, Engine, Engine/Performance
 Description:Options for adjusting quality and performance.
+version: 4.9
 
 [VAR:Topic]
 [OBJECT:Topic]
