#UE4doc #docs:
Added engine version field
Removed monospace from a term

diff --git a/Engine/Documentation/Source/Engine/Performance/Scalability/ScalabilityAndYou/ScalabilityAndYou.INT.udn b/Engine/Documentation/Source/Engine/Performance/Scalability/ScalabilityAndYou/ScalabilityAndYou.INT.udn
--- a/Engine/Documentation/Source/Engine/Performance/Scalability/ScalabilityAndYou/ScalabilityAndYou.INT.udn
+++ b/Engine/Documentation/Source/Engine/Performance/Scalability/ScalabilityAndYou/ScalabilityAndYou.INT.udn
@@ -2,6 +2,7 @@ Availability:Public
 Title: Scalability and The Developer
 Crumbs: %ROOT%, Engine, Engine/Performance, Engine/Performance/Scalability
 Description:Desc
+version: 4.9
 
 [VAR:Topic]
 [OBJECT:Topic]
@@ -109,7 +110,7 @@ A wide range of scalability costs more QA time. A few good practices can help:
 * If you run in lowest quality settings and something looks quite better than expected - it might miss a scalability setting (fixing this could mean more players can enjoy the game).
 * Test usability of the scalability: report anything if you find it confusing and wrong.
 * At any time the user interface should be readable and functional. Please test low resolution and lowest quality settings and look for unreadable fonts and elements outside of the screen.
-* If some artifact only happen at some specific scalability level (e.g. only on "low") you can isolate the issue further (e.g. shadows are likely affected by "shadow quality") going down to the actual console variable setting (look into `Scalability.ini` file).
+* If some artifact only happen at some specific scalability level (e.g. only on "low") you can isolate the issue further (e.g. shadows are likely affected by "shadow quality") going down to the actual console variable setting (look into Scalability.ini file).
 * Know how to quantify performance (stat fps, stat unit, ...), milliseconds are more useful than fps ("it became 10fps faster" is not meaningful but "it became 5ms faster" is).
 * Some notebooks have multiple GPUs (e.g. Intel HD4000 integrated graphics to save battery and NVIDIA 580m for performance). By default, you want to run on the faster GPU but it is best to run test both (use the driver setting to change that). When running on battery, it might pick the slower one or even switches it at runtime.
 
