#UE4 Docs: Updating text and images to reflect the latest chagnes to the tool.

diff --git a/Engine/Documentation/Source/Engine/Rendering/Materials/MaterialInstances/Editor/MaterialInstanceEditor.INT.udn b/Engine/Documentation/Source/Engine/Rendering/Materials/MaterialInstances/Editor/MaterialInstanceEditor.INT.udn
--- a/Engine/Documentation/Source/Engine/Rendering/Materials/MaterialInstances/Editor/MaterialInstanceEditor.INT.udn
+++ b/Engine/Documentation/Source/Engine/Rendering/Materials/MaterialInstances/Editor/MaterialInstanceEditor.INT.udn
@@ -2,14 +2,12 @@ Availability: Public
 Title:Material Instance Editor User Guide
 Crumbs: %ROOT%, Engine, Engine/Rendering/Materials, Engine/Rendering/Materials/MaterialInstances
 Description:Guide to using the Material Instance Editor for modifying MaterialInstanceConstants.
-
+Version: 4.9
+Skilllevel: Intermediate
 [TOC(start:2)]
 
-
-
-## Overview
-
-The material instance editor is used for modifying parameters for material instances. If you are unfamiliar with material instances, please see the [Instanced Materials](Engine/Rendering/Materials/MaterialInstances) page for more information.
+The Material Instance editor is used for modifying parameters for Material Instances. 
+If you are unfamiliar with Material Instances, please see the [Instanced Materials](Engine/Rendering/Materials/MaterialInstances) page for more information.
 
 
 
@@ -54,8 +52,11 @@ The Material Instance Editor is comprised of these regions:
 | ![](button_MIE_Plane.png) | Previews the Material Instance on a Plane.  |
 | ![](button_MIE_Cube.png) | Previews the Material Instance on a Cube.  |
 | ![](button_MIE_Mesh.png) | Previews the Material Instance on a custom Static Mesh.  |
-| ![](button_MIE_Grid.png) |  Toggles the background grid in the material instance Viewport. |
-| ![](button_MIE_Realtime.png) |  If enabled, the Material on the preview mesh will update in real-time.  Disable this flag for editor performance. |
+| ![](Viewport_Options_Button.png) | Enables real time preview, FPS and allows you to set the FOV of the preview window.  |
+| ![](View_Type_Button.png) | Changes between Perspective and Orthographic viewing options. |
+| ![](Viewmode_Options_Button.png) | Changes between the different standard view modes like Lit, Unlit, Wireframe and others. |
+| ![](Show_Button.png) | Enables or disables rendering stats, the grid and the background. |
+
 
 ### Details Pane
 
@@ -68,8 +69,8 @@ The **Details** pane in the Material Instance editor is where all of the 'work'
  no longer be grayed out and you will be able to adjust the parameter to your liking. 
 
 * **Parameter Groups** - In this section you can adjust all of the values that have been exposed to the Material Instance from the Parent Material. 
-* **General** - The General section allows you to change the Parent Material, Phys Material, and also adjust how this Material Instance will influence the Lightmass build. 
-* **Material Overrides** - Material Overrides allows you to adjust Material base values 
+* **General** - The General section allows you to change the Parent Material, Phys Material, adjust how this Material Instance will influence the Lightmass build and override some of the properties inherited from the parent Material. 
+
 
 
 ### Parent List
@@ -107,14 +108,15 @@ samples being used by the Material.
 
 ![](MIE_Ins_Count.png)
 
-If your Material has some kind of movement in it you and you want to preview it, you can enable Real - Time preview by clicking on the stop watch icon.
+If your Material has some kind of movement in it you and you want to preview it, you can enable Real - Time preview by first clicking on the small white triangle to enable Viewport Options and then click on the stop watch icon.
+You can also enable / disable real time preview by pressing **CTRL + R** at the same time on the keyboard.
 
 ![](MIE_Real_Time.png)
 
 
 ## Creating Instances
 
-[INCLUDE:Engine/Rendering/Materials\MaterialInstances#MakingAnInstance]
+[INCLUDE:Engine/Rendering/Materials/MaterialInstances#MakingAnInstance]
 
 ## Overriding Parameters
 
