#UE4doc #docs:
Added engine version field
Changed font settings for multiple terms
Replaced a couple images due to out of date UI
Removed an image due to incorrect placement

diff --git a/Engine/Documentation/Source/Engine/Performance/Scalability/ScalabilityReference/ScalabilityReference.INT.udn b/Engine/Documentation/Source/Engine/Performance/Scalability/ScalabilityReference/ScalabilityReference.INT.udn
--- a/Engine/Documentation/Source/Engine/Performance/Scalability/ScalabilityReference/ScalabilityReference.INT.udn
+++ b/Engine/Documentation/Source/Engine/Performance/Scalability/ScalabilityReference/ScalabilityReference.INT.udn
@@ -2,6 +2,7 @@ Availability:Public
 Title: Scalability Reference
 Crumbs: %ROOT%, Engine, Engine/Performance/Scalability
 Description:Scalability options, properties, and console variables.
+version: 4.9
 
 [VAR:Topic]
 [OBJECT:Topic]
@@ -53,7 +54,7 @@ The Scalability settings allow you to adjust the quality of various features in
 
 ## Scalability Settings
 
-Accessing the Scalability settings in the editor is done through the **Quick Settings** menu in the toolbar ![](%ROOT%/Shared/Icons/icon_view_40x.png). The **Engine Scalability Settings** contains short cuts to the most commonly used setting defined in the `BaseScalability.ini` file, while **Material Quality Level** sets the global material quality setting.
+Accessing the Scalability settings in the editor is done through the **Settings** menu in the toolbar. The **Engine Scalability Settings** contains short cuts to the most commonly used setting defined in the BaseScalability.ini file, while **Material Quality Level** sets the global material quality setting.
 
 [REGION:lightbox]
 [![](ScaleR_ViewQual.png)(h:490)](ScaleR_ViewQual.png) [![](ScaleR_ViewMatQual.png)(h:490)](ScaleR_ViewMatQual.png)
@@ -74,12 +75,12 @@ Left: 50% no AA, Middle: 50% with AA, Right: 100% (no resolution scale) with AA
 Having a softer input image helps the up-sampling step. This means this scalability options benefits from another scalability option: the anti-aliasing quality.
 [/REGION]
 
-This setting can be access with the r.ScreenPercentage console variable. It takes in a value from 10-100 with -1 also equaling 100%.
+This setting can be access with the **_r.ScreenPercentage_** console variable. It takes in a value from 10-100 with -1 also equaling 100%.
 
 
 ## View Distance
 
-Objects can be culled based on their distance to the viewer. By default, all objects are not distance culled (Desired max draw distance of 0). On top of the designer specified value, there is a global scalability setting working like a multiplier (r.ViewDistanceScale). Here you can see some grass objects (Desired max draw distance of 1000):
+Objects can be culled based on their distance to the viewer. By default, all objects are not distance culled (Desired max draw distance of 0). On top of the designer specified value, there is a global scalability setting working like a multiplier (**_r.ViewDistanceScale_**). Here you can see some grass objects (Desired max draw distance of 1000):
 
 ![](SC_00006.png)(w:920)
 [REGION:caption]
@@ -93,11 +94,11 @@ left: r.ViewDistanceScale 0.4, middle: r.ViewDistanceScale 0.7, right: r.ViewDis
 [![](ScaleR_AA_small.png)(w:920)](ScaleR_AA.png)
 [/REGION]
 [REGION:caption]
-From the left: r.PostProcessAAQuality 0 to 6. The first 3 equal the **Low**, **Medium**, **High/Epic** settings in the **View -> Scalability Options -> AA** settings.
+From the left: r.PostProcessAAQuality 0 to 6. The first 3 equal the **Low**, **Medium**, **High**, **Epic** settings in the **View -> Scalability Options -> AA** settings.
 [/REGION]
-Adjusting the Anti-Aliasing quality level using the r.PostProcessAAQuality console command will adjust the quality of which ever Anti-Aliasing method you are using (FXAA or Temporal AA). For either Anti-Aliasing method, a value of 0 used with r.PostProcessAAQuality will disable the effect. For FXAA, the effect of values 2, 4, and 6, can be seen in the above image; the smoothing of jagged edges becomes better and better. Values above 6 have no effect.
+Adjusting the Anti-Aliasing quality level using the **_r.PostProcessAAQuality_** console command will adjust the quality of which ever Anti-Aliasing method you are using (FXAA or Temporal AA). For either Anti-Aliasing method, a value of 0 used with _r.PostProcessAAQuality_ will disable the effect. For FXAA, the effect of values 2, 4, and 6, can be seen in the above image; the smoothing of jagged edges becomes better and better. Values above 6 have no effect.
 
-For Temporal AA, there is a trade off between fill speed of the effect and quality, the higher the value you use. r.PostProcessAAQuality 2 with Temporal AA is fast to settle but jitter caused by the effect will be more pronounced. r.PostProcessAAQuality 4 will settle slower but will not jitter.
+For Temporal AA, there is a trade off between fill speed of the effect and quality, the higher the value you use. _r.PostProcessAAQuality_ 2 with Temporal AA is fast to settle but jitter caused by the effect will be more pronounced. _r.PostProcessAAQuality_ 4 will settle slower but will not jitter.
 
 
 ## Post Processing - sg.PostProcessQuality
@@ -106,10 +107,10 @@ For Temporal AA, there is a trade off between fill speed of the effect and quali
 [![](SC_PostProcessComp.png)(w:900)](SC_PostProcessComp.png)
 [/REGION]
 [REGION:caption]
-sg.PostProcessQuality set from 0 on the left to 3 on the right.
+**_sg.PostProcessQuality_** set from 0 on the left to 3 on the right.
 [/REGION]
 
-The **View** -> **Engine Scalability Settings** -> **Post Processing** option adjusts the quality of the Post Processing effects in accordance with the settings found in the BaseScalaiblity.ini file located in [UE4_InstallPath]/Engine/Config folder. With the **Low** setting equating to sg.PostProcessQuality 0 and **Epic** to sg.PostProcessQuality 3.
+The **View** -> **Engine Scalability Settings** -> **Post Processing** option adjusts the quality of the Post Processing effects in accordance with the settings found in the BaseScalaiblity.ini file located in [UE4_InstallPath]/Engine/Config folder. With the **Low** setting equating to _sg.PostProcessQuality_ 0 and **Epic** to _sg.PostProcessQuality_ 3.
 
 [REGION:simpletable]
 |sg.PostProcessQuality 0|sg.PostProcessQuality 1|
@@ -200,7 +201,7 @@ The **View** -> **Engine Scalability Settings** -> **Post Processing** option ad
 sg.ShadowQuality set from 0 on the left to 3 on the right.
 [/REGION]
 
-The **View** -> **Engine Scalability Settings** -> **Shadows** option adjusts the quality of dynamic shadows in accordance with the settings found in the BaseScalaiblity.ini file located in [UE4_InstallPath]/Engine/Config folder. With the **Low** setting equating to sg.ShadowQuality 0 and **Epic** to sg.ShadowQuality 3.
+The **View** -> **Engine Scalability Settings** -> **Shadows** option adjusts the quality of dynamic shadows in accordance with the settings found in the BaseScalaiblity.ini file located in [UE4_InstallPath]/Engine/Config folder. With the **Low** setting equating to _sg.ShadowQuality_ 0 and **Epic** to _sg.ShadowQuality_ 3.
 
 [REGION:simpletable]
 |sg.ShadowQuality 0|sg.ShadowQuality 1|
@@ -377,7 +378,7 @@ The scalability mode can be changed from the console:
 
 	r.DetailMode 1
 
-With r.DetailMode 0 = low, r.DetailMode 1 = medium and, r.DetailMode 2 = high detail modes.
+With _r.DetailMode_ 0 = low, _r.DetailMode_ 1 = medium and, _r.DetailMode_ 2 = high detail modes.
 
 ![](SC_00005.png)
 
@@ -387,7 +388,7 @@ It is easy to use this to disable decals, detail objects, lights or individual p
 
 ![](ScaleR_ViewMatQual.png)
 
-Material can use the "Quality Switch" material expression to disable some expensive material parts that have minor effect on the final look. To see the effect of this, you need to toggle to the low quality mode.
+Material can use the Quality Switch material expression to disable some expensive material parts that have minor effect on the final look. To see the effect of this, you need to toggle to the low quality mode.
 
 [REGION:lightbox]
 [![](SC_SwitchNode.png)(w:900)](SC_SwitchNode.png)
@@ -432,6 +433,6 @@ Most material editor outputs only affect the pixel shader. World position offset
 
 ![](SC_Skeletal.png)(w:920)
 
-Skeletal meshes can have static Level of Details models. With the console variable r.SkeletalMeshLODBias, it is possible to globally bias the LOD level. Depending on the quality of your first or second LOD, it might be a good idea to have this option in the scalability settings. Here you can see the scalability setting on top of the distance dependent LOD.
+Skeletal meshes can have static Level of Details models. With the console variable **_r.SkeletalMeshLODBias_**, it is possible to globally bias the LOD level. Depending on the quality of your first or second LOD, it might be a good idea to have this option in the scalability settings. Here you can see the scalability setting on top of the distance dependent LOD.
 
 
