#DocFix

Forum request on how a CameraComponent is added to an actor

diff --git a/Engine/Documentation/Source/Programming/Tutorials/AutoCamera/1/AutoCamera_1.INT.udn b/Engine/Documentation/Source/Programming/Tutorials/AutoCamera/1/AutoCamera_1.INT.udn
--- a/Engine/Documentation/Source/Programming/Tutorials/AutoCamera/1/AutoCamera_1.INT.udn
+++ b/Engine/Documentation/Source/Programming/Tutorials/AutoCamera/1/AutoCamera_1.INT.udn
@@ -37,7 +37,7 @@ version: 4.9
 
 	[REGION:note]We can use almost any actor class for this step. The MyActor class we created in the QuickStart tutorial might be interesting to substitute for the Cube.[/REGION]
 
-1. When our **Cube** actor is placed, we should add a **CameraComponent** to it, and set the position and rotation of that **CameraComponent** to give us a different view of the scene than the **CameraActor** we placed earlier.
+1. When our **Cube** actor is placed, add a **CameraComponent** by clicking the **+ Add Component** button in the **Details** panel for the **Cube**. You can now set the position and rotation of that **CameraComponent** to give us a different view of the scene than the **CameraActor** we placed earlier.
 
 	[![](CameraTwoPlacement.png)](CameraTwoPlacement.png)
 
