#UE4doc #docs:
Added engine version field
Changed font settings for multiple terms
Rewrote how a section explained where Bounds Scale was for a component
Added a picture to illustrate the previous point

diff --git a/Engine/Documentation/Source/Engine/Performance/Visualizers/OutOfBoundsPixels/OutOfBoundsPixels.INT.udn b/Engine/Documentation/Source/Engine/Performance/Visualizers/OutOfBoundsPixels/OutOfBoundsPixels.INT.udn
--- a/Engine/Documentation/Source/Engine/Performance/Visualizers/OutOfBoundsPixels/OutOfBoundsPixels.INT.udn
+++ b/Engine/Documentation/Source/Engine/Performance/Visualizers/OutOfBoundsPixels/OutOfBoundsPixels.INT.udn
@@ -2,6 +2,7 @@ Availability:Public
 Title: Out of Bounds Pixels
 Crumbs: %ROOT%, Engine, Engine/Performance/Visualizers
 Description:Visualizing out of bounds pixels in the Unreal Editor viewport.
+version: 4.9
 
 [VAR:Topic]
 [OBJECT:Topic]
@@ -60,9 +61,7 @@ artifact will be flickering as you rotate the camera, since it will be incorrect
 
 ## Usage
 
-[REGION:raw]
-To enable this mode on any viewport, **Left-click** the **Show Button** ![](OOBP_button_show.png) on the viewport and select **Visualize -> Visualize Out of Bounds Pixels**.
-[/REGION]
+To enable this mode on any viewport, **Left-click** the **Show Button** on the viewport and select **Visualize -> Out of Bounds Pixels**.
 
 ![](OOBP_showMenu.png)
 
@@ -70,8 +69,8 @@ To enable this mode on any viewport, **Left-click** the **Show Button** ![](OOBP
 
 Blueprint components, often times, can be set to inherit the bounds of their parent, which can cause issues if the parent has smaller bounds than the child.
 
-The cube object below is a Blueprint composed of a **Box Component** (the red box in these images) and a **Static Mesh Component**. The Static Mesh is set to inherit its 
-bounds from the box. If the box extents are smaller than the Static Mesh, then Visualize Out of Bounds Pixels will show those portions in various colors (blue, yellow, 
+The cube object below is a Blueprint composed of a Box Component (the red box in these images) and a Static Mesh Component. The Static Mesh is set to inherit its 
+bounds from the box. If the box extents are smaller than the Static Mesh, then Out of Bounds Pixels will show those portions in various colors (blue, yellow, 
 and white).
 
 [REGION:lightbox]
@@ -96,13 +95,15 @@ bounding box:
 
 Because portions of this object will be visible after the bounding box is completely hidden or off camera, it will flicker or disappear at inappropriate times.
 
-In this case, BoundsScale under **Component -> Rendering -> Advanced** can be used to increase the bounding box of the mesh until it contains all the pixels.  Be careful 
+In this case, **Bounds Scale** under the expanded field for **Rendering** in a Component's **Details** panel can be used to increase the bounding box of the mesh until it contains all the pixels.  Be careful 
 doing this on shadow casters though, large shadow casters have a large performance cost.  On Skeletal Meshes, the bounds are derived from the physics asset, and individual 
 bones can be marked as contributing to the bounds or not.
 
+![](OOBP_BoundsScale.png)
+
 ## Skeletal Meshes
 
 The primary reason a Skeletal Mesh becomes displaced from its bounds is due to the animation causing the mesh to move away from its Actor root. However, a Skeletal Mesh with 
-a Physics Asset will use the Physics Asset for calculating its bounds. To enable this, **right-click** the Skeletal Mesh in the **Content Browser** and select 
+a Physics Asset will use the Physics Asset for calculating its bounds. To enable this, **Right-click** the Skeletal Mesh in the **Content Browser** and select 
 **Create Physics Asset**. Using the default settings will get you past any errors being caused by the Skeletal Mesh bounds, but you will probably want to edit the results. 
-See the [**Physics Asset Tool**](Engine/Physics\PhAT) documentation for more information on editing a **Physics Asset**.
+See the [Physics Asset Tool](Engine/Physics\PhAT) documentation for more information on editing a **Physics Asset**.
