#UE4doc #docs:
Edited a few steps to match the current version's workflow
Changed font settings for multiple

diff --git a/Engine/Documentation/Source/Engine/Physics/Vehicles/VehicleUserGuide/VehicleIntro.INT.udn b/Engine/Documentation/Source/Engine/Physics/Vehicles/VehicleUserGuide/VehicleIntro.INT.udn
--- a/Engine/Documentation/Source/Engine/Physics/Vehicles/VehicleUserGuide/VehicleIntro.INT.udn
+++ b/Engine/Documentation/Source/Engine/Physics/Vehicles/VehicleUserGuide/VehicleIntro.INT.udn
@@ -31,7 +31,7 @@ Navigation:topic
 This document will cover the creation of a vehicle using the **Wheeled Vehicle** Blueprint class, in a blank project. The process will take you through the creation of all the 
 necessary in-editor items, but assumes you have a **Skeletal Mesh** and a **Physics Asset** for your vehicle ready to go.
 
-If you need information on getting a vehicle Skeletal Mesh and Physics Asset into the engine, please see the [**Vehicle Content Guide**](Engine/Physics/Vehicles/VehcileContentCreation). 
+If you need information on getting a vehicle Skeletal Mesh and Physics Asset into the engine, please see the [Vehicle Content Guide](Engine/Physics/Vehicles/VehcileContentCreation). 
 Alternatively, you can work out of the **Vehicle Template** or **Vehicle Game** to have access to a prebuilt Skeletal Mesh and Physics Asset.
 
 
@@ -52,12 +52,13 @@ A vehicle in Unreal Engine 4 is composed of a number of assets:
 
 [REGION:note]
 This document will cover the Blueprint setup required to make a functioning vehicle in Unreal Engine 4, and assumes you already have a vehicle asset to work with. See the 
-[**Vehicle Content Guide**](Engine/Physics/Vehicles/VehcileContentCreation) for more information on setting up the Skeletal Mesh and Physics Asset.
+[
+](Engine/Physics/Vehicles/VehcileContentCreation) for more information on setting up the Skeletal Mesh and Physics Asset.
 [/REGION]
 
 ## Creating the Vehicle Blueprint
 
-To create a new **Vehicle Blueprint**: in the Content Browser **click** **Add New** button -> **Blueprint** -> search for "vehicle" and **click** **WheeledVehicle** from the list.
+To create a new **Vehicle Blueprint**: in the **Content Browser** click **Add New** button -> **Blueprint** -> search for "vehicle" and click **WheeledVehicle** from the list.
 
 [REGION:raw]
 ![](newVehicle.png)(convert:false)
@@ -65,7 +66,7 @@ To create a new **Vehicle Blueprint**: in the Content Browser **click** **Add Ne
 
 ## Creating the Wheel Blueprints
 
-To create a new **Wheel Blueprint**: in the Content Browser **click** **Add New** button -> **Blueprint** -> search for "wheel" and **click** **VehicleWheel** from the list. Repeat 
+To create a new **Wheel Blueprint**: in the **Content Browser** click **Add New** button -> **Blueprint** -> search for "wheel" and click **VehicleWheel** from the list. Repeat 
 this at least once so you have a front and rear wheel type.
 
 [REGION:raw]
@@ -78,7 +79,7 @@ suspension, and many other properties to give your vehicle the handling you desi
 
 ## Creating the TireType Data Asset
 
-To create a new **TireType Data Asset**: in the Content Browser **click** **Add New** button -> **Misc** -> **Data Asset** -> choose **TireType** -> **click** **Select** button.
+To create a new **TireType Data Asset**: in the **Content Browser** click **Add New** button -> **Misc** -> **Data Asset** -> choose **TireType** -> click the **Select** button.
 
 [REGION:raw]
 ![](newDataAsset.png)(convert:false)
@@ -93,8 +94,8 @@ There is a property slot on the **VehicleWheel** Blueprint type for the **TireTy
 
 ## Creating the Animation Blueprint
 
-To create a new **Animation Blueprint**: in the Content Browser **click** **Add New** button -> **Animation** -> **Animation Blueprint** -> choose the **Skeleton** from the list 
-that is for your vehicle  -> click the **Ok** button.
+To create a new **Animation Blueprint**: in the **Content Browser** click **Add New** button -> **Animation** -> **Animation Blueprint** -> choose the **Skeleton** from the list 
+that is for your vehicle -> choose **VehicleAnimInstance**  -> click the **Ok** button.
 
 [REGION:raw]
 ![](newAnimBP.png)(convert:false)
@@ -103,7 +104,7 @@ that is for your vehicle  -> click the **Ok** button.
 
 ## Editing the Animation Blueprint
 
-**Double-click** the **Animation Blueprint** from the Content Browser to edit the vehicle Animation Blueprint in Persona. The **Anim Graph** for a vehicle has been greatly 
+**Double-click** the **Animation Blueprint** from the **Content Browser** to edit the vehicle Animation Blueprint in Persona. The **Anim Graph** for a vehicle has been greatly 
 simplified for the 4.2 release:
 
 1. **Right-click** in the **Anim Graph**, from the context menu search for the **Mesh Space Ref Pose** node, and select it to create it.
@@ -146,7 +147,7 @@ Node, and the **Look At** nodes will make sure the suspension stays attached to
 
 ## Editing the Tires Data Asset
 
-**Double-click** the **TireType Data Asset** from the Content Browser.
+**Double-click** the **TireType Data Asset** from the **Content Browser**.
 
 ![](tires02.png)
 
@@ -155,7 +156,7 @@ Node, and the **Look At** nodes will make sure the suspension stays attached to
 
 ## Editing the Wheels Blueprints
 
-**Double-click** the **Wheel Blueprint** from the Content Browser to bring up the Blueprint Editor.
+**Double-click** the **Wheel Blueprint** from the **Content Browser** to bring up the Blueprint Editor.
 
 
 1. Switch the Blueprint Editor to Defaults Editing mode:
@@ -187,9 +188,9 @@ vehicle.
 
 This is really where the vehicle starts to come together. If your art is ready, and you have setup your wheels, then you will be able to fashion a working vehicle at this point.
 
-**Double-click** the **Vehicle Blueprint** from the Content Browser to bring up the Blueprint Editor.
+**Double-click** the **Vehicle Blueprint** from the **Content Browser** to bring up the Blueprint Editor.
 
-1. In **Components Mode**, on the **Skeletal Mesh Component** called "Mesh," set the **Skeletal Mesh** property to your vehicle's Skeletal Mesh asset.
+1. Click on the **Skeletal Mesh Component's** called "Mesh" and in its **Details** panel, set the **Skeletal Mesh** property to your vehicle's Skeletal Mesh asset.
 
 	![](componentsMode.png)
 
@@ -199,7 +200,7 @@ This is really where the vehicle starts to come together. If your art is ready,
 	![](AnimBPandSkelMesh.png)
 	[/REGION]
 
-1. Remain in **Components Mode** and add a **Camera Component**.
+1. Return to the **Components** panel and add a **Camera Component**.
 
 	[REGION:raw]
 	![](addCamera.png)
@@ -213,7 +214,7 @@ This is really where the vehicle starts to come together. If your art is ready,
 
 	![](cameraControllerOff.png)
 
-1. In **Defaults Mode**, in the **Vehicle Setup** category, **expand** "Wheel Setups."
+1. Select the **Vehicle Movement Component** and in the **Vehicle Setup** category, expand **Wheel Setups**.
 
 	![](defaultsMode.png)
 
@@ -273,10 +274,7 @@ There are a few ways to setup the vehicle for testing, and some steps will be de
 		![](newGameMode.png)
 		[/REGION]
 
-	1. **Double-Click** your new Game Mode Blueprint to edit it.
-	1. Switch to the **Defaults Mode**.
-
-		![](defaultsMode.png)
+	1. **Double-click** your new Game Mode Blueprint to edit it.
 
 	1. **Game Mode** Category -> **Default Pawn Class** -> **change** "None" to your **Vehicle Blueprint**.
 
@@ -291,7 +289,7 @@ There are a few ways to setup the vehicle for testing, and some steps will be de
 	[/REGION]
 
 1. In your vehicle's Blueprint, you can then call on the Axis and Action events you created in the **Project Settings Input** category.
-	1. **Edit** your Vehicle Blueprint, and **switch** to **Graph Mode**
+	1. **Edit** your Vehicle Blueprint, and open the **Event Graph**.
 
 		![](graphMode.png)
 
@@ -308,8 +306,8 @@ There are a few ways to setup the vehicle for testing, and some steps will be de
 		![](handbrakeBP.png)
 
 1. Player Start
-	* If you do not have a **Player Start** Actor in the world, you can either add one to your map using the **Modes Panel** -> **Placement Mode** -> **Basic** -> then **clicking and 
-	dragging** a **Player Start** into the world.
+	* If you do not have a **Player Start** Actor in the world, you can either add one to your map using the **Modes Panel** -> **Placement Mode** -> **Basic** -> then clicking and 
+	dragging a **Player Start** into the world.
 
 	[REGION:raw]
 	![](playerStart.png)
@@ -322,8 +320,8 @@ There are a few ways to setup the vehicle for testing, and some steps will be de
 	![](Play.png)
 	[/REGION]
 
-	Given the Axis and Action Bindings you setup above, "**W**" will increase your throttle, "**S**" will apply the brakes, "**A**" will turn left, "**D**" will turn right, and 
-	"**Space Bar**" will apply the Handbrake.
+	Given the Axis and Action Bindings you setup above, **W** will increase your throttle, **S** will apply the brakes, **A** will turn left, **D** will turn right, and 
+	**Space Bar** will apply the Handbrake.
 
 
 
