#UE4doc #docs - Updated images and text to reflect changes in UI. Addressed UEDOC-2480

diff --git a/Engine/Documentation/Source/Engine/Rendering/PostProcessEffects/PostProcessMaterials/PostProcessMaterials.INT.udn b/Engine/Documentation/Source/Engine/Rendering/PostProcessEffects/PostProcessMaterials/PostProcessMaterials.INT.udn
--- a/Engine/Documentation/Source/Engine/Rendering/PostProcessEffects/PostProcessMaterials/PostProcessMaterials.INT.udn
+++ b/Engine/Documentation/Source/Engine/Rendering/PostProcessEffects/PostProcessMaterials/PostProcessMaterials.INT.udn
@@ -34,10 +34,10 @@ Through post process settings (usually defined with post process volumes or came
 At the moment, only **materials** and **material instances** are blendable assets. The engine provides a few post process materials but with this feature
 you can create your own **custom post processing** without any programmer help.
 
-Simply assign one or more post process materials to a post process volume in the **Misc** section to the **Blendables**. First press the + to add
+Simply assign one or more post process materials to a post process volume in the **Blendables** section. First press the + to add
 new slots, select a material in the **Content Browser**, and press the left arrow to assign it. The order here is not important and unused slots are simply ignored.
 
-![](PostProcessSettings.gif)
+![](PostProcessSettings.png)
 
 ## Making a Simple Post Process Material
 
