#UE4doc #docs - Unbolded hyperlink

diff --git a/Engine/Documentation/Source/Programming/UnrealArchitecture/Reference/Properties/Properties.INT.udn b/Engine/Documentation/Source/Programming/UnrealArchitecture/Reference/Properties/Properties.INT.udn
--- a/Engine/Documentation/Source/Programming/UnrealArchitecture/Reference/Properties/Properties.INT.udn
+++ b/Engine/Documentation/Source/Programming/UnrealArchitecture/Reference/Properties/Properties.INT.udn
@@ -46,7 +46,7 @@ For most uses, Unreal relies on the TCHAR type for characters.  The TEXT() macro
 
 	MyDogPtr->DogName = FName(TEXT("Samson Aloysius"));
 
-For more on the three string types, when to use each one, and how to work with them, see the **[String Handling documentation](Programming/UnrealArchitecture/StringHandling)**.
+For more on the three string types, when to use each one, and how to work with them, see the [String Handling documentation](Programming/UnrealArchitecture/StringHandling).
 
 ## Property Specifiers
 
