#DocFix

Forum request to update documentation (master_rigel)

diff --git a/Engine/Documentation/Source/Engine/UMG/QuickStart/2/UMGQS_2.INT.udn b/Engine/Documentation/Source/Engine/UMG/QuickStart/2/UMGQS_2.INT.udn
--- a/Engine/Documentation/Source/Engine/UMG/QuickStart/2/UMGQS_2.INT.udn
+++ b/Engine/Documentation/Source/Engine/UMG/QuickStart/2/UMGQS_2.INT.udn
@@ -114,7 +114,9 @@ Now that we have our visual layout, we can add the hooks to provide the function
 
 	![](Step2_20.png)
 
-1.	In the **Graph** under the **Event Construct**, **Right-click** and add a **Get Player Character** node then **Cast To FirstPersonCharacter**. 
+1.	In the **Graph** under the **Event Construct** node, **Right-click** and add a **Get Player Character** node. 
+
+1.  Drag off the **Return Value** pin and select **Cast to FirstPersonCharacter**
 
 	![](Step2_21.png)
 
@@ -126,6 +128,8 @@ Now that we have our visual layout, we can add the hooks to provide the function
 
 	This will allow us to access the variables contained inside our **FirstPersonCharacter** Blueprint. 
 
+1.  Click **Compile** in the toolbar to compile the script.
+
 1.	Return to the **Designer** tab and select the **Progress Bar** next to **Health**. 
 
 1.	In the **Details** panel under **Appearance**, click the **Bind** option next to **Percent** and set it to use **Health** from **MyCharacter**. 
@@ -148,7 +152,9 @@ Now that we have our visual layout, we can add the hooks to provide the function
 
 	![](Step2_27.png)(w:820)
 
-1.	Off **MyCharacter**, get **Ammo** and connect it to the **Return Value** of the **Return Node**. 
+1.	Drag off the pin **MyCharacter** and select get **Ammo**. 
+
+1.  Connect the Ammo pin to the **Return Value** of the **Return Node**. 
 
 	![](Step2_28.png)
 
