#UE4doc #diocs - Added engine version field.

diff --git a/Engine/Documentation/Source/Programming/UnrealArchitecture/Reference/Properties/Specifiers/PropertySpecifiers.INT.udn b/Engine/Documentation/Source/Programming/UnrealArchitecture/Reference/Properties/Specifiers/PropertySpecifiers.INT.udn
--- a/Engine/Documentation/Source/Programming/UnrealArchitecture/Reference/Properties/Specifiers/PropertySpecifiers.INT.udn
+++ b/Engine/Documentation/Source/Programming/UnrealArchitecture/Reference/Properties/Specifiers/PropertySpecifiers.INT.udn
@@ -2,6 +2,7 @@ Title: Property Specifiers
 Crumbs:%ROOT%, Programming/UnrealArchitecture/Reference,  Programming/UnrealArchitecture/Reference, Programming/UnrealArchitecture/Reference/Properties
 Description: Keywords used when declaring gameplay properties to specify how the property behaves with various aspects of the engine and editor.
 Availability:Public
+Version: 4.9
 
 [EXCERPT:main]
 When declaring properties, specifiers can be added to the declaration to control how the property behaves with various 
