#UE4doc #docs - Added Engine Version Field. Bolded term.

diff --git a/Engine/Documentation/Source/Programming/UnrealArchitecture/Reference/Properties/Specifiers/AssetRegistrySearchable/PropertySpecifier_AssetRegistrySearchable.INT.udn b/Engine/Documentation/Source/Programming/UnrealArchitecture/Reference/Properties/Specifiers/AssetRegistrySearchable/PropertySpecifier_AssetRegistrySearchable.INT.udn
--- a/Engine/Documentation/Source/Programming/UnrealArchitecture/Reference/Properties/Specifiers/AssetRegistrySearchable/PropertySpecifier_AssetRegistrySearchable.INT.udn
+++ b/Engine/Documentation/Source/Programming/UnrealArchitecture/Reference/Properties/Specifiers/AssetRegistrySearchable/PropertySpecifier_AssetRegistrySearchable.INT.udn
@@ -4,5 +4,5 @@ Description:
 Availability:Public
 Version: 4.9
 
-The AssetRegistrySearchable keyword indicates that this property and its value will be automatically added to the asset registry for any asset class instances containing this as a member variable.  It is not legal
+The **AssetRegistrySearchable** keyword indicates that this property and its value will be automatically added to the asset registry for any asset class instances containing this as a member variable.  It is not legal
 to use on struct properties or parameters.
