#UE4doc #docs - Bolded editor terms

diff --git a/Engine/Documentation/Source/Programming/UnrealArchitecture/Reference/Functions/Specifiers/BlueprintCallable/FunctionSpecifier_BlueprintCallable.INT.udn b/Engine/Documentation/Source/Programming/UnrealArchitecture/Reference/Functions/Specifiers/BlueprintCallable/FunctionSpecifier_BlueprintCallable.INT.udn
--- a/Engine/Documentation/Source/Programming/UnrealArchitecture/Reference/Functions/Specifiers/BlueprintCallable/FunctionSpecifier_BlueprintCallable.INT.udn
+++ b/Engine/Documentation/Source/Programming/UnrealArchitecture/Reference/Functions/Specifiers/BlueprintCallable/FunctionSpecifier_BlueprintCallable.INT.udn
@@ -4,4 +4,4 @@ Description:
 Availability:Public
 Version: 4.9
 
-The function can be executed in a Blueprint or Level Blueprint graph.
+The function can be executed in a **Blueprint** or **Level Blueprint** graph.
