#UE4doc #docs - Bolded editor terms

diff --git a/Engine/Documentation/Source/Programming/UnrealArchitecture/Reference/Functions/Specifiers/BlueprintNativeEvent/FunctionSpecifier_BlueprintNativeEvent.INT.udn b/Engine/Documentation/Source/Programming/UnrealArchitecture/Reference/Functions/Specifiers/BlueprintNativeEvent/FunctionSpecifier_BlueprintNativeEvent.INT.udn
--- a/Engine/Documentation/Source/Programming/UnrealArchitecture/Reference/Functions/Specifiers/BlueprintNativeEvent/FunctionSpecifier_BlueprintNativeEvent.INT.udn
+++ b/Engine/Documentation/Source/Programming/UnrealArchitecture/Reference/Functions/Specifiers/BlueprintNativeEvent/FunctionSpecifier_BlueprintNativeEvent.INT.udn
@@ -4,4 +4,4 @@ Description:
 Availability:Public
 Version: 4.9
 
-This function is designed to be overridden by a Blueprint, but also has a native implementation. Provide a body named [FunctionName]_Implementation instead of [FunctionName]; the autogenerated code will include a thunk that calls the implementation method when necessary.
+This function is designed to be overridden by a **Blueprint**, but also has a native implementation. Provide a body named [FunctionName]_Implementation instead of [FunctionName]; the autogenerated code will include a thunk that calls the implementation method when necessary.
diff --git a/Engine/Documentation/Source/Programming/UnrealArchitecture/Reference/Functions/Specifiers/BlueprintPure/FunctionSpecifier_BlueprintPure.INT.udn b/Engine/Documentation/Source/Programming/UnrealArchitecture/Reference/Functions/Specifiers/BlueprintPure/FunctionSpecifier_BlueprintPure.INT.udn
--- a/Engine/Documentation/Source/Programming/UnrealArchitecture/Reference/Functions/Specifiers/BlueprintPure/FunctionSpecifier_BlueprintPure.INT.udn
+++ b/Engine/Documentation/Source/Programming/UnrealArchitecture/Reference/Functions/Specifiers/BlueprintPure/FunctionSpecifier_BlueprintPure.INT.udn
@@ -4,4 +4,4 @@ Description:
 Availability:Public
 Version: 4.9
 
-The function does not affect the owning object in any way and can be executed in a Blueprint or Level Blueprint graph. 
+The function does not affect the owning object in any way and can be executed in a **Blueprint** or **Level Blueprint** graph. 
