#UE4doc #docs - Bolded editor terms

diff --git a/Engine/Documentation/Source/Programming/UnrealArchitecture/Reference/Functions/Specifiers/CustomThunk/FunctionSpecifier_CustomThunk.INT.udn b/Engine/Documentation/Source/Programming/UnrealArchitecture/Reference/Functions/Specifiers/CustomThunk/FunctionSpecifier_CustomThunk.INT.udn
--- a/Engine/Documentation/Source/Programming/UnrealArchitecture/Reference/Functions/Specifiers/CustomThunk/FunctionSpecifier_CustomThunk.INT.udn
+++ b/Engine/Documentation/Source/Programming/UnrealArchitecture/Reference/Functions/Specifiers/CustomThunk/FunctionSpecifier_CustomThunk.INT.udn
@@ -4,4 +4,4 @@ Description:
 Availability:Public
 Version: 4.9
 
-The UnrealHeaderTool code generator will not produce a execFoo thunk for this function; it is up to the user to provide one.
+The **UnrealHeaderTool** code generator will not produce a execFoo thunk for this function; it is up to the user to provide one.
