#UE4doc #docs - Added engine version field.

diff --git a/Engine/Documentation/Source/Programming/UnrealArchitecture/Reference/Metadata/MetadataSpecifiers.INT.udn b/Engine/Documentation/Source/Programming/UnrealArchitecture/Reference/Metadata/MetadataSpecifiers.INT.udn
--- a/Engine/Documentation/Source/Programming/UnrealArchitecture/Reference/Metadata/MetadataSpecifiers.INT.udn
+++ b/Engine/Documentation/Source/Programming/UnrealArchitecture/Reference/Metadata/MetadataSpecifiers.INT.udn
@@ -2,6 +2,7 @@ Title: Metadata Specifiers
 Crumbs:%ROOT%, Programming, Programming/UnrealArchitecture/Reference
 Description: Metadata keywords used when declaring classes, functions, and interfaces to specify how they behave with various aspects of the engine and editor.
 Availability:Public
+Version: 4.9
 
 When declaring classes, functions, and interfaces, metadata specifiers can be added to the declaration to control how they behave with various
 aspects of the engine and editor.
