#UE4doc #docs - Capitalized term.

diff --git a/Engine/Documentation/Source/Programming/UnrealArchitecture/Reference/Metadata/BlueprintProtected/FunctionMetadataSpecifier_BlueprintProtected.INT.udn b/Engine/Documentation/Source/Programming/UnrealArchitecture/Reference/Metadata/BlueprintProtected/FunctionMetadataSpecifier_BlueprintProtected.INT.udn
--- a/Engine/Documentation/Source/Programming/UnrealArchitecture/Reference/Metadata/BlueprintProtected/FunctionMetadataSpecifier_BlueprintProtected.INT.udn
+++ b/Engine/Documentation/Source/Programming/UnrealArchitecture/Reference/Metadata/BlueprintProtected/FunctionMetadataSpecifier_BlueprintProtected.INT.udn
@@ -4,7 +4,7 @@ Description:
 Availability:Public
 Version: 4.9
 
-This function can only be called on 'this' in a blueprint. It cannot be called on another instance.
+This function can only be called on 'this' in a **Blueprint**. It cannot be called on another instance.
 
 ## Usage
 
