#UE4doc #doc - Bolded term.

diff --git a/Engine/Documentation/Source/Programming/UnrealArchitecture/Reference/Metadata/CannotImplementInterfaceBlueprint/InterfaceMetadataSpecifier_CannotImplementInterfaceInBlueprint.INT.udn b/Engine/Documentation/Source/Programming/UnrealArchitecture/Reference/Metadata/CannotImplementInterfaceBlueprint/InterfaceMetadataSpecifier_CannotImplementInterfaceInBlueprint.INT.udn
--- a/Engine/Documentation/Source/Programming/UnrealArchitecture/Reference/Metadata/CannotImplementInterfaceBlueprint/InterfaceMetadataSpecifier_CannotImplementInterfaceInBlueprint.INT.udn
+++ b/Engine/Documentation/Source/Programming/UnrealArchitecture/Reference/Metadata/CannotImplementInterfaceBlueprint/InterfaceMetadataSpecifier_CannotImplementInterfaceInBlueprint.INT.udn
@@ -4,7 +4,7 @@ Description:
 Availability:Public
 Version: 4.9
 
-This interface cannot be implemented by a Blueprint (e.g., it has only non-exposed C++ member methods).
+This interface cannot be implemented by a **Blueprint** (e.g., it has only non-exposed C++ member methods).
 
 Any interfaces that define a function that isn't a [BlueprintNativeEvent](Programming/UnrealArchitecture/Reference/Functions/Specifiers/BlueprintImplementableEvent) or [BlueprintImplementableEvent](Programming/UnrealArchitecture/Reference/Functions/Specifiers/BlueprintNativeEvent) must be declared as `CannotImplementInterfaceInBlueprint`.
 
