#UE4doc #docs - Added Engine version field.

diff --git a/Engine/Documentation/Source/Engine/Rendering/Materials/Functions/Reference/Blends/Blends.INT.udn b/Engine/Documentation/Source/Engine/Rendering/Materials/Functions/Reference/Blends/Blends.INT.udn
--- a/Engine/Documentation/Source/Engine/Rendering/Materials/Functions/Reference/Blends/Blends.INT.udn
+++ b/Engine/Documentation/Source/Engine/Rendering/Materials/Functions/Reference/Blends/Blends.INT.udn
@@ -2,6 +2,7 @@ Availability: Public
 title:Blends
 Crumbs: %ROOT%, Engine, Engine/Rendering/Materials,  Engine/Rendering/Materials/Functions
 Description:Functions designed to blend one color with another, similar to blend modes in popular image editing applications.
+Version: 4.9
 
 [TOC (start:2 end:3)]
 
