#UE4doc #docs - Added Engine Version field.

diff --git a/Engine/Documentation/Source/Engine/UI/LevelEditor/Layers/Layers.INT.udn b/Engine/Documentation/Source/Engine/UI/LevelEditor/Layers/Layers.INT.udn
--- a/Engine/Documentation/Source/Engine/UI/LevelEditor/Layers/Layers.INT.udn
+++ b/Engine/Documentation/Source/Engine/UI/LevelEditor/Layers/Layers.INT.udn
@@ -2,6 +2,9 @@ Availability: Public
 TItle:Layers Panel User Guide
 Crumbs: %ROOT%, Engine, Engine/UI, Engine/UI/LevelEditor
 Description:Tool for organizing Actors in the level into layers for quick selection and controlling visibility.
+Version: 4.9
+SkillLevel: Beginner
+
 
 [TOC(start:2)]
 
diff --git a/Engine/Documentation/Source/Gameplay/Framework/GameplayFramework.INT.udn b/Engine/Documentation/Source/Gameplay/Framework/GameplayFramework.INT.udn
--- a/Engine/Documentation/Source/Gameplay/Framework/GameplayFramework.INT.udn
+++ b/Engine/Documentation/Source/Gameplay/Framework/GameplayFramework.INT.udn
@@ -2,6 +2,7 @@ Availability:Public
 Title:Gameplay Framework
 Crumbs:%ROOT%, Gameplay
 Description:Core systems, such as game rules, player input and controls, cameras, and user interfaces.
+Version: 4.9
 
 To show how the different pieces of game framework interlock, here is a simple example. Imagine a simple concept for a game, 
 where a rabbit races a snail.  The base of the game framework is the GameMode.  The **GameMode** sets the rules for the game,
