#UE4doc #docs - Added engine version fields. Added Skill Level fields.

diff --git a/Engine/Documentation/Source/Gameplay/HowTo/SpawnAndDestroyActors/Blueprints/GHT_1B.INT.udn b/Engine/Documentation/Source/Gameplay/HowTo/SpawnAndDestroyActors/Blueprints/GHT_1B.INT.udn
--- a/Engine/Documentation/Source/Gameplay/HowTo/SpawnAndDestroyActors/Blueprints/GHT_1B.INT.udn
+++ b/Engine/Documentation/Source/Gameplay/HowTo/SpawnAndDestroyActors/Blueprints/GHT_1B.INT.udn
@@ -1,7 +1,9 @@
 Availability: Public
 Title:Spawning/Destroying an Actor in Blueprints
 Crumbs: %ROOT%, Gameplay, Gameplay/HowTo/SpawnAndDestroyActors
-Description:A How To Guide for Spawning/Destroying Actors through Blueprints in Unreal Engine 4. 
+Description:A How To Guide for Spawning/Destroying Actors through Blueprints in Unreal Engine 4.
+Version: 4.9
+SkillLevel:Intermediate
 
 [TOC (start:2)]
 
