#UE4doc #docs:
Added engine version field
Unbolded a hyperlink

diff --git a/Engine/Documentation/Source/Engine/UMG/HowTo/UMGHowTo.INT.udn b/Engine/Documentation/Source/Engine/UMG/HowTo/UMGHowTo.INT.udn
--- a/Engine/Documentation/Source/Engine/UMG/HowTo/UMGHowTo.INT.udn
+++ b/Engine/Documentation/Source/Engine/UMG/HowTo/UMGHowTo.INT.udn
@@ -4,15 +4,16 @@ Crumbs: %ROOT%, Engine, Engine\UMG
 Description:A How To Guide for creating UI elements with Unreal Motion Graphics UI Designer (UMG). 
 Related: Engine\UMG
 Related: Engine\UMG\UserGuide
+version: 4.9
+
 
-## Overview
 The **Unreal Motion Graphics UI Designer (UMG) - How To...** pages will show you how to create various UI elements using UMG. Each example comes with step-by-step instructions and several images to help guide you through the process of creating the described content. Reference links are also included to provide more detailed documentation of the systems and content being outlined in each guide. 
 
 Refer to the table below for more information and click each link to view its contents. 
 
 |How To... | Description |
 | ------- | ------------ |
-| **[](Engine\UMG\HowTo\Create3DWidgets)** | %Engine\UMG\HowTo\Create3DWidgets:description% |
+| [](Engine\UMG\HowTo\Create3DWidgets) | %Engine\UMG\HowTo\Create3DWidgets:description% |
 
 
 
