#UE4doc #docs:
Replaced multiple images to show updated UI
Replaced mentions of Control with Ctrl
Edited a few steps and removed a few as they were based off old UI

diff --git a/Engine/Documentation/Source/Engine/UMG/HowTo/Create3DWidgets/Create3DWidgets.INT.udn b/Engine/Documentation/Source/Engine/UMG/HowTo/Create3DWidgets/Create3DWidgets.INT.udn
--- a/Engine/Documentation/Source/Engine/UMG/HowTo/Create3DWidgets/Create3DWidgets.INT.udn
+++ b/Engine/Documentation/Source/Engine/UMG/HowTo/Create3DWidgets/Create3DWidgets.INT.udn
@@ -101,13 +101,13 @@ In this step, we will create the Main Menu that we want to display in the game w
 
 	This will set the button to appear one color when **Normal** and another color when the button is **Hovered** over with the mouse. 
 	
-1.	In the **Hierarchy** window, click the **Text Block** and press **Control+C** to copy it, then click the **Vertical Box** and **Control+V** to paste it. 
+1.	In the **Hierarchy** window, click the **Text Block** and press **Ctrl+C** to copy it, then click the **Vertical Box** and **Ctrl+V** to paste it. 
 
 1.	Repeat the previous step 2 more times so that there are 4 **Text Blocks**. 
 
 	![](UMG_3D_9.png)
 
-1.	In the **Hierarchy** window, click the **Button** and press **Control+C** to copy it, then click the **Vertical Box** and **Control+V** to paste it. 
+1.	In the **Hierarchy** window, click the **Button** and press **Ctrl+C** to copy it, then click the **Vertical Box** and **Ctrl+V** to paste it. 
 
 1.	Repeat the previous step so that there are 3 **Buttons**. 
 
@@ -117,21 +117,13 @@ In this step, we will create the Main Menu that we want to display in the game w
 
 	![](UMG_3D_11.png)
 
-1.	For each **Button** widget in the **Hierarchy**, in the **Details** panel for each, set the **Size** to **Fill**. 
-
-	![](UMG_3D_12.png)
-
-1.	For each **Text Block** widget in the **Hierarchy**, in the **Details** panel for each, set them to **Center** horizontally and vertically. 
-
-	![](UMG_3D_13.png)
-
 1.	Your layout should look similar to below. 
 
 	![](UMG_3D_14.png)
 
 1.	In the **Hierarchy** window, click the top **Text Block** then in the **Details** panel under **Content**, set the **Text** to anything you wish. 
 
-1.	Also for the **Text Block** in the **Details** panel under **Appearance**, you can change the size of the **Font** under the **Font** option. 
+1.	Also for the **Text Block** in the **Details** panel under **Appearance**, you can change the size of the Font under the **Font** option. 
 
 	![](UMG_3D_15.png)
 
@@ -193,7 +185,7 @@ In this step, we will create the Actor that will contain and display our UI Menu
 
 	![](UMG_3D_24.png)
 
-1.	Click the **Graph** tab to jump to the **Graph Editor** tab. 
+1.	Click the **Event Graph** tab. 
 
 	![](UMG_3D_25.png)
 
