#UE4doc #docs:
Added engine version field

diff --git a/Engine/Documentation/Source/Engine/UMG/QuickStart/UMGQuickStart.INT.udn b/Engine/Documentation/Source/Engine/UMG/QuickStart/UMGQuickStart.INT.udn
--- a/Engine/Documentation/Source/Engine/UMG/QuickStart/UMGQuickStart.INT.udn
+++ b/Engine/Documentation/Source/Engine/UMG/QuickStart/UMGQuickStart.INT.udn
@@ -2,6 +2,7 @@ Availability: Public
 Title:UMG UI Designer Quick Start Guide
 Crumbs:%ROOT%, Engine, Engine/UMG
 Description:Getting started with using Unreal Motion Graphics in Unreal Engine 4.
+version: 4.9
 
 [VAR:Topic]
 [OBJECT:Topic]
