#UE4doc #docs:
Added engine version field
Changed font settings for a few terms
Replaced multiple images due to the UI being outdated

diff --git a/Engine/Documentation/Source/Engine/UMG/QuickStart/1/UMGQS_1.INT.udn b/Engine/Documentation/Source/Engine/UMG/QuickStart/1/UMGQS_1.INT.udn
--- a/Engine/Documentation/Source/Engine/UMG/QuickStart/1/UMGQS_1.INT.udn
+++ b/Engine/Documentation/Source/Engine/UMG/QuickStart/1/UMGQS_1.INT.udn
@@ -2,7 +2,7 @@ Availability: Public
 Title:1. Required Project Setup
 Crumbs:%ROOT%, Engine, Engine/UMG, Engine/UMG/QuickStart
 Description:Here we will create a Widget Blueprint used for the HUD and set our character up to pass information to the HUD. 
-
+version: 4.9
 
 [VAR:Steps]
 [OBJECT:Navigation]
@@ -25,7 +25,7 @@ Description:Here we will create a Widget Blueprint used for the HUD and set our
 For this tutorial we are using the **Blueprint First Person** template with **Starter Content** enabled. If you do not know how to create a new project and use a template, please see the [](Engine/Basics/Projects/Browser) page for more information. 
 [/REGION]
 
-1. Click the **Add New** button inside the **Content Browser**, then under _User Interface_ select **Widget Blueprint** and name it **HUD**.
+1. Click the **Add New** button inside the **Content Browser**, then under **User Interface** select **Widget Blueprint** and name it **HUD**.
 
 	![](UMGQS1.png)
 
