#UE4doc #docs:
Added engine version field
Rewrote a few steps to be accurate with the most recent UI

diff --git a/Engine/Documentation/Source/Engine/UMG/QuickStart/2/UMGQS_2.INT.udn b/Engine/Documentation/Source/Engine/UMG/QuickStart/2/UMGQS_2.INT.udn
--- a/Engine/Documentation/Source/Engine/UMG/QuickStart/2/UMGQS_2.INT.udn
+++ b/Engine/Documentation/Source/Engine/UMG/QuickStart/2/UMGQS_2.INT.udn
@@ -2,6 +2,7 @@ Availability: Public
 Title:2. Displaying Health, Energy and Ammo
 Crumbs:%ROOT%, Engine, Engine/UMG, Engine/UMG/QuickStart
 Description:In this step we will hook up and display the Health, Energy and Ammo of the player on our HUD. 
+version: 4.9
 
 
 [VAR:Steps]
@@ -132,7 +133,7 @@ Now that we have our visual layout, we can add the hooks to provide the function
 
 1.	Return to the **Designer** tab and select the **Progress Bar** next to **Health**. 
 
-1.	In the **Details** panel under **Appearance**, click the **Bind** option next to **Percent** and set it to use **Health** from **MyCharacter**. 
+1.	In the **Details** panel under **Progress**, click the **Bind** option next to **Percent** and set it to use **Health** from **MyCharacter**. 
 
 	![](Step2_23.png)
 
