#UE4doc #docs:
Added engine version field
Changed font settings for a couple terms

diff --git a/Engine/Documentation/Source/Engine/UMG/QuickStart/3/UMGQS_3.INT.udn b/Engine/Documentation/Source/Engine/UMG/QuickStart/3/UMGQS_3.INT.udn
--- a/Engine/Documentation/Source/Engine/UMG/QuickStart/3/UMGQS_3.INT.udn
+++ b/Engine/Documentation/Source/Engine/UMG/QuickStart/3/UMGQS_3.INT.udn
@@ -1,7 +1,8 @@
 Availability: Public
 Title:3. Creating a Main Menu
 Crumbs:%ROOT%, Engine, Engine/UMG, Engine/UMG/QuickStart
-Description:In this step we will create the visual layout for our Main Menu. 
+Description:In this step we will create the visual layout for our Main Menu.
+version: 4.9
 
 
 [VAR:Steps]
@@ -56,7 +57,7 @@ First bit of business is to create the layout of our Main Menu which we will do
 
 	By updating the names of our buttons so that we know what each one does, this will make it easier to script functionality for them. 
 
-1.	Update each of the **Text** widgets as shown below via the **Text** section under _Content_ in the **Details** panel. 
+1.	Update each of the **Text** widgets as shown below via the **Text** section under **Content** in the **Details** panel. 
 
 	![](Hierarchy6.png)
 
@@ -68,7 +69,7 @@ First bit of business is to create the layout of our Main Menu which we will do
 
 	Here we are adjusting the size of all of our buttons as well as defining a color for the button to take when mousing over it. 
 
-1.	Select the **Image** in the **Hierarchy** then in the **Details** under _Appearance_ and _Brush_, select a **Texture** or **Material** to use. 
+1.	Select the **Image** in the **Hierarchy** then in the **Details** under **Appearance** and **Brush**, select a **Texture** or **Material** to use. 
 
 	![](Hierarchy8.png)
 
