#UE4 Docs: Updating the docs to reflect the newest oculus ue4 updates.

diff --git a/Engine/Documentation/Source/Platforms/Oculus/QuickStart/1/OculusQuickStart_1.INT.udn b/Engine/Documentation/Source/Platforms/Oculus/QuickStart/1/OculusQuickStart_1.INT.udn
--- a/Engine/Documentation/Source/Platforms/Oculus/QuickStart/1/OculusQuickStart_1.INT.udn
+++ b/Engine/Documentation/Source/Platforms/Oculus/QuickStart/1/OculusQuickStart_1.INT.udn
@@ -2,7 +2,7 @@ Title:1. Register as an Oculus Developer
 Crumbs:%ROOT%, Platforms, Platforms/Oculus, Platforms/Oculus/QuickStart
 Description:Register as an Oculus developer to gain access to necessary programs and documentation
 Availability:Public
-
+Version: 4.9
 
 [VAR:Steps]
 [OBJECT:Navigation]
@@ -19,23 +19,28 @@ Availability:Public
 [/OBJECT]
 [/VAR]
 
-%Steps%
-
-1. Visit the [Oculus Developer Center](https://developer.oculus.com/) and click on [Downloads](https://developer.oculus.com/downloads/). 
-
-	![](Oculus_Downloads.png)
+In this section we will go over how to register as an Oculus developer if you are not already one. 
+While it is not mandatory to register as an Oculus developer, registering will allow you to access to the Oculus forums as well as submit games to the Oculus store.
 
-From the [Oculus Developer Center](https://developer.oculus.com/) you can also click the Log In link to create an account, or access your existing account. Creating an account
-enables you to post in the forums, submit your game to Oculus Share, and more.
-
-* If you have an Oculus account, enter your e-mail and password and click the **Log In** button.
-
-	![](2-2-click-the-log-in-link.jpg)
+%Steps%
 
-* If you do not have an Oculus account, click **Create a free account**.
-	1. Fill out the registration form following the instructions on the page.
-	1. Oculus will e-mail you a link to verify your e-mail address and confirm your registration.
-	1. Return to the [Oculus Developer Center](https://developer.oculus.com/) and click the **Log In** link.
+Visit the [Oculus Developer Center](https://developer.oculus.com/) and click on person icon in the upper right hand corner of the page to log in if you already have an Oculus account. 
+Please note that if you already have an account you can skip ahead to the next step.
+	
+![](Oculus_Downloads.png)
+
+[region:note]
+While you do not need an Oculus account do to download the required Oculus runtime software, creating an account enables you to post in the forums, submit your game to Oculus Share, and more.
+[/region]
+
+If you do not have an Oculus account you can make one by clicking on the person icon in the upper right hand corner of the page to get to the [login page](https://id.oculus.com/US/en-US/login/). Once on the page you can do the following to setup an account.
+	
+ ![](2-2-click-the-log-in-link.png)
+	
+1. Click on the **Create a free account** option.
+1. Fill out the registration form following the instructions on the page.
+1. Oculus will e-mail you a link to verify your e-mail address and confirm your registration.
+1. Return to the [Oculus Developer Center](https://developer.oculus.com/) and click the **Log In** link.
 
 
 %Steps%
diff --git a/Engine/Documentation/Source/Platforms/Oculus/QuickStart/2/OculusQuickStart_2.INT.udn b/Engine/Documentation/Source/Platforms/Oculus/QuickStart/2/OculusQuickStart_2.INT.udn
--- a/Engine/Documentation/Source/Platforms/Oculus/QuickStart/2/OculusQuickStart_2.INT.udn
+++ b/Engine/Documentation/Source/Platforms/Oculus/QuickStart/2/OculusQuickStart_2.INT.udn
@@ -2,7 +2,7 @@ Title:2. Download the Oculus Runtime
 Crumbs:%ROOT%, Platforms, Platforms/Oculus, Platforms/Oculus/QuickStart
 Description:Download the needed software for running games on the Oculus Rift
 Availability:Public
-
+Version: 4.9
 
 [VAR:Steps]
 [OBJECT:Navigation]
@@ -20,59 +20,29 @@ Availability:Public
 [/OBJECT]
 [/VAR]
 
+In this section we will go over which version of the Oculus PC Runtime drivers you need to download and were you can download them form.
+
 %Steps%
 
 
 [REGION:note]
-The version of the Oculus Runtime that is currently compatible with Unreal Engine 4 is **Oculus SDK v0.4.2 Beta**. 
+The version of the Oculus Runtime that is currently compatible with the latest version of Unreal Engine 4 is **Oculus SDK v0.7.0.0-BETA**. 
 [/REGION]
 
-%Globals:OSSelect%
 
-[OBJECT:ToggleButtonContent]
-	[PARAMLITERAL:category]
-	OS
-	[/PARAMLITERAL]
-	[PARAMLITERAL:id]
-	windows
-	[/PARAMLITERAL]
-	[PARAMLITERAL:active]
-	active_button_content
-	[/PARAMLITERAL]
-	[PARAM:content]
-	[INCLUDE:#windows]
-	[/PARAM]
-[/OBJECT]
-[OBJECT:ToggleButtonContent]
-	[PARAMLITERAL:category]
-	OS
-	[/PARAMLITERAL]
-	[PARAMLITERAL:id]
-	mac
-	[/PARAMLITERAL]
-	[PARAMLITERAL:active]
-	[/PARAMLITERAL]
-	[PARAM:content]
-	[INCLUDE:#macos]
-	[/PARAM]
-[/OBJECT]
-<!--
-[EXCERPT:windows]
-1. Click the [Oculus Runtime for Windows version 0.4.2-beta](https://developer.oculus.com/downloads/#id=9f27e6a8). 
+1. From the web browser of your choice got the Oculus Developer webpage by clicking on the link below.
+	
+	[https://developer.oculus.com/](https://developer.oculus.com/) 
+
+1. Download the [PC RUNTIME( v0.7.0.0-BETA)](https://developer.oculus.com/downloads/pc/0.7.0.0-beta/Oculus_Runtime_for_Windows/) by clicking on the PC RUNTIME(V0.7.0-BETA) button.
+
+	![](Downloads.png)
 
-	![](3-oculus-runtime-for-windows.jpg)
-1. Read and agree to the terms of the EULA to start the download.
+1. Once you have clicked on the download button you will be brought to the webpage below. 
+ Make sure to read the EULA and Terms of Use and if you agree to them then click on the box indicating you do and then the Download button will become available.
+ Click on the **DOWNLOAD** button to begin the driver download.
 
-	![](3-oculus-eula.jpg)
-[/EXCERPT:windows]
-[EXCERPT:macos]
-[INCLUDE:Platforms/Oculus#warning]
-1. Click the [Oculus Runtime for OS X version 0.4.2-beta](https://developer.oculus.com/downloads/#id=4d85655b).
-	![](3-oculus-runtime-for-os-x.jpg)
-1. Read and agree to the terms of the EULA to start the download.
+	![](3-oculus-eula.png)
 
-	![](3-oculus-eula.jpg)
-[/EXCERPT:macos]
--->
 
 %Steps%
diff --git a/Engine/Documentation/Source/Platforms/Oculus/QuickStart/3/OculusQuickStart_3.INT.udn b/Engine/Documentation/Source/Platforms/Oculus/QuickStart/3/OculusQuickStart_3.INT.udn
--- a/Engine/Documentation/Source/Platforms/Oculus/QuickStart/3/OculusQuickStart_3.INT.udn
+++ b/Engine/Documentation/Source/Platforms/Oculus/QuickStart/3/OculusQuickStart_3.INT.udn
@@ -2,7 +2,9 @@ Title:3. Install the Oculus Runtime
 Crumbs:%ROOT%, Platforms, Platforms/Oculus, Platforms/Oculus/QuickStart
 Description:Install the needed software for running games on the Oculus Rift
 Availability:Public
+Version: 4.9
 
+In the section we will go over how you can install the Oculus drivers once they have been downloaded to your development PC.
 
 [VAR:Steps]
 [OBJECT:Navigation]
@@ -23,90 +25,44 @@ Availability:Public
 %Steps%
 
 
-%Globals:OSSelect%
-
-[OBJECT:ToggleButtonContent]
-	[PARAMLITERAL:category]
-	OS
-	[/PARAMLITERAL]
-	[PARAMLITERAL:id]
-	windows
-	[/PARAMLITERAL]
-	[PARAMLITERAL:active]
-	active_button_content
-	[/PARAMLITERAL]
-	[PARAM:content]
-	[INCLUDE:#windows]
-	[/PARAM]
-[/OBJECT]
-[OBJECT:ToggleButtonContent]
-	[PARAMLITERAL:category]
-	OS
-	[/PARAMLITERAL]
-	[PARAMLITERAL:id]
-	mac
-	[/PARAMLITERAL]
-	[PARAMLITERAL:active]
-	[/PARAMLITERAL]
-	[PARAM:content]
-	[INCLUDE:#macos]
-	[/PARAM]
-[/OBJECT]
-<!--
-[EXCERPT:windows]
-1. Open the downloaded runtime installer.
-1. Click **Next**.
+1. Locate where the **0.7.0 BETA** drivers that were download in the last step and double-click on them to being the install process.
+
+1. When the installer starts, Click the **Next** button to being the install.
 	
 	![](oculus_setup_1.png)
+
 1. Accept the license agreement, then click **Next**.
 
 	![](oculus_setup_2.png)
+
 1. Specify an install location, then click **Next**.
 
 	![](oculus_setup_3.png)
+
 1. Click **Next**.
 	
 	![](oculus_setup_4.png)
+
 1. During the installation process, you will be prompted to install the imaging devices software. Click **Install**.
 
 	![](oculus_setup_6.png)
+
 1. You will also be prompted to install the display adapters software. Click **Install**.
 
 	![](oculus_setup_7.png)
+
 1. When setup is complete, click **Finish**.
 	
 	![](oculus_setup_8.png)
+
 1. The README file mentions updating the firmware for your DK2. This will be covered in the next step of the Quick Start.
 
 	 ![](oculus_setup_9.png)
+
 1. Reboot your computer when prompted.
 
 	![](oculus_setup_10.png)
-[/EXCERPT:windows]
-[EXCERPT:macos]
-[INCLUDE:Platforms/Oculus#warning]
-
-1. Open the downloaded runtime installer.
-1. Click **Next**.
-	
-	![](mac_oculus_1.png)(w:516)
-1. Accept the license agreement, then click **Next**.
-
-	![](mac_oculus_2.png)(w:516)
-1. Specify an install location, then click **Next**.
-
-	![](mac_oculus_3.png)(w:516)
-1. Click **Next**.
-	
-	![](mac_oculus_4.png)(w:516)
-1. When setup is complete, click **Finish**.
-	
-	![](mac_oculus_5.png)(w:516)
-1. The README file mentions updating the firmware for your DK2. This will be covered in the next step of the Quick Start.
 
-	 ![](mac_oculus_6.png)(w:516)
 
-[/EXCERPT:macos]
--->
 
 %Steps%
diff --git a/Engine/Documentation/Source/Platforms/Oculus/QuickStart/4/OculusQuickStart_4.INT.udn b/Engine/Documentation/Source/Platforms/Oculus/QuickStart/4/OculusQuickStart_4.INT.udn
--- a/Engine/Documentation/Source/Platforms/Oculus/QuickStart/4/OculusQuickStart_4.INT.udn
+++ b/Engine/Documentation/Source/Platforms/Oculus/QuickStart/4/OculusQuickStart_4.INT.udn
@@ -2,7 +2,7 @@ Title:4. Set Up Your Oculus Rift DK2 Hardware
 Crumbs:%ROOT%, Platforms, Platforms/Oculus, Platforms/Oculus/QuickStart
 Description:Connect and update your Oculus Rift headset.
 Availability:Public
-
+Version: 4.9
 
 [VAR:Steps]
 [OBJECT:Navigation]
@@ -20,98 +20,34 @@ Availability:Public
 [/OBJECT]
 [/VAR]
 
-%Steps%
+In this section we will look at how you connect your Oculus HMD to your PC and set it up to see if it is working with the 0.7.0 Oculus BETA drivers.
 
+%Steps%
 
-%Globals:OSSelect%
-
-[OBJECT:ToggleButtonContent]
-	[PARAMLITERAL:category]
-	OS
-	[/PARAMLITERAL]
-	[PARAMLITERAL:id]
-	windows
-	[/PARAMLITERAL]
-	[PARAMLITERAL:active]
-	active_button_content
-	[/PARAMLITERAL]
-	[PARAM:content]
-	[INCLUDE:#windows]
-	[/PARAM]
-[/OBJECT]
-[OBJECT:ToggleButtonContent]
-	[PARAMLITERAL:category]
-	OS
-	[/PARAMLITERAL]
-	[PARAMLITERAL:id]
-	mac
-	[/PARAMLITERAL]
-	[PARAMLITERAL:active]
-	[/PARAMLITERAL]
-	[PARAM:content]
-	[INCLUDE:#macos]
-	[/PARAM]
-[/OBJECT]
-<!--
-[EXCERPT:windows]
 1. Connect your DK2 headset to your computer, following the included Instruction Manual (which can also be found [online](http://static.oculus.com/sdk-downloads/documents/Oculus_Rift_DK2_Instruction_Manual.pdf)).
-1. Run the **Oculus Configuration Utility** by:
-	* **Right-clicking** on the Oculus icon in the system tray and selecting **Configuration Utility**, or
-	* Opening the Start Menu and navigating to **Oculus Runtime > OculusConfigUtil**, or
-	* Running `OculusConfigUtil.exe` found in `[InstallDirectory]/Tools/`.
-1. Click on **Tools > Rift Display Mode**.
 
-	![](RiftDisplayMode_menu.png)
-1. Make sure the headset is in **Direct HMD Access from Apps** mode, then click **Apply** to save any changes. Close the window.
-	
-	![](DirectHMD.png)
+1. Run the **Oculus Configuration Utility** by **Right-clicking** on the Oculus icon in the system tray and selecting **Configuration Utility**.
 
-[REGION:tip]
-Your firmware version is listed in the Configuration Utility. If your version is not 2.12, complete the following steps to update the headset's firmware.
-[/REGION]
+	![](Open_Config_Utility.png)
 
-1. Click on **Tools > Advanced > Update Firmware**.
+1. When opened, the Oculus Configuration Utility will look like the following image.
 
-	![](UpdateFirmware_menu.png)
+	![](Oculus_Config_Utility.png)
 
-	[REGION:note]
-	Make sure Unreal Editor, Unreal Engine Launcher, and any other programs that might be using the Oculus are closed before you update the firmware. Otherwise, you will get a warning about closing other software and the firmware will not install.
-	[/REGION]
+1. Under the **User** section press the **Plus Sign** to add a new User profile. 
 
-1. Select `DK2Firmware_2_12.ovrf` when the file dialog appears, then click **Open**.
+	![](Oculus_Config_Add_Profile.png)
 
-	![](Firmware.png)
-1. The **Oculus Configuration Utility** will close. Reopen it by:
-	* Opening the Start Menu and navigating to **Oculus Runtime > OculusConfigUtil**, or
-	* Running `OculusConfigUtil.exe` found in `[InstallDirectory]/Tools/`.
-1. Verify that the firmware is 2.12.
-[/EXCERPT:windows]
-[EXCERPT:macos]
-[INCLUDE:Platforms/Oculus#warning]
-
-1. Connect your DK2 headset to your computer, following the included Instruction Manual (which can also be found [online](http://static.oculus.com/sdk-downloads/documents/Oculus_Rift_DK2_Instruction_Manual.pdf)).
-1. Run the **Oculus Configuration Utility** by:
-	* Using the OculusConfigUtil shortcut on the desktop, or
-	* Running `OculusConfigUtil.app` found in `[InstallDirectory]/Tools/`.
+1. In the **New User** pop up box that is displayed enter a name for the profile and then click the **OK** button.. 
 
-[REGION:tip]
-Your firmware version is listed in the Configuration Utility. If your version is not 2.12, complete the following steps to update the headset's firmware.
-[/REGION]
+	![](Oculus_Config_Name_Profile.png)
 
-1. Click on **Tools > Advanced > Update Firmware**.
+1. Then set the in the User section set the **Gender** to the gender of the person using the Head Mounted Display(HMD) and set the **Player Height** to that persons height.
 
-	![](firmware_mac.png)(w:420)
+	![](Oculus_Config_Setup_Profile.png)
 
-	[REGION:note]
-	Make sure Unreal Editor, Unreal Engine Launcher, and any other programs that might be using the Oculus are closed before you update the firmware. Otherwise, you will get a warning about closing other software and the firmware will not install.
-	[/REGION]
+1. Click the **Show Demo Scene** button on the Oculus Configuration Utility and you should now see something similar to the following image being displayed on screen as well as being displayed in the HMD.
 
-1. Select `DK2Firmware_2_12.ovrf` when the file dialog appears, then click **Open**.
-1. The **Oculus Configuration Utility** will close. Reopen it by:
-	* Using the OculusConfigUtil shortcut on the desktop, or
-	* Running `OculusConfigUtil.app` found in `[InstallDirectory]/Tools/`.
-1. Verify that the firmware is 2.12.
-[/EXCERPT:macos]
--->
+	![](Oculus_Config_Demo_Sceen.png)
 
 %Steps%
diff --git a/Engine/Documentation/Source/Platforms/Oculus/QuickStart/5/OculusQuickStart_5.INT.udn b/Engine/Documentation/Source/Platforms/Oculus/QuickStart/5/OculusQuickStart_5.INT.udn
--- a/Engine/Documentation/Source/Platforms/Oculus/QuickStart/5/OculusQuickStart_5.INT.udn
+++ b/Engine/Documentation/Source/Platforms/Oculus/QuickStart/5/OculusQuickStart_5.INT.udn
@@ -2,7 +2,9 @@ Title:5. Test out VR Content
 Crumbs:%ROOT%, Platforms, Platforms/Oculus, Platforms/Oculus/QuickStart
 Description:Load and play content designed for the Oculus Rift using Unreal Engine 4.
 Availability:Public
+Version: 4.9
 
+In this section we will cover how you can test to see if your Dk2 is working inside Unreal Engine 4.
 
 [VAR:Steps]
 [OBJECT:Navigation]
@@ -22,30 +24,37 @@ Availability:Public
 
 %Steps%
 
-1. If you have not already, download the [](Resources\ContentExamples) project.
+1. If you have not already, download the [](Resources\ContentExamples) project from the Unreal Engine Launcher.
+
 1. Open the **Content Examples** project.
 	
 	![](CE_Start.png)(w:500)
+
 1. Click on **File > Open Level**.
-1. Select `VR.umap`, then click **Open**.
+
+1. Select the **VR.umap**, then click **Open**.
 
 	![](VRmap.png)
+
 1. Click on the arrow next to **Play** in the **Level Editor Toolbar** to open the **Play In Editor** options.
 
 	![](PlayDropdown.png)
-1. Select **Standalone Game** to launch the game. **Standalone Game** is the only **Play In Editor** game testing mode that is compatible with the Oculus.
+
+1. Select the **VR Preview** option and then put the HMD on.
 
 	![](StandAloneGame.png)
 
 	[REGION:tip]
-	Once you have selected **Standalone Game** for a project, it will remain in that mode until you pick a different one. This means you can just click on the **Play** icon to resume testing.
+	You can use the **Standalone Game** option but if you do make sure to press **Alt+Enter** to enter VR mode. If you do not do this you will not see anything when you put the HMD on.
 	[/REGION]
 
-1. Press **Alt+Enter** to enter VR mode. This will send your game's display to the Oculus, as well as change the view on the screen to stereo view.
+1. With the HMD on you should now something similar to the image below being displayed in the HMD.
 
 	![](VRMode.png)(w:800)
-1. At this point, you can put on the Oculus and move around the level. To end gameplay:
-	* Press **~** and type `quit` in the console, or 
+
+1. To end game play you can do either of the following two options.
+
+	* Press **~** and type **Quit** in the console
 	* Press **Shift+F1** to reclaim your mouse cursor and close the window.
 
 
diff --git a/Engine/Documentation/Source/Platforms/Oculus/QuickStart/6/OculusQuickStart_6.INT.udn b/Engine/Documentation/Source/Platforms/Oculus/QuickStart/6/OculusQuickStart_6.INT.udn
--- a/Engine/Documentation/Source/Platforms/Oculus/QuickStart/6/OculusQuickStart_6.INT.udn
+++ b/Engine/Documentation/Source/Platforms/Oculus/QuickStart/6/OculusQuickStart_6.INT.udn
@@ -2,7 +2,9 @@ Title:6. Additional Resources
 Crumbs:%ROOT%, Platforms, Platforms/Oculus, Platforms/Oculus/QuickStart
 Description:Read more about Oculus development with Unreal Engine.
 Availability:Public
+Version: 4.9
 
+Here are some additional resources that provide useful information for developing VR projects inside of Unreal Engine 4.
 
 [VAR:Steps]
 [OBJECT:Navigation]
@@ -53,8 +55,37 @@ Availability:Public
 		| ------------ | ---------------- |
 		| CouchKnights | Vehicle Template[REGION:caption]Press **Tab** to switch to being inside the car.[/REGION] |
 		| ![](store_RealisticRendering_screenshot_3.png)(w:420) | ![](store_MobileTemple_screenshot_4.png)(w:420) |
-		| Architecture[REGION:caption]Disable the Matinee fly-through, then walk around the apartment level.[/REGION] | Simple Temple Environment |
+		| Architecture[REGION:caption] Disable the Matinee fly-through to walk around.[/REGION] | Simple Temple Environment |
 		[/REGION]
 	[/PARAM]
 [/OBJECT]
-%Steps%
+
+
+## Flashing Oculus Firmware
+
+Depending on when you got your Oculus HMD you might need to update your firmware.
+The current firmware version is listed in the Oculus Configuration Utility. 
+If your firmware version is not 2.12, complete the following steps to update the HMD firmware.
+
+1. Click on **Tools > Advanced > Update Firmware**.
+
+	![](Platforms\Oculus\QuickStart\4\UpdateFirmware_menu.png)
+
+	[REGION:note]
+	Make sure Unreal Editor, Unreal Engine Launcher, and any other programs that might be using the Oculus are closed before you update the firmware. Otherwise, you will get a warning about closing other software and the firmware will not install.
+	[/REGION]
+
+1. Select `DK2Firmware_2_12.ovrf` when the file dialog appears, then click **Open**.
+
+	![](Platforms\Oculus\QuickStart\4\Firmware.png)
+
+1. The **Oculus Configuration Utility** will close. Reopen it by:
+
+	* Opening the Start Menu and navigating to **Oculus Runtime > OculusConfigUtil**, or
+	* Running `OculusConfigUtil.exe` found in `[InstallDirectory]/Tools/`.
+
+1. Verify that the firmware is 2.12 by going to **Help** and clicking on **About**. From this menu you will be able to determine which version of the Oculus drivers you are using.
+
+	![](Oculus_Current_Firmware.png)
+
+%Steps%
\ No newline at end of file
diff --git a/Engine/Documentation/Source/Platforms/Oculus/QuickStart/OculusQuickStart.INT.udn b/Engine/Documentation/Source/Platforms/Oculus/QuickStart/OculusQuickStart.INT.udn
--- a/Engine/Documentation/Source/Platforms/Oculus/QuickStart/OculusQuickStart.INT.udn
+++ b/Engine/Documentation/Source/Platforms/Oculus/QuickStart/OculusQuickStart.INT.udn
@@ -2,6 +2,7 @@ Title:Oculus Rift Quick Start
 Crumbs:%ROOT%, Platforms, Platforms/Oculus
 Description:Instructions on setting up and launching games on the Oculus Rift.
 Availability:Public
+Version: 4.9
 
 [VAR:Topic]
 [OBJECT:Topic]
