#UE4 Docs: fixed list formatting by using variables instead of includes

diff --git a/Engine/Documentation/Source/Engine/Rendering/Materials/Editor/Interface/MaterialEditorInterface.INT.udn b/Engine/Documentation/Source/Engine/Rendering/Materials/Editor/Interface/MaterialEditorInterface.INT.udn
--- a/Engine/Documentation/Source/Engine/Rendering/Materials/Editor/Interface/MaterialEditorInterface.INT.udn
+++ b/Engine/Documentation/Source/Engine/Rendering/Materials/Editor/Interface/MaterialEditorInterface.INT.udn
@@ -23,38 +23,38 @@ The Material Editor UI consists of a menu bar, a toolbar, and, by default, four
 There is also an HLSL Code panel, but by default, it is not visible. To display it, on the **Window** menu, click **HLSL Code**. For more information, see [HLSL Code](#hlslcodepanel).
 [/REGION]
 
-[Include:Shared/Editors/Common/EditorMenuItems/#UITabTips]
+%Shared/Editors/Common/EditorMenuItems:UITabTips%
 
 ## Menu Bar
 
 ### File
 
-[Include:Shared/Editors/Common/EditorMenuItems/#MenuFileSave]
-[Include:Shared/Editors/Common/EditorMenuItems/#MenuFileSaveOptions]
+%Shared/Editors/Common/EditorMenuItems:MenuFileSave%
+%Shared/Editors/Common/EditorMenuItems:MenuFileSaveOptions%
 
 ### Edit
 
-[Include:Shared/Editors/Common/EditorMenuItems/#MenuEditUndo]
-[Include:Shared/Editors/Common/EditorMenuItems/#MenuEditSettings]
+%Shared/Editors/Common/EditorMenuItems:MenuEditUndo%
+%Shared/Editors/Common/EditorMenuItems:MenuEditSettings%
 
 ### Assets
 
-[Include:Shared/Editors/Common/EditorMenuItems/#MenuAssetContentBrowser]
+%Shared/Editors/Common/EditorMenuItems:MenuAssetContentBrowser%
 
 ### Window
 
-[Include:Shared/Editors/Common/EditorMenuItems/#MenuWindowToolbar]
-[Include:Shared/Editors/Common/EditorMenuItems/#MenuWindowViewport]
+%Shared/Editors/Common/EditorMenuItems:MenuWindowToolbar%
+%Shared/Editors/Common/EditorMenuItems:MenuWindowViewport%
 * **Graph** - Toggles the display of the **Graph** panel.
-[Include:Shared/Editors/Common/EditorMenuItems/#MenuWindowDetails]
+%Shared/Editors/Common/EditorMenuItems:MenuWindowDetails%
 * **Palette** - Toggles the display of the **Palette** panel.
 * **Stats** - Toggles the display of the **Stats** panel.
 * **HLSL Code** - Toggles the display of the **HLSL Code** panel (which is closed by default).
-[Include:Shared/Editors/Common/EditorMenuItems/#MenuWindow]
+%Shared/Editors/Common/EditorMenuItems:MenuWindow%
 
 ### Help
 
-[Include:Shared/Editors/Common/EditorMenuItems/#MenuHelp]
+%Shared/Editors/Common/EditorMenuItems:MenuHelp%
 
 
 ## Toolbar
diff --git a/Engine/Documentation/Source/Shared/Editors/Common/EditorMenuItems/EditorMenuItems.INT.udn b/Engine/Documentation/Source/Shared/Editors/Common/EditorMenuItems/EditorMenuItems.INT.udn
--- a/Engine/Documentation/Source/Shared/Editors/Common/EditorMenuItems/EditorMenuItems.INT.udn
+++ b/Engine/Documentation/Source/Shared/Editors/Common/EditorMenuItems/EditorMenuItems.INT.udn
@@ -5,82 +5,82 @@ Crumbs:
 
 
 
-[EXCERPT:MenuFileSave]
+[VAR:MenuFileSave]
 * **Save** - Saves the asset you are currently working on.
-[/EXCERPT:MenuFileSave]
+[/VAR]
 
-[EXCERPT:MenuFileSaveOptions]
+[VAR:MenuFileSaveOptions]
 * **Open Asset** - Opens the **Global Asset Picker** to quickly find an asset and open the appropriate editor.
 * **Save All** - Saves all unsaved levels and assets for your project.
 * **Choose Files to Save** - Brings up a dialog box that allows you to choose which levels and assets you want to save for your project.
 * **Connect To Source Control** - Brings up a dialog box that allows you to choose or interact with a source-control system that the Unreal Editor can integrate with.
-[/EXCERPT:MenuFileSaveOptions]
+[/VAR]
 
 
 
-[EXCERPT:MenuEditUndo]
+[VAR:MenuEditUndo]
 * **Undo** - Undoes the most recent action.
 * **Redo** - Redoes the most recent Undo, if it was the last action taken.
-[/EXCERPT:MenuEditUndo]
+[/VAR]
 
-[EXCERPT:MenuEditSettings]
+[VAR:MenuEditSettings]
 * **Editor Preferences** - Gives you a list of options, any of which open that part of the **Editor Preferences** window, where you can modify your Unreal Editor preferences.
 * **Project Settings** - Gives you a list of options, any of which open that part of the **Project Settings** window, where you can modify various settings for your Unreal Engine project.
-[/EXCERPT:MenuEditSettings]
+[/VAR]
 
 
 
 
 
-[EXCERPT:MenuAssetContentBrowser]
+[VAR:MenuAssetContentBrowser]
 * **Find in Content Browser** - Locates and selects the current asset in the **Content Browser**.
-[/EXCERPT:MenuAssetContentBrowser]
+[/VAR]
 
-[EXCERPT:MenuAssetViewReference]
+[VAR:MenuAssetViewReference]
 * **View Reference** - Launches the reference viewer showing the selected assets' references. 
-[/EXCERPT:MenuAssetViewReference]
+[/VAR]
 
-[EXCERPT:MenuAssetSizeMap]
+[VAR:MenuAssetSizeMap]
 * **Size Map** - Displays an interactive map showing the approximate size of this asset and everything it references. 
-[/EXCERPT:MenuAssetSizeMap]
+[/VAR]
 
-[EXCERPT:MenuAssetReimport]
+[VAR:MenuAssetReimport]
 * **Reimport _filename_** - Reimports the current asset from its original location on disk.
-[/EXCERPT:MenuAssetReimport]
+[/VAR]
 
 
 
 
 
-[EXCERPT:MenuWindowToolbar]
+[VAR:MenuWindowToolbar]
 * **Toolbar** - Toggles the display of the Toolbar.
-[/EXCERPT:MenuWindowToolbar]
+[/VAR]
 
-[EXCERPT:MenuWindowViewport]
+[VAR:MenuWindowViewport]
 * **Viewport** - Toggles the display of the **Viewport** panel.
-[/EXCERPT:MenuWindowViewport]
+[/VAR]
 
-[EXCERPT:MenuWindowPreview]
+[VAR:MenuWindowPreview]
 * **Preview** - Toggles the display of the **Preview** panel.
-[/EXCERPT:MenuWindowPreview]
+[/VAR]
 
-[EXCERPT:MenuWindowDetails]
+[VAR:MenuWindowDetails]
 * **Details** - Toggles the display of the **Details** panel.
-[/EXCERPT:MenuWindowDetails]
+[/VAR]
 
-[EXCERPT:MenuWindowDestructibleSettings]
+[VAR:MenuWindowDestructibleSettings]
 * **Destructible Settings** - Toggles the display of the **Destructible Settings** panel.
-[/EXCERPT:MenuWindowDestructibleSettings]
+[/VAR]
 
-[EXCERPT:MenuWindowFractureSettings]
+[VAR:MenuWindowFractureSettings]
 * **Fracture Settings** - Toggles the display of the **Fracture Settings** panel.
-[/EXCERPT:MenuWindowFractureSettings]
+[/VAR]
 
-[EXCERPT:MenuWindowChunkParameters]
+[VAR:MenuWindowChunkParameters]
 * **Chunk Parameters** - Toggles the display of the **Chunk Parameters** panel.
-[/EXCERPT:MenuWindowChunkParameters]
+[/VAR]
 
-[EXCERPT:MenuWindow]
+[VAR:MenuWindow]
 * **Content Browser** - Brings up the **Content Browser** in a separate window.
 * **Developer Tools** - Brings up the selected **Developer Tool** in a separate window (Camera Debugger, Collision Analyzer, Debug Tools, Modules, Widget Reflector).
 	* **Blueprint Debugger** - Brings up the **Blueprint Debugger** in a separate window.
@@ -100,12 +100,12 @@ Crumbs:
 * **Reset Layout** - Resets your layout to a default arrangement. Requires the editor to restart after saving your changes and creating a backup of your settings. 
 * **Save Layout** - Saves the current layout of the panels as the new default layout.
 * **Enable Fullscreen** - Enables fullscreen mode for the application, expanding across the entire monitor.
-[/EXCERPT:MenuWindow]
+[/VAR]
 
 
 
 
-[EXCERPT:MenuHelp]
+[VAR:MenuHelp]
 [PUBLISH:Rocket]
 * **Browse Documentation** - Opens a browser window and navigates to the documentation about this tool.
 * **Forums** - Opens a browser window to the Unreal Engine Forums where you can discuss the Unreal Engine with other users.
@@ -122,13 +122,13 @@ Crumbs:
 * **Visit EpicGames.com** - Opens EpicGames.com in a browser window.
 * **About Unreal Editor** - Displays application credits, copyright information, and build information.
 [/PUBLISH:Licensee]
-[/EXCERPT:MenuHelp]
+[/VAR]
 
 
 
 
-[EXCERPT:UITabTips]
+[VAR:UITabTips]
 [REGION:tip]
 You can close any panel by clicking the small "X" in the upper-right corner of the tab. You can also hide any panel by **Right-clicking** on the tab, and then clicking **Hide Tab** on the context menu that appears. To once again display a panel that you have closed, click that panel's name on the **Window** menu.
 [/REGION]
-[/EXCERPT:UITabTips]
+[/VAR]
