UEDOC-2353

Updating images on this page and adding a new Stats Panel Section

diff --git a/Engine/Documentation/Source/Engine/Rendering/Materials/Editor/Interface/MaterialEditorInterface.INT.udn b/Engine/Documentation/Source/Engine/Rendering/Materials/Editor/Interface/MaterialEditorInterface.INT.udn
--- a/Engine/Documentation/Source/Engine/Rendering/Materials/Editor/Interface/MaterialEditorInterface.INT.udn
+++ b/Engine/Documentation/Source/Engine/Rendering/Materials/Editor/Interface/MaterialEditorInterface.INT.udn
@@ -11,7 +11,7 @@ Version: 4.9
 
 The Material Editor UI consists of a menu bar, a toolbar, and, by default, four open panels.
 
-![materialEditorAlt.png](materialEditor.png)
+![](UpdatedUI.jpg)
 
 1. [Menu Bar](#MenuBar) - Lists the menu options for the current Material.
 1. [Toolbar](#Toolbar) - Contains tools that you can use to work with your Material.
@@ -119,15 +119,9 @@ all Material properties.
 
 ## Graph Panel
 
-![Material Editor Graph Panel](MaterialEditor_Graph.png)
+![](MaterialGraph1.jpg)
 
-This panel contains a graph of all material expressions that belong to this Material.  The number of shader 
-instructions used in the Material as well as any [Compiler Errors](Engine/Rendering/Materials/Editor/#CompilerErrors) are displayed in the 
-upper left corner.  The fewer the number of instructions, the less costly the Material.  Material expression 
-nodes that are not connected to the base material node do not contribute to the instruction count (cost) of 
-the Material.
-
-By default, each Material contains a single base Material node. This node has a series of inputs, each associated with 
+This panel contains a graph of all material expressions that belong to this Material. By default, each Material contains a single base Material node. This node has a series of inputs, each associated with 
 a different aspect of the Material, that other material nodes can connect to.
 
 ![Material Editor Material Node](MaterialEditor_MaterialNode.png)
@@ -147,6 +141,15 @@ You can filter the material nodes listed in the **Palette** panel by choosing ei
 For more information about material functions, see [Material Functions](Engine/Rendering/Materials/Functions).
 
 
+## Stats Panel
+
+![](StatsPanel.jpg)
+
+The number of shader instructions used in the Material as well as any [Compiler Errors](Engine/Rendering/Materials/Editor/#CompilerErrors) are displayed in this panel. The fewer the number of instructions, the less costly the Material.  Material expression 
+nodes that are not connected to the base material node do not contribute to the instruction count (cost) of 
+the Material.
+
+
 ## HLSL Code Panel
 
 ![](MaterialEditor_HLSLPanel.png)
