#UE4doc #docs - Updated Images to use the Gray guy. Addresess UEDOC-2301

diff --git a/Engine/Documentation/Source/Engine/Physics/WalkableSlope/WalkableSlope.INT.udn b/Engine/Documentation/Source/Engine/Physics/WalkableSlope/WalkableSlope.INT.udn
--- a/Engine/Documentation/Source/Engine/Physics/WalkableSlope/WalkableSlope.INT.udn
+++ b/Engine/Documentation/Source/Engine/Physics/WalkableSlope/WalkableSlope.INT.udn
@@ -62,7 +62,7 @@ Depending on the editor you are in, or if you are looking at an Actor in the lev
 
 This setting will increase the walkable slope of a **Physics Body** up to the value listed in the **Walkable Slope Angle** property. If **Walkable Slope Angle** is set to 75, and the angle of the surface of the **Physics Body** is 65, the Character will be able to walk up it, regardless of the Character's **Walkable Floor Angle**:
 
-![](increased.png)
+![](increased1.png)
 
 _The blue angle is the default **Character Movement Component Walkable Floor Angle**, while the green represents the new **Increased Walkable Slope** angle._
 
@@ -75,7 +75,7 @@ While **Walkable Slope Behavior** is set to "Increase Walkable Slope":
 
 This setting will "cap" the walkable slope of a **Physics Body** to the value listed in the **Walkable Slope Angle** property. If **Walkable Slope Angle** is set to 25, and the angle of the surface of the **Physics Body** is 35 degrees, the Character will be unable to walk up it regardless of the Character's **Walkable Floor Angle**:
 
-![](Decreased.png)
+![](Decreased1.png)
 
 _The blue angle is the default **Character Movement Component Walkable Floor Angle**, while the red represents the new **Decreased Walkable Slope** angle._
 
