#UE4doc #docs - Added Engine Version field. Added Skill Levels field. Replaced Image with Comparison Slider. Edited one sentence to indicate that Eye Adaptation can be disabled in the Post Process Volume settings.

diff --git a/Engine/Documentation/Source/Engine/Rendering/Materials/HowTo/Human_Skin/HumanSkin.INT.udn b/Engine/Documentation/Source/Engine/Rendering/Materials/HowTo/Human_Skin/HumanSkin.INT.udn
--- a/Engine/Documentation/Source/Engine/Rendering/Materials/HowTo/Human_Skin/HumanSkin.INT.udn
+++ b/Engine/Documentation/Source/Engine/Rendering/Materials/HowTo/Human_Skin/HumanSkin.INT.udn
@@ -8,14 +8,27 @@ Related: Engine/Rendering/Materials/HowTo/Subsurface_Scattering
 Related: Engine/Rendering/Materials/LightingModels/SubSurfaceProfile
 Related: Engine/Rendering/Materials/LightingModels/SubSurface
 Related: Resources/ContentExamples/MaterialNodes/1_7
+Version: 4.9
+SkillLevel: Advanced
 
 [TOC(start:2)]
 
 
+[OBJECT:ComparisonSlider]
+	[PARAM:before]
+	[![Not using SubsurfaceProfile](Engine\Rendering\Materials\LightingModels\SubSurfaceProfile\Results_1.png)(w:700 convert:false)](Engine\Rendering\Materials\LightingModels\SubSurfaceProfile\Results_1.png)
+	[/PARAM]
+	[PARAM:after]
+	[![Using SubsurfaceProfile](Engine\Rendering\Materials\LightingModels\SubSurfaceProfile\Results_2.png)(w:700 convert:false)](Engine\Rendering\Materials\LightingModels\SubSurfaceProfile\Results_2.png)
+	[/PARAM]
+[/OBJECT]
 
+<!---
 [REGION:lightbox]
 [![](Engine\Rendering\Materials\LightingModels\SubSurfaceProfile\Results.png)(w:920 convert:false)](Engine\Rendering\Materials\LightingModels\SubSurfaceProfile\Results.png)
 [/REGION]
+--->
+
 
 Getting Human skin to look correct when rendered is a very complex topic with no single perfect solution for getting correct results. 
 Because of this, please treat the following information _as a starting point_ for your work as many things are needed in order to achieve a realistic looking image. 
@@ -33,7 +46,7 @@ When viewing your content inside the editor, make sure to disable the **Editor E
 [/region]
 
 For the following example, we will be working with the Skin Rendering Content example. 
-If have not done so, please go ahead and download the Content Examples from the Market Place now and open up the Skin Rendering Map. 
+If have not done so, please go ahead and download the Content Examples from the Marketplace now and open up the Skin Rendering Map. 
 Also please keep in mind that the images in this document are explicitly low resolution and are not meant to be used as source work.
 
 [region:note]
@@ -46,14 +59,14 @@ Here is a list of some things that you need to be aware of before you begin to c
 * Remember that when creating textures, it would be ideal if you start with very high resolution textures (e.g. 2Kx2K or higher) as this will give the best results. 
 * For faster iteration, it is best to create a Material Instance and expose Scalar and Vector Parameters for faster and easier tweaking as the Material does not need to be recompiled during tweaking. 
 * Keep in mind that assets should be verified in multiple lighting conditions to ensure that they are correct. 
-  Even having Eye Adaptation enabled can make a huge difference in how things look in the end. 
+  Even having Eye Adaptation enabled in the Post Process Volume can make a huge difference in how things look in the end. 
 * To avoid shadow artifacts will require extra effort because you will need to tweak things like shadow bias or use a more focused spot light cone for better lighting resolution.
 * When tweaking, it might be useful to zoom in by adjusting the editor FOV as there will be no clipping through geometry and less distortions. 
 * The editor Field Of View(FOV) can be adjusted by clicking on the Drop Down Arrow that is next to the icon that says Perspective and then adjusting the FOC value using the mouse. 
 	
 	![](HS_Change_FOV.png)
 
-* When viewing your work, you might notice that the streaming system does not take the zoom into account and could blur your texture when you do not want to. This was an intentional design decision made for 3D games but this could change at some point in the future. To prevent this from happening, you can change the texture Streaming Distance Multiplier  on the mesh instance.
+* When viewing your work, you might notice that the streaming system does not take the zoom into account and could blur your texture when you do not want to. This was an intentional design decision made for 3D games but this could change at some point in the future. To prevent this from happening, you can change the texture Streaming Distance Multiplier on the mesh instance.
 
 ## Material 
 
