#DocFix

Updating out of date images pointed out on the forums

diff --git a/Engine/Documentation/Source/Engine/UMG/QuickStart/2/UMGQS_2.INT.udn b/Engine/Documentation/Source/Engine/UMG/QuickStart/2/UMGQS_2.INT.udn
--- a/Engine/Documentation/Source/Engine/UMG/QuickStart/2/UMGQS_2.INT.udn
+++ b/Engine/Documentation/Source/Engine/UMG/QuickStart/2/UMGQS_2.INT.udn
@@ -79,7 +79,7 @@ Here we will begin setting up the visual aspects of our HUD.
 
 1.	Select the **Progress Bar** next to health and in the **Details** panel set the **Fill Color and Opacity** to a green color. 
 
-	![](Step2_24.png)(w:800)
+	![](UMGColorPicker1.jpg)(w:800)
 
 	You may notice that Progress Bar does not change color when assigning a color. This is because the **Percent** value to fill the bar is set to 0.0 (you can change this to test out different colors, later we will link this to the health value of our character). 
 
@@ -135,17 +135,17 @@ Now that we have our visual layout, we can add the hooks to provide the function
 
 1.	In the **Details** panel under **Progress**, click the **Bind** option next to **Percent** and set it to use **Health** from **MyCharacter**. 
 
-	![](Step2_23.png)
+	![](HealthProgress.jpg)
 
 	This binds the value of the Progress Bar to our **Health** variable inside our **First Person Character** Blueprint. Now whenever our Health value changes in our Character Blueprint it will automatically be updated in our HUD as well. You may also notice that the Progress Bar does not change after binding it. That is because we've stated this to occur upon **Event Construct** (or when the HUD is constructed, which occurs when the game is started).
 
 1.	Select the **Progress Bar** next to Energy and repeat the process above, binding **Percent** to **Energy** from **MyCharacter**. 
 
-	![](Step2_25.png)
+	![](EnergyProgress.jpg)
 
 1.	Select the 25 following the Ammo text, then in the **Details** panel for **Text**, click **Bind** and **Create Binding**. 
 
-	![](Step2_26.png)
+	![](CreateBinding.jpg)
 
 	Now, like the way we previously used **Sub-Object Properties** to bind properties of the same type, we can also create our own custom bindings. Here we are going to bind our **Text** property to an **Integer** property from our **First Person Character** Blueprint giving us a display of our current Ammo.  
 
diff --git a/Engine/Documentation/Source/Engine/UMG/QuickStart/4/UMGQS_4.INT.udn b/Engine/Documentation/Source/Engine/UMG/QuickStart/4/UMGQS_4.INT.udn
--- a/Engine/Documentation/Source/Engine/UMG/QuickStart/4/UMGQS_4.INT.udn
+++ b/Engine/Documentation/Source/Engine/UMG/QuickStart/4/UMGQS_4.INT.udn
@@ -90,7 +90,7 @@ Time to get our hands dirty and script the functionality of our Main Menu. At th
 
 1.	Open the **First Person Character** Blueprint inside the **Content/FirstPersonBP/Blueprints** folder. 
 
-1.	Following the **Event Begin Play** script, add a **Get Player Controller** then **Set Input Mode Game Only** node. 
+1.	Following the **Event Begin Play** script, add a **Get Player Controller** then add a **Set Input Mode Game Only** node. 
 
 	![](PlayerCharacterBlueprint.png)(w:820)
 
