#UE4doc #docs - Added Engine version field. Added Skill Level field. Updated Images so latest UI and to match colors for instructional guide.

diff --git a/Engine/Documentation/Source/Engine/Rendering/Materials/HowTo/Instancing/Instancing.INT.udn b/Engine/Documentation/Source/Engine/Rendering/Materials/HowTo/Instancing/Instancing.INT.udn
--- a/Engine/Documentation/Source/Engine/Rendering/Materials/HowTo/Instancing/Instancing.INT.udn
+++ b/Engine/Documentation/Source/Engine/Rendering/Materials/HowTo/Instancing/Instancing.INT.udn
@@ -8,6 +8,8 @@ Related: Engine/Rendering/Materials/MaterialInputs
 Related: Engine/Rendering/Materials/Editor
 Related: Engine/Rendering/Materials/IntroductionToMaterials
 Related: Engine/Rendering/Materials/MaterialProperties
+Version:4.9
+SkillLevel: Intermediate
 
 [TOC(start:2 end:3)]
 
@@ -222,7 +224,7 @@ We will then change the Parent Material to show how changes to the Parent will a
    To do this, use the Mouse to **Right-Click** on the **Material_Master** in the **Content Browser** and then select **Create New Material Instance** from the menu. 
    Do this two more times until you have a total of three Material Instances. 
 
-	![](Materails_Instances_Create_3_New_MI.png)
+	![](Materails_Instances_Create_3_New_MI.png)(w:800)
 
 1. Now with the three additional Material Instances created it is time to open them up and begin adjusting their **Parameters**.
    To open one of the Material Instances up **Double-Click** on it with the **Left Mouse Button** inside of the **Content Browser**. 
