#UE4doc #docs - Added Engine Version Field. Added Skill Level Field. Bolded/unbolded terms.

diff --git a/Engine/Documentation/Source/Engine/Rendering/Materials/HowTo/Main_Material_Node/Main_Material_Node.INT.udn b/Engine/Documentation/Source/Engine/Rendering/Materials/HowTo/Main_Material_Node/Main_Material_Node.INT.udn
--- a/Engine/Documentation/Source/Engine/Rendering/Materials/HowTo/Main_Material_Node/Main_Material_Node.INT.udn
+++ b/Engine/Documentation/Source/Engine/Rendering/Materials/HowTo/Main_Material_Node/Main_Material_Node.INT.udn
@@ -7,6 +7,8 @@ Related: Engine/Rendering/Materials/MaterialInputs
 Related: Engine/Rendering/Materials/Editor
 Related: Engine/Rendering/Materials/IntroductionToMaterials
 Related: Engine/Rendering/Materials/MaterialProperties
+Version: 4.9
+SkillLevel: Beginner
 
 [TOC(start:2 end:3)]
 
@@ -14,11 +16,11 @@ Related: Engine/Rendering/Materials/MaterialProperties
 
 The Main Material node of the Material Graph is where all of the Material Expression nodes that are added to the Material Graph will end up being connected to. 
 The results of what is input into the Main Material node is what will be displayed when the Material is compiled and applied to objects in game. 
-In the following How - To, we will cover all you need to know about how to use the Main Material node.
+In the following How-to, we will cover all you need to know about how to use the Main Material node.
 
 ## The Main Material Node
 
-Materials are created using a specialized coding language called High Level Shading Language or HLSL.
+**Materials** are created using a specialized coding language called **High Level Shading Language** or HLSL.
 HLSL allows Materials to talk directly to the graphics hardware giving artist and coders better control over what is being displayed on screen.
 Inside Unreal Engine 4 (UE4) the Material Expression nodes that are used to create Materials contain small snippets of this HLSL code.
 To display the results of all these small snippets of HLSL code, the Main Material Node is used.
@@ -67,25 +69,25 @@ like in the image below.
 ### Main Material Node Details Panel Properties
 
 The Main Material nodes **Details** panel is where you can adjust properties that relate to how this Material can be used. 
-From changing the Blend Mode that is used to adjusting the types of objects the Material can interact with, the Main Material Node **Details** panel can be thought of as the properties of the Material. 
-Here is a brief breakdown of what each section in the Main Material node **Details** panel does.
+From changing the Blend Mode that is used to adjusting the types of objects the Material can interact with, the Main Material Node Details panel can be thought of as the properties of the Material. 
+Here is a brief breakdown of what each section in the Main Material node Details panel does.
 
 ![](Material_MMN_Properties.png)
 
 
 |Section Name: | Description|
 |---|---|
-|Physical Material|This allows you to specify what type of Physical Material is used for this Material.|
-| Material|This is where you will spend most of your time when editing Materials. The Material section is where you can change the Material Domain, Blend Mode, and Shading Model along with many more options. |
-| Translucency|This is where you can adjust how Translucency in this Material will function. Please note that this is only editable if the Material Blend Mode is set to Translucent. |
-| Translucency Self Shadowing|This is where you can adjust how Translucent self shadowing looks and behaves. Please note that this is only editable if the Material Blend Mode is set to Translucent.|
-| Usage|This is where you set what type of objects this Material will work with. Usage flags are usually set automatically by the editor. However if you know that this Material is supposed to be used with a certain object type, make sure you enable it here to avoid any future mistakes|
-| Mobile|This is where you set how this Material should work when used on a mobile device like a smart phone.|
-| Tessellation|This is where you can enable the Material to make use of Hardware Tessellation.|
-| Post Process Material|This is where you define how this Material should work with Post Process and Tone Mapping. Please note that this is only editable if the Material Domain is set to Post Process.|
-| Lightmass|This is where you can adjust how this Material will interact with Lightmass.|
-| Material Interface|This is where you can define a Static Mesh to preview the Material with.|
-| Thumbnail|This controls how the Thumbnail that you see in the **Content Browser** is displayed.|
+| **Physical Material** |This allows you to specify what type of Physical Material is used for this Material.|
+| **Material**|This is where you will spend most of your time when editing Materials. The Material section is where you can change the Material Domain, Blend Mode, and Shading Model along with many more options. |
+| **Translucency**|This is where you can adjust how Translucency in this Material will function. Please note that this is only editable if the Material Blend Mode is set to Translucent. |
+| **Translucency Self Shadowing**|This is where you can adjust how Translucent self shadowing looks and behaves. Please note that this is only editable if the Material Blend Mode is set to Translucent.|
+| **Usage**|This is where you set what type of objects this Material will work with. Usage flags are usually set automatically by the editor. However if you know that this Material is supposed to be used with a certain object type, make sure you enable it here to avoid any future mistakes|
+| **Mobile**|This is where you set how this Material should work when used on a mobile device like a smart phone.|
+| **Tessellation**|This is where you can enable the Material to make use of Hardware Tessellation.|
+| **Post Process Material**|This is where you define how this Material should work with Post Process and Tone Mapping. Please note that this is only editable if the Material Domain is set to Post Process.|
+| **Lightmass**|This is where you can adjust how this Material will interact with Lightmass.|
+| **Material Interface**|This is where you can define a Static Mesh to preview the Material with.|
+| **Thumbnail**|This controls how the Thumbnail that you see in the **Content Browser** is displayed.|
 
 
 ## Using the Main Material Node
@@ -195,7 +197,7 @@ Setting up a Material to use the Main Material Node is something that can be acc
 
 The Main Material node plays a huge role in how Materials will be displayed and what inputs that Material can use.
 No matter how many Material Expression nodes a Material uses the effects of those nodes will not be seen until they are input into the Main Material node.
-Also keep in mind that if the Material is not doing what you want it to do, make sure that you have the correct options set in the **Details** panel.
+Also keep in mind that if the Material is not doing what you want it to do, make sure that you have the correct options set in the Details panel.
 
 
 
