#UE4doc #docs - Added Engine Version Field. Added Skill Level field. Bolded/unbolded terms.

diff --git a/Engine/Documentation/Source/Engine/Rendering/Materials/HowTo/Making_Functions/Making_Functions.INT.udn b/Engine/Documentation/Source/Engine/Rendering/Materials/HowTo/Making_Functions/Making_Functions.INT.udn
--- a/Engine/Documentation/Source/Engine/Rendering/Materials/HowTo/Making_Functions/Making_Functions.INT.udn
+++ b/Engine/Documentation/Source/Engine/Rendering/Materials/HowTo/Making_Functions/Making_Functions.INT.udn
@@ -8,13 +8,16 @@ Related: Engine/Rendering/Materials/Functions/Editing
 Related: Engine/Rendering/Materials/Functions/Using
 Related: Engine/Rendering/Materials/ExpressionReference
 Related: Engine/Rendering/Materials
+Version: 4.9
+SkillLevel: Intermediate
+
 
 [TOC(start:2)]
 
 
 The ability to reuse part of one Material Graph inside another Material Graph can greatly speed up both Material creation time and Material workflow. 
-Unreal Engine 4 (UE4) offers users the ability to do this using what is known as a Material Function.
-In this How To example, we will be re-creating the Fresnel style shading effect that you often see inside of Video Games using a Material Function. 
+Unreal Engine 4 (UE4) offers users the ability to do this using what is known as a **Material Function**.
+In this How-to example, we will be re-creating the Fresnel style shading effect that you often see inside of Video Games using a Material Function. 
 The Fresnel style effect is most commonly used for highlighting the edges of object to give them a fake rim light and also helps to make the object's silhouette better visible in the dark. 
 
 [region:tip]
@@ -100,7 +103,7 @@ Creating Material Functions can be accomplished in the few following steps.
 	
 	![](Material_Function_Adding_Clamp_Node.png)
 
-1. After the Clamp node, we need to add a [One-Minus node](Engine/Rendering/Materials/ExpressionReference/Math#oneminus) to invert our effect so that only the edges are effected. You can find the One-Minus node by searching for OneMinus in the Palette. When completed, you should have something that looks like the following. 
+1. After the Clamp node, we need to add a [One-Minus node](Engine/Rendering/Materials/ExpressionReference/Math#oneminus) to invert our effect so that only the edges are effected. You can find the One-Minus node by searching for **OneMinus** in the **Palette**. When completed, you should have something that looks like the following. 
 	
 	![](Material_Function_Adding_OneMinus.png)
 
@@ -129,7 +132,7 @@ Creating Material Functions can be accomplished in the few following steps.
 	
 	![](Material_Function_Adding_Power.png)
 
-1. Next, we are going to add a Multiply Node so that we can control the intensity of the Fresnel effect.  To add a Multiply Node search for it in the **Palette** and then Drag and Drop it into the Material Graph. When completed, you should have something that looks like this. 
+1. Next, we are going to add a Multiply Node so that we can control the intensity of the Fresnel effect.  To add a **Multiply** Node search for it in the **Palette** and then Drag and Drop it into the Material Graph. When completed, you should have something that looks like this. 
 	
 	![](Material_Function_Adding_Multiply.png)
 
@@ -182,7 +185,7 @@ Creating Material Functions can be accomplished in the few following steps.
 	[/PARAMLITERAL]
 	[/OBJECT]
 
-1. For this example, we are going to add the following nodes to our Material so that we have a bit better looking preview. Remember to find the nodes listed below search for the nodes' name in the **Palette** menu and then Drag and Drop the node from the **Palette** menu right into the Material Graph. When completed, you should have something that looks like this. 
+1. For this example, we are going to add the following nodes to our Material so that we have a bit better looking preview. Remember to find the nodes listed below search for the nodes' name in the **Palette** menu and then Drag and Drop the node from the Palette menu right into the Material Graph. When completed, you should have something that looks like this. 
 	
 	* **[Constant 3 Vector Node](Engine/Rendering/Materials/ExpressionReference/Constant#constant3vector)** x 1 
 
