#UE4doc #docs - Added Engine Version Field. Added Skill Level field. Bolded/unbolded terms.

diff --git a/Engine/Documentation/Source/Engine/Rendering/Materials/HowTo/Masking/Masking.INT.udn b/Engine/Documentation/Source/Engine/Rendering/Materials/HowTo/Masking/Masking.INT.udn
--- a/Engine/Documentation/Source/Engine/Rendering/Materials/HowTo/Masking/Masking.INT.udn
+++ b/Engine/Documentation/Source/Engine/Rendering/Materials/HowTo/Masking/Masking.INT.udn
@@ -8,6 +8,8 @@ Related: Engine/Content/Types/Textures/Properties/Interface
 Related: Engine/Rendering/Materials/ExpressionReference/Textures
 Related: Engine/Rendering/Materials/HowTo/EmissiveGlow
 Related: Engine/Rendering/Materials/HowTo/Instancing
+Version: 4.9
+SkillLevel: Intermediate
 
 
 [TOC(start:2 end:3)]
@@ -15,12 +17,12 @@ Related: Engine/Rendering/Materials/HowTo/Instancing
 
 
 When creating 3D assets, you might find that you need the ability to define different surfaces types within the same Material. 
-A simple and cheap way to achieve this is to use a Texture Mask that defines which surface should be affected by which section of the Material. 
-In the following How To, we will cover how you can use Texture Masking inside of your Unreal Engine 4 (UE4) Materials.
+A simple and cheap way to achieve this is to use a **Texture Mask** that defines which surface should be affected by which section of the Material. 
+In the following How-to, we will cover how you can use Texture Masking inside of your Unreal Engine 4 (UE4) Materials.
 
 ## Texture Masks
 
-A texture Mask is a grayscale texture, or a single channel (R, G, B, or A) of a texture, used to limit the area of an effect inside of a Material. 
+A Texture Mask is a grayscale texture, or a single channel (R, G, B, or A) of a texture, used to limit the area of an effect inside of a **Material**. 
 Quite often, a mask will be contained within a single channel of another texture, such as the **Alpha Channel** of the Diffuse or Normal map. 
 This is a good way to make use of unused channels and keep the number of textures being sampled in the Material to a minimum. 
 Technically, any channel of any texture can be thought of and used as a texture mask. 
@@ -83,7 +85,7 @@ The other most common thing to do with a Mask Texture is to store different kind
 A perfect example of this technique can be found by looking at the SM_Chair Static Mesh that comes with the Starter Content. 
 
 [region:tip]
-You can find SM_Chair and all the content that goes with it by first selecting the **Starter Content** folder in the **Content Browser** and then entering chair into the search box. 
+You can find SM_Chair and all the content that goes with it by first selecting the **Starter Content** folder in the **Content Browser** and then entering "chair" into the search box. 
 This will display all the content that is related to the chair. If you do not see the chair, it probably means that you did not include the Starter Content with your project.
 To fix this, you will need to either create a new project or try the [Migrating Assets](Engine/Content/Browser/UserGuide/Migrate) tool to move the chair content from another project 
 into this one.
