#UE4doc #docs - Added Engine version field. Unbolded hyperlinks.

diff --git a/Engine/Documentation/Source/Programming/UnrealArchitecture/Reference/ProgrammingReference.INT.udn b/Engine/Documentation/Source/Programming/UnrealArchitecture/Reference/ProgrammingReference.INT.udn
--- a/Engine/Documentation/Source/Programming/UnrealArchitecture/Reference/ProgrammingReference.INT.udn
+++ b/Engine/Documentation/Source/Programming/UnrealArchitecture/Reference/ProgrammingReference.INT.udn
@@ -2,6 +2,7 @@ Availability:Public
 Title: Unreal Architecture
 Crumbs:%ROOT%, Programming
 Description:Reference to creating and implementing gameplay classes.
+Version: 4.9
 
 When programming gameplay elements using C++ code, each module can contain many C++ classes. 
 
@@ -9,15 +10,15 @@ When programming gameplay elements using C++ code, each module can contain many
 ![](Programming/Basics/ProjectModuleClassOrg.png)
 [/REGION]
 
-Each class defines a template for a new Actor or Object. Within the class header file, the class and any class [**functions**](Programming\UnrealArchitecture\Reference\Functions) and [**properties**](Programming\UnrealArchitecture\Reference\Properties)
-are declared. Classes can also contain [**structs**](Programming\UnrealArchitecture\Reference\Structs), data structures that help with organization and manipulation of related properties. Structures can also be defined on their own.
-[**interface**](Programming\UnrealArchitecture\Reference\Interfaces) allow additional gameplay behavior to be implemented by different classes. 
+Each class defines a template for a new Actor or Object. Within the class header file, the class and any class [functions](Programming\UnrealArchitecture\Reference\Functions) and [properties](Programming\UnrealArchitecture\Reference\Properties)
+are declared. Classes can also contain [structs](Programming\UnrealArchitecture\Reference\Structs), data structures that help with organization and manipulation of related properties. Structures can also be defined on their own.
+[interface](Programming\UnrealArchitecture\Reference\Interfaces) allow additional gameplay behavior to be implemented by different classes. 
 
 When programming with Unreal Engine, it is possible to have standard C++ classes, functions, and variables. These can be defined using standard C++ syntax. However, `UCLASS()`, `UFUNCTION()`, and
 `UPROPERTY()` macros can be used to make Unreal Engine aware of the new classes, functions, and variables. For instance, a variable with a declaration prefaced by a `UPROPERTY()` macro can 
 be garbage collected by the engine, and can be displayed and edited within Unreal Editor. There are also `UINTERFACE()` and `USTRUCT()` macros, and keywords for each macro that can be used
-to specify the [**class**](Programming\UnrealArchitecture\Reference\Classes\Specifiers)'s, [**function**](Programming\UnrealArchitecture\Reference\Functions\Specifiers)'s, [**property**](Programming\UnrealArchitecture\Reference\Properties\Specifiers)'s,
- [**interface**](Programming\UnrealArchitecture\Reference\Interfaces\Specifiers)'s, or [**struct**](Programming\UnrealArchitecture\Reference\Structs\Specifiers)'s behavior within Unreal Engine and Unreal Editor.
+to specify the [class](Programming\UnrealArchitecture\Reference\Classes\Specifiers)'s, [function](Programming\UnrealArchitecture\Reference\Functions\Specifiers)'s, [property](Programming\UnrealArchitecture\Reference\Properties\Specifiers)'s,
+ [interface](Programming\UnrealArchitecture\Reference\Interfaces\Specifiers)'s, or [struct](Programming\UnrealArchitecture\Reference\Structs\Specifiers)'s behavior within Unreal Engine and Unreal Editor.
 
 [OBJECT:TopicButtonList]
 	[PARAM:icon]
