#UE4doc #docs - Added Engine version field.

diff --git a/Engine/Documentation/Source/Programming/UnrealArchitecture/Reference/Classes/Specifiers/ClassSpecifiers.INT.udn b/Engine/Documentation/Source/Programming/UnrealArchitecture/Reference/Classes/Specifiers/ClassSpecifiers.INT.udn
--- a/Engine/Documentation/Source/Programming/UnrealArchitecture/Reference/Classes/Specifiers/ClassSpecifiers.INT.udn
+++ b/Engine/Documentation/Source/Programming/UnrealArchitecture/Reference/Classes/Specifiers/ClassSpecifiers.INT.udn
@@ -2,6 +2,7 @@ Title: Class Specifiers
 Crumbs:%ROOT%, Programming/UnrealArchitecture/Reference,  Programming/UnrealArchitecture/Reference, Programming/UnrealArchitecture/Reference/Classes
 Description: Keywords used when declaring gameplay classes to specify how the class behaves with various aspects of the engine and editor.
 Availability:Public
+Version: 4.9
 
 [EXCERPT:main]
 When declaring classes, specifiers can be added to the declaration to control how the class behaves with various 
