#DocFix

Upadting the steps on the Tutorial (Forums)

diff --git a/Engine/Documentation/Source/Engine/UMG/QuickStart/4/UMGQS_4.INT.udn b/Engine/Documentation/Source/Engine/UMG/QuickStart/4/UMGQS_4.INT.udn
--- a/Engine/Documentation/Source/Engine/UMG/QuickStart/4/UMGQS_4.INT.udn
+++ b/Engine/Documentation/Source/Engine/UMG/QuickStart/4/UMGQS_4.INT.udn
@@ -50,7 +50,11 @@ Time to get our hands dirty and script the functionality of our Main Menu. At th
 
 1.	Hold **Ctrl** and drag in the **MainMenu** and **OptionsMenu** variables. 
 
-1.	Drag off **MainMenu** and **Set Visibility** (set to **Hidden**) then **Set Visibility** off **OptionsMenu** to **Visible**.  
+1.	Drag off **MainMenu** and add a **Set Visibility** node. Set **In Visibility** to Hidden.
+
+1.  Drag off and add a **Set Visibility** node. Set **Invisibility** to Visible
+
+1.  Attach target to the **OptionsMenu** variable.
 
 	![](OptionsButton.png)(w:820)
 
