#UE4doc #docs - Capitalized single instance of Blueprint

diff --git a/Engine/Documentation/Source/Programming/UnrealArchitecture/Reference/Classes/Specifiers/Blueprintable/ClassSpecifier_Blueprintable.INT.udn b/Engine/Documentation/Source/Programming/UnrealArchitecture/Reference/Classes/Specifiers/Blueprintable/ClassSpecifier_Blueprintable.INT.udn
--- a/Engine/Documentation/Source/Programming/UnrealArchitecture/Reference/Classes/Specifiers/Blueprintable/ClassSpecifier_Blueprintable.INT.udn
+++ b/Engine/Documentation/Source/Programming/UnrealArchitecture/Reference/Classes/Specifiers/Blueprintable/ClassSpecifier_Blueprintable.INT.udn
@@ -4,4 +4,4 @@ Description:
 Availability:Public
 Version: 4.9
 
-Exposes this class as an acceptable base class for creating blueprints. The default is NotBlueprintable, unless inherited otherwise. This is inherited by subclasses.
+Exposes this class as an acceptable base class for creating Blueprints. The default is NotBlueprintable, unless inherited otherwise. This is inherited by subclasses.
