#UE4doc #docs - Capitalized single instance of Blueprint

diff --git a/Engine/Documentation/Source/Programming/UnrealArchitecture/Reference/Classes/Specifiers/BlueprintType/ClassSpecifier_BlueprintType.INT.udn b/Engine/Documentation/Source/Programming/UnrealArchitecture/Reference/Classes/Specifiers/BlueprintType/ClassSpecifier_BlueprintType.INT.udn
--- a/Engine/Documentation/Source/Programming/UnrealArchitecture/Reference/Classes/Specifiers/BlueprintType/ClassSpecifier_BlueprintType.INT.udn
+++ b/Engine/Documentation/Source/Programming/UnrealArchitecture/Reference/Classes/Specifiers/BlueprintType/ClassSpecifier_BlueprintType.INT.udn
@@ -4,4 +4,4 @@ Description:
 Availability:Public
 Version: 4.9
 
-Exposes this class as a type that can be used for variables in blueprints.
+Exposes this class as a type that can be used for variables in Blueprints.
