#UE4doc #docs - Added Engine Version Field. Removed Overview header.

diff --git a/Engine/Documentation/Source/Engine/Paper2D/Flipbooks/Components/FlipbookComponents.INT.udn b/Engine/Documentation/Source/Engine/Paper2D/Flipbooks/Components/FlipbookComponents.INT.udn
--- a/Engine/Documentation/Source/Engine/Paper2D/Flipbooks/Components/FlipbookComponents.INT.udn
+++ b/Engine/Documentation/Source/Engine/Paper2D/Flipbooks/Components/FlipbookComponents.INT.udn
@@ -2,8 +2,8 @@ Availability:Public
 Title:Flipbook Components
 Crumbs: %ROOT%, Engine, Engine/Paper2D, Engine/Paper2D/Flipbooks
 Description:Information on working with Flipbook Components in Blueprints or C++. 
+Version: 4.9
 
-##Overview
 
 **Flipbook Components** are regular primitive components, and can be posed arbitrarily in 3D, attached to other components, or have other components attached to them. Each Flipbook Component instance can specify a custom color that will be passed down to the Flipbook Material as a Vertex Color. They can also have a custom Material specified that will override the default Material defined in the Flipbook. 
 
diff --git a/Engine/Documentation/Source/Programming/UnrealArchitecture/Reference/Classes/Specifiers/Const/ClassSpecifier_Const.INT.udn b/Engine/Documentation/Source/Programming/UnrealArchitecture/Reference/Classes/Specifiers/Const/ClassSpecifier_Const.INT.udn
--- a/Engine/Documentation/Source/Programming/UnrealArchitecture/Reference/Classes/Specifiers/Const/ClassSpecifier_Const.INT.udn
+++ b/Engine/Documentation/Source/Programming/UnrealArchitecture/Reference/Classes/Specifiers/Const/ClassSpecifier_Const.INT.udn
@@ -4,4 +4,4 @@ Description:
 Availability:Public
 Version: 4.9
 
-All properties and functions in this class are const and should be exported as const. This flag is inherited by subclasses.
+All properties and functions in this class are `const` and should be exported as `const`. This flag is inherited by subclasses.
