#UE4doc #docs - Added Engine Version field. Added Skill Level field. Updated images to use correct color for call out areas.

diff --git a/Engine/Documentation/Source/Engine/Paper2D/Flipbooks/Components/Blueprints/Components_BP.INT.udn b/Engine/Documentation/Source/Engine/Paper2D/Flipbooks/Components/Blueprints/Components_BP.INT.udn
--- a/Engine/Documentation/Source/Engine/Paper2D/Flipbooks/Components/Blueprints/Components_BP.INT.udn
+++ b/Engine/Documentation/Source/Engine/Paper2D/Flipbooks/Components/Blueprints/Components_BP.INT.udn
@@ -6,6 +6,8 @@ Related: Engine/Paper2D/Flipbooks
 Related: Engine/Paper2D/Sprites
 Related: Engine/Paper2D
 Related: Engine/Blueprints
+Version: 4.9
+SkillLevel: Beginner
 
 [REGION:banner]
 ![](Components_Banner.png)
@@ -79,7 +81,7 @@ An example of such is indicated below.
 
 ![](Collision2.png)(w:820)
 
-Above we've created a simple blinking box Blueprint and added a Box Collision component to provide the collision for the box by changing the **Box Extent** values to encapsulate our box and changing the **Collision Presets** to **Block All** to provide the actual collision. 
+Above we've created a simple blinking box Blueprint and added a **Box Collision** component to provide the collision for the box by changing the **Box Extent** values to encapsulate our box and changing the **Collision Presets** to **Block All** to provide the actual collision. 
 
 ## Working with Flipbook Components
 
