#UE4doc #docs - Added Engine Version field. Added Skill Level field. Updated images to use correct color for call out areas.

diff --git a/Engine/Documentation/Source/Engine/Paper2D/Flipbooks/Editor/FlipbookEditor.INT.udn b/Engine/Documentation/Source/Engine/Paper2D/Flipbooks/Editor/FlipbookEditor.INT.udn
--- a/Engine/Documentation/Source/Engine/Paper2D/Flipbooks/Editor/FlipbookEditor.INT.udn
+++ b/Engine/Documentation/Source/Engine/Paper2D/Flipbooks/Editor/FlipbookEditor.INT.udn
@@ -7,6 +7,8 @@ Related: Engine/Paper2D/Flipbooks
 Related: Engine/Paper2D/Sprites
 Related: Engine/Rendering/Materials
 Related: Engine/Content/Types/Textures 
+Version: 4.9
+SkillLevel: Beginner
 
 [TOC (start:2 end:2)]
 
@@ -114,7 +116,7 @@ While it is possible to set collision on a **Flipbook** by using the **Collision
 * Have non-animated collision only applied to the first frame of the animation. 
 * Have the collision changed each frame of the animation (very expensive in terms of performance and not recommended). 
 
-In many cases it may make more sense to attach a Collision Component (for example, a Box Component) to a socket on the Flipbook Component rather than directly generating collision from source frames, especially if you want some collision regions to overlap while others block. 
+In many cases it may make more sense to attach a **Collision Component** (for example, a Box Component) to a socket on the Flipbook Component rather than directly generating collision from source frames, especially if you want some collision regions to overlap while others block. 
 
 Character Movement still relies on the collision capsule, in most cases you want to use a different collision channel for the Flipbook than BlockAll to prevent the capsule from getting stuck on the Flipbook (move/be blocked based on the capsule, and test for overlap or collision with projectiles based on the Flipbook). 
 
