#UE4doc #docs - Added Engine Version field. Added Skill Level field. Updated images to use correct color for call out areas. Removed Overview header.

diff --git a/Engine/Documentation/Source/Engine/Paper2D/Importing/Paper2DImporting.INT.udn b/Engine/Documentation/Source/Engine/Paper2D/Importing/Paper2DImporting.INT.udn
--- a/Engine/Documentation/Source/Engine/Paper2D/Importing/Paper2DImporting.INT.udn
+++ b/Engine/Documentation/Source/Engine/Paper2D/Importing/Paper2DImporting.INT.udn
@@ -5,13 +5,14 @@ Availability:Public
 Related: Engine\Paper2D\Sprites
 Related: Engine\Paper2D\Flipbooks
 Related: Engine\Paper2D\TileMaps
+Version: 4.9
+SkillLevel: Beginner
 
 
 [TOC (start:2 end:2)]
 
-## Overview
 
-**Paper 2D** in Unreal Engine 4 supports a number of different import options for your externally created content. Below are some examples of importing assets from some of the more commonly used content creation applications. 
+Paper 2D in Unreal Engine 4 supports a number of different import options for your externally created content. Below are some examples of importing assets from some of the more commonly used content creation applications. 
 
 
 ## Importing from Texture Packer
