#UE4doc #docs:
Added engine version field
Changed a couple steps to expand on the information provided.

diff --git a/Engine/Documentation/Source/Engine/UMG/QuickStart/4/UMGQS_4.INT.udn b/Engine/Documentation/Source/Engine/UMG/QuickStart/4/UMGQS_4.INT.udn
--- a/Engine/Documentation/Source/Engine/UMG/QuickStart/4/UMGQS_4.INT.udn
+++ b/Engine/Documentation/Source/Engine/UMG/QuickStart/4/UMGQS_4.INT.udn
@@ -2,6 +2,7 @@ Availability: Public
 Title:4. Scripting the Main Menu
 Crumbs:%ROOT%, Engine, Engine/UMG, Engine/UMG/QuickStart
 Description:In this step we will script the functionality for our Main Menu and load into game. 
+version: 4.9
 
 
 [VAR:Steps]
@@ -104,13 +105,13 @@ Time to get our hands dirty and script the functionality of our Main Menu. At th
 
 	![](WorldSettings.png)
 
-1.	Under _Game Mode_, set the **GameMode Override** to **FirstPersonGameMode** and change the **Default Pawn Class** to **Character**. 
+1.	Under **Game Mode**, set the **GameMode Override** to **FirstPersonGameMode** and change the **Default Pawn Class** to **Character**. 
 
 	![](Overrides.png)
 
 	Here we are assigning a Game Mode for this Map and changing the player's character to use the Engine Default Character inside of the playable **FirstPersonCharacter** Blueprint as we do not want the player to be able to run around and shoot behind the Main Menu. 
 
-1.	Click the **Play** Button from the Main Toolbar to play in the editor. 
+1.	Click the **Play** Button from the Main Toolbar to play in the editor. To test the resolution options, select **New Editor Window** from the **Play-in-Editor** dropdown menu.
 
 	[OBJECT:EmbeddedVideo]
 	[PARAMLITERAL:width]
