#UE4doc #docs:
Added engine version field
Changed a step as it was incorrect
Replaced multiple images due to out of date UI
Changed the highlighting in multiple images from Green to Yellow to match the guidelines

diff --git a/Engine/Documentation/Source/Engine/UMG/QuickStart/5/UMGQS_5.INT.udn b/Engine/Documentation/Source/Engine/UMG/QuickStart/5/UMGQS_5.INT.udn
--- a/Engine/Documentation/Source/Engine/UMG/QuickStart/5/UMGQS_5.INT.udn
+++ b/Engine/Documentation/Source/Engine/UMG/QuickStart/5/UMGQS_5.INT.udn
@@ -2,6 +2,7 @@ Availability: Public
 Title:5. Creating / Scripting an In-Game Pause Menu
 Crumbs:%ROOT%, Engine, Engine/UMG, Engine/UMG/QuickStart
 Description:Here we will create an in-game pause menu which will allow us to pause or exit our game. 
+version: 4.9
 
 
 [VAR:Steps]
@@ -122,7 +123,7 @@ With the visuals set, next we will provide the scripted functionality for our Pa
 
 	When **Quit** is clicked, this will turn on the mouse cursor for menu navigation when returning to the Main Menu. 
 
-1.	Next add the **Remove HUD** function created above followed by **Open Level** (set to **MainMenu**) and connect as shown. 
+1.	Next add the **Remove HUD** function created above followed by **Open Level** (set to **Main** or the name of the Level that you have your MainMenu on) and connect as shown. 
 
 	[REGION:lightbox]
 	[![](PauseScript8.png)(w:840)](PauseScript8.png)
