#UE4doc #docs:
Added engine version field
Changed the text for multiple hyperlinks to match the pages they link to.

diff --git a/Engine/Documentation/Source/Engine/UMG/QuickStart/6/UMGQS_6.INT.udn b/Engine/Documentation/Source/Engine/UMG/QuickStart/6/UMGQS_6.INT.udn
--- a/Engine/Documentation/Source/Engine/UMG/QuickStart/6/UMGQS_6.INT.udn
+++ b/Engine/Documentation/Source/Engine/UMG/QuickStart/6/UMGQS_6.INT.udn
@@ -2,6 +2,7 @@ Availability: Public
 Title:6. Further Reading
 Crumbs:%ROOT%, Engine, Engine/UMG, Engine/UMG/QuickStart
 Description:The final step in the UMG Quick Start is to provide additional links to reference documentation.
+version: 4.9
 
 
 [VAR:Steps]
@@ -23,13 +24,13 @@ Description:The final step in the UMG Quick Start is to provide additional links
 
 ---
 
-For more information on the topics covered in this Quick Start Guide, and across the entire editor, see the [Unreal Editor Manual](Engine).
+For more information on the topics covered in this Quick Start Guide, and across the entire editor, see the [Engine Features](Engine) page.
 
 As for coverage related to the topics in this guide:
 
-* For more information on Unreal Motion Graphics, see: [UMG Overview](Engine/UMG)
+* For more information on Unreal Motion Graphics, see: [UMG UI Designer](Engine/UMG)
 
-* For more information on Blueprints, see: [Blueprints Landing Page](Engine\Blueprints)
+* For more information on Blueprints, see: [Blueprints Visual Scripting](Engine\Blueprints)
 
 ---
 
