#UE4doc #docs:
Added engine version field
Replaced mentions of Preview Size with Screen Size as it has changed
Replaced multiple images and recolored a few images due to incorrect UI and incorrect arrows due to doc guidelines

diff --git a/Engine/Documentation/Source/Engine/UMG/UserGuide/Anchors/Anchors.INT.udn b/Engine/Documentation/Source/Engine/UMG/UserGuide/Anchors/Anchors.INT.udn
--- a/Engine/Documentation/Source/Engine/UMG/UserGuide/Anchors/Anchors.INT.udn
+++ b/Engine/Documentation/Source/Engine/UMG/UserGuide/Anchors/Anchors.INT.udn
@@ -5,11 +5,11 @@ Description:Using Anchors to keep UI Widgets in fixed locations based on differe
 Related: Engine/UMG
 Related: Engine/UMG/UserGuide/DPIScaling
 Related: Engine/Blueprints
+version: 4.9
 
 
 [TOC(start:2 end:2)]
 
-## Overview
 
 **Anchors** are used to define a UI widget's desired location on a **Canvas Panel** and maintains that position with varying screen sizes. Anchors are normalized where **Min(0,0)** and **Max(0,0)** is the upper left corner and **Min(1,1)** and **Max(1,1)** is the bottom right corner. 
 
@@ -25,10 +25,10 @@ In the image below, we have a Button placed on a Canvas Panel and we have left t
 
 ![](Button1.png)(w:720)
 
-The solid yellow line above tells the Button how much to move from the Anchor along the Y axis based on the Canvas size, starting in the upper left of the viewport. The dotted yellow line tells the Button how much to move from the Anchor along the X axis, again based on Canvas size and starting from the upper left of the viewport. In the lower left corner of the window (inside the yellow box) you can see the **Preview Size** of the Canvas you are working with. 
+The vertical yellow line above tells the Button how much to move from the Anchor along the Y axis based on the Canvas size, starting in the upper left of the viewport. The horizontal yellow line tells the Button how much to move from the Anchor along the X axis, again based on Canvas size and starting from the upper left of the viewport. In the lower left corner of the window (inside the yellow box) you can see the **Screen Size** of the Canvas you are working with. 
 
 [REGION:tip]
-Click the **Preview Size** Button in the graph to change the currently used size. It is a good habit to check your UI widgets to see how they will appear on different screen sizes or aspect ratios and adjust accordingly. 
+Click the **Screen Size** Button in the graph to change the currently used size. It is a good habit to check your UI widgets to see how they will appear on different screen sizes or aspect ratios and adjust accordingly. 
 [/REGION]
 
 Playing in game, depending upon viewport size, we might see something like below where the yellow box represents the Anchor point: 
