#UE4 Docs: minor tweaks to text and formatting

diff --git a/Engine/Documentation/Source/Platforms/Android/GettingStarted/1/AndroidQS_1.INT.udn b/Engine/Documentation/Source/Platforms/Android/GettingStarted/1/AndroidQS_1.INT.udn
--- a/Engine/Documentation/Source/Platforms/Android/GettingStarted/1/AndroidQS_1.INT.udn
+++ b/Engine/Documentation/Source/Platforms/Android/GettingStarted/1/AndroidQS_1.INT.udn
@@ -24,9 +24,9 @@ Skilllevel:Intermediate
 %Steps%
 
 
-Before we can start developing for Android we first need to install Android Works. 
-Android Works will help to ensure that you have all the required programs and code that is needed for you to be able to deploy to Android devices.
-In the following section we will cover how you can get Android Works installed on your development PC.
+The first thing we need to do to start developing for Android is install Android Works, which will 
+ensure that you have all the required programs and code needed to deploy to Android devices.
+In the following section, we will cover how you can get Android Works installed on your development PC.
 
 
 %Globals:OSSelect%
diff --git a/Engine/Documentation/Source/Platforms/Android/GettingStarted/2/AndroidQS_2.INT.udn b/Engine/Documentation/Source/Platforms/Android/GettingStarted/2/AndroidQS_2.INT.udn
--- a/Engine/Documentation/Source/Platforms/Android/GettingStarted/2/AndroidQS_2.INT.udn
+++ b/Engine/Documentation/Source/Platforms/Android/GettingStarted/2/AndroidQS_2.INT.udn
@@ -65,7 +65,12 @@ Now that Android Works has been installed it is now time to setup your Android d
 
 1.	Open the **Settings** application on your Android device.
 
-1.	Look for the **Developer Options** menu. If it is not available, you need to enable Developer Options on your device. For more about developer options, visit [Android's Using Hardware Devices](http://developer.android.com/tools/device.html) page.
+1.	Tap **Developer Options** to open the Developer Options menu. 
+
+	[REGION:note]
+	If you do not see this item, you need to enable Developer Options on your device. 
+	For more about developer options, visit [Android's Using Hardware Devices](http://developer.android.com/tools/device.html) page.
+	[/REGION]
 
 1.	Scroll down and select **About Phone**. Depending on your Android device, you may have an **About Device**, **About Tablet**, or **About Shield** option instead.
 	[REGION: note]
@@ -131,8 +136,12 @@ Now that Android Works has been installed it is now time to setup your Android d
 
 1. Open the **Settings** application on your Android device.
 
-1. Look for the **Developer Options** menu. If it is not available, you need to enable Developer Options on your device. 
+1. Tap **Developer Options** to open the Developer Options menu. 
+	
+	[REGION:note]
+	If you do not see this item, you need to enable Developer Options on your device. 
 	For more about developer options, visit [Android's Using Hardware Devices](http://developer.android.com/tools/device.html) page.
+	[/REGION]
 
 1. Scroll down and select **About Phone**. Depending on your Android device, you may have an **About Device**, **About Tablet**, or **About Shield** option instead.
 
diff --git a/Engine/Documentation/Source/Platforms/Android/GettingStarted/4/AndroidQS_4.INT.udn b/Engine/Documentation/Source/Platforms/Android/GettingStarted/4/AndroidQS_4.INT.udn
--- a/Engine/Documentation/Source/Platforms/Android/GettingStarted/4/AndroidQS_4.INT.udn
+++ b/Engine/Documentation/Source/Platforms/Android/GettingStarted/4/AndroidQS_4.INT.udn
@@ -30,7 +30,7 @@ In the following section we will be taking a look at how you can deploy a single
 For this section to work you will need to make sure that your Android device is plugged into your development PC via USB. If you have not done this already please do this before going any further.
 [/region]
 
-Before you can deploy your level to your Android device you need to first determining which type of Texture format your Android Device supports.
+To deploy your level to your Android device you need to determining which type of Texture format your Android Device supports.
 This is because different types of Android devices require different types of Texture formats to work correctly. 
 The chart below shows you which Texture formats work with which Android based devices.
 
