#UE4 Docs: reduced image size

diff --git a/Engine/Documentation/Source/Platforms/Android/GettingStarted/5/AndroidQS_5.INT.udn b/Engine/Documentation/Source/Platforms/Android/GettingStarted/5/AndroidQS_5.INT.udn
--- a/Engine/Documentation/Source/Platforms/Android/GettingStarted/5/AndroidQS_5.INT.udn
+++ b/Engine/Documentation/Source/Platforms/Android/GettingStarted/5/AndroidQS_5.INT.udn
@@ -101,7 +101,7 @@ We are currently packaging in **Development** because we are testing this game,
 
 1.	After the .BAT file has finished you should now see your project ready to be used on your Android device.
 
-	![](Game_On_Phone.png)
+	![](Game_On_Phone.png)(w:320)
 
 [/EXCERPT:WindowsSetup]
 
@@ -143,7 +143,7 @@ We are currently packaging in **Development** because we are testing this game,
 
 1. After the .command file has finished you should now see your project ready to be used on your Android device.
 
-	![](Game_On_Phone.png)
+	![](Game_On_Phone.png)(w:320)
 
 [/EXCERPT:MacOSSetup]
 -->
