#UE4doc #docs:
Added engine version field
Overhauled almost all images on page due to out of date UI and incorrect usage of nodes
Changed a few steps to reflect the changed images

diff --git a/Engine/Documentation/Source/Engine/UMG/UserGuide/Animation/Animation.INT.udn b/Engine/Documentation/Source/Engine/UMG/UserGuide/Animation/Animation.INT.udn
--- a/Engine/Documentation/Source/Engine/UMG/UserGuide/Animation/Animation.INT.udn
+++ b/Engine/Documentation/Source/Engine/UMG/UserGuide/Animation/Animation.INT.udn
@@ -2,6 +2,7 @@ Availability:Public
 Title: Animation
 Crumbs:%ROOT%, Engine, Engine/UMG, Engine/UMG/UserGuide
 Description:How to create animated UI elements in UMG. 
+version: 4.9
 
 
 [REGION:fullwidth raw]
@@ -15,7 +16,7 @@ Located along the bottom of the **Widget Blueprint Editor** are two windows whic
 
 ## Adding Animations
 
-In order to get started with animations in UMG, you will need to first add an animation track. This can be done by clicking the **New** button in the **Animations** window. Once you add an animation (yellow box below), you will be prompted to enter a name for the animation track (green arrow below). 
+In order to get started with animations in UMG, you will need to first add an animation track. This can be done by clicking the **+Animation** button in the **Animations** window. Once you add an animation (yellow box below), you will be prompted to enter a name for the animation track. 
 
 ![](Animation2.png)
 
@@ -23,7 +24,7 @@ After adding an animation track, the **Timeline** will become active and you can
 
 ## Adding Animation Keys
 
-There are two methods in which you can add Keys to an animation track. The first is by using the **Auto Key** checkbox located inside the **Timeline** window (green arrow below). This will automatically add Keys to the Timeline when you make adjustments to a value that supports Keyframing. 
+There are two methods in which you can add Keys to an animation track. The first is by using the **Auto Key** checkbox located inside the **Timeline** window (yellow arrow below). This will automatically add Keys to the Timeline when you make adjustments to a value that supports Keyframing. 
 
 ![](Animation3.png)
 
@@ -43,7 +44,7 @@ In the image above, the icon next to each of the settings can be clicked on and
 
 ## Changing Key Values
 
-You can change the value for a Key (or all keys) at a particular time on the Timeline by holding **Control** and clicking on the Key. 
+You can change the value for multiple Keys at a particular time on the Timeline by holding **Ctrl** and clicking on the Keys. 
 
 ![](Animation5.png)
 
@@ -95,13 +96,13 @@ R_8nsjHZ9ak
 
 	![](Animation2.png)
 
-1.	Click the **Auto Key** button in the **Timeline Window** (the name of your Animation should also appear in the **Timeline**). 
+1.	Click the **Auto Key** button in the **Timeline Window**. 
 
 	![](Animation3.png)
 
 1.	Move the **Timeline Bar** to **0.00** on the Timeline, then click on the **Vertical Box** containing your Button.  
 
-1.	Holding **Control**, click in the center of the Medallion Widget (green arrow) and drag the **Vertical Box** off screen (1). 
+1.	Holding **Control**, click in the center of the Medallion Widget (yellow arrow) and drag the **Vertical Box** off screen (1). 
 
 	![](Animation4.png)
 
@@ -111,7 +112,7 @@ R_8nsjHZ9ak
 
 	![](Animation3b.png)
 
-1.	Move the **Timeline Bar** to **1.00** then holding **Control**, click and move the **Vertical Box** up so it is fully visible. 
+1.	Move the **Timeline Bar** to **1.00** then holding **Ctrl**, click and move the **Vertical Box** up so it is fully visible. 
 
 1.	In the **Details** panel for the **Vertical Box**, click **Anchors** and select the **Center** option. 
 
