#UE4doc #docs:
Added engine version field
Replaced all images as they were using out of date UI/Nodes
Changed a step as it was referring to an incorrect method
Changed font settings for multiple terms

diff --git a/Engine/Documentation/Source/Engine/UMG/UserGuide/CreatingWidgets/CreatingWidgets.INT.udn b/Engine/Documentation/Source/Engine/UMG/UserGuide/CreatingWidgets/CreatingWidgets.INT.udn
--- a/Engine/Documentation/Source/Engine/UMG/UserGuide/CreatingWidgets/CreatingWidgets.INT.udn
+++ b/Engine/Documentation/Source/Engine/UMG/UserGuide/CreatingWidgets/CreatingWidgets.INT.udn
@@ -7,9 +7,9 @@ Once you have created and laid out your **Widget Blueprint**, in order for it to
 
 ![](CreateWidgetNode.png)(w:860)
 
-In the example above, the **Create Widget** node is calling the **Widget Blueprint** specified under the **Class** section (**Owning Player** is the **Player Controller** it is applied to where blank results in the default Player Controller) and the **Return Value** is the resulting product. 
+In the example above, the **Create Widget** node is calling the Widget Blueprint specified under the **Class** section (**Owning Player** is the **Player Controller** it is applied to where blank results in the default Player Controller) and the **Return Value** is the resulting product. 
 
-Also in the example above, the Return Value is being assigned to a variable titled _Main Menu_ which allows us to access the Widget Blueprint later instead of re-creating the Widget (and this can be used to remove the Widget). Finally, the **Add to Viewport** is what draws the **Widget Blueprint** onscreen and we are specifying that the _Main Menu_ variable (which holds the created Widget) be the Target that is added.
+Also in the example above, the Return Value is being assigned to a variable titled Main Menu which allows us to access the Widget Blueprint later instead of re-creating the Widget (and this can be used to remove the Widget). Finally, the **Add to Viewport** is what draws the Widget Blueprint onscreen and we are specifying that the Main Menu variable (which holds the created Widget) be the Target that is added.
 
 
 [REGION:note]
@@ -29,7 +29,7 @@ There may be instances when you want the player to interact with your UI and the
 * The center node above, **Set Input Mode Game Only** enables input only for the game and ignores UI elements (this is perfect for non-interactive UI elements like Health, Points, or Time display).
 * The right node above, **Set Input Mode UIOnly** is an edge case node and is used when you only want to allow for UI navigation only and no game input. This completely disables all game controls and UI consumes all input, use this node sparingly.  
 
-To supplement the nodes above, you may want to **Enable/Disable** the **Mouse Cursor** from being displayed. You can do so by using the **Set Show Mouse Cursor** node. Drag off a **Get Player Controller** node, then use the **Set Show Mouse Cursor** node and check it to _True_ or _False_ to Display or Hide the Mouse Cursor. 
+To supplement the nodes above, you may want to **Enable/Disable** the **Mouse Cursor** from being displayed. You can do so by using the **Set Show Mouse Cursor** node. Drag off a **Get Player Controller** node, then use the **Set Show Mouse Cursor** node and check it to True or False to display or hide the Mouse Cursor. 
 
 ![](ShowCursor.png)
 
@@ -42,7 +42,7 @@ Instead of using the **Add to Viewport** node for each Widget, it is possible to
 
 ![](AddChild.png)
 
-The example above shows how a **Scroll Box** Widget called _Player List_ is appended to a new Widget called _Main Menu_ using the **Add Child** function. 
+The example above shows how a **Scroll Box** Widget called start_Button is appended to a new Widget called Main Menu using the **Add Child** function. 
 
 [REGION:note]
 The **Add Child** node is for parenting one widget to another in a panel while the **Add to Viewport** adds the widget to the root viewport as if it is a new window. In order to remove a child widget, you need to get your Parent and call **RemoveChild**.
