#UE4doc #docs:
Added engine version field
Replaced a few images as the UI was out of date
Changed a line of text as it referred to the incorrect section to find DPI Scaling in

diff --git a/Engine/Documentation/Source/Engine/UMG/UserGuide/DPIScaling/Scaling.INT.udn b/Engine/Documentation/Source/Engine/UMG/UserGuide/DPIScaling/Scaling.INT.udn
--- a/Engine/Documentation/Source/Engine/UMG/UserGuide/DPIScaling/Scaling.INT.udn
+++ b/Engine/Documentation/Source/Engine/UMG/UserGuide/DPIScaling/Scaling.INT.udn
@@ -2,12 +2,13 @@ Availability:Public
 Title: DPI Scaling
 Crumbs:%ROOT%, Engine, Engine/UMG, Engine/UMG/UserGuide
 Description:How to access and adjust DPI Scaling Rules. 
+version: 4.9
 
 [REGION:banner]
 ![](DPIscaling.png)(convert:false)
 [/REGION]
 
-**UMG** supports automatic scaling for resolution-independent UI. There are default DPI scaling rules applied to every project which configured as you see fit from inside the **Project Settings** menu under the **Rendering** section. 
+**UMG** supports automatic scaling for resolution-independent UI. There are default DPI scaling rules applied to every project which configured as you see fit from inside the **Project Settings** menu under the **User Interface** section. 
 
 ![](ProjectSettings.png)(w:820)
 
@@ -18,7 +19,7 @@ For the **DPI Scale Rule**, you can set this to one of four options:
 * **Horizontal** - This evaluates the scale curve based on the X axis of the viewport. 
 * **Vertical** - This option evaluates the scale curve based on the Y axis of the viewport. 
 
-You can add points to the curve by **Right-Clicking** and selecting the **Add Key** option. When clicking on a Key on the curve, you can set the resolution and its corresponding scale value via the input boxes. Also for the curve, you can supply an external **Float Curve** or create a Float Curve based on the currently applied settings by expanding the **DPI Curve** option. 
+You can add points to the curve by **Right-clicking** and selecting the **Add Key** option. When clicking on a Key on the curve, you can set the resolution and its corresponding scale value via the input boxes. Also for the curve, you can supply an external **Float Curve** or create a Float Curve based on the currently applied settings by expanding the **DPI Curve** option. 
 
 ![](ExternalCurve.png)
 
