#UE4doc #docs:
Resubmitting as not all changes were saved

diff --git a/Engine/Documentation/Source/Engine/UMG/UserGuide/CreatingWidgets/CreatingWidgets.INT.udn b/Engine/Documentation/Source/Engine/UMG/UserGuide/CreatingWidgets/CreatingWidgets.INT.udn
--- a/Engine/Documentation/Source/Engine/UMG/UserGuide/CreatingWidgets/CreatingWidgets.INT.udn
+++ b/Engine/Documentation/Source/Engine/UMG/UserGuide/CreatingWidgets/CreatingWidgets.INT.udn
@@ -2,6 +2,7 @@ Availability:Public
 Title: Creating Widgets
 Crumbs:%ROOT%, Engine, Engine/UMG, Engine/UMG/UserGuide
 Description:How to create and display a Widget Blueprint in game. 
+version: 4.9
 
 Once you have created and laid out your **Widget Blueprint**, in order for it to be displayed in-game, you will need call it by using the **Create Widget** and **Add to Viewport** nodes inside another Blueprint (**Level Blueprint** or a **Character Blueprint** for example). 
 
