#UE4doc #docs:
Added engine version field
Replaced an image as it was using out of date UI
Changed font settings for a few terms

diff --git a/Engine/Documentation/Source/Engine/UMG/UserGuide/Events/Events.INT.udn b/Engine/Documentation/Source/Engine/UMG/UserGuide/Events/Events.INT.udn
--- a/Engine/Documentation/Source/Engine/UMG/UserGuide/Events/Events.INT.udn
+++ b/Engine/Documentation/Source/Engine/UMG/UserGuide/Events/Events.INT.udn
@@ -2,13 +2,14 @@ Availability:Public
 Title: Events
 Crumbs:%ROOT%, Engine, Engine/UMG, Engine/UMG/UserGuide
 Description:A guide to using Events in UMG. 
+version: 4.9
 
 There are a couple methods in which you can call **Events** in **UMG** which are outlined below. 
 
 ### Bindable Events
-**Bindable Events** are a way for **UMG** to mimic the behavior currently used by Slate which needs a single handler to tell it if the Event was handled. You can bind a function in your **Widget Blueprint** to the event from the **Details** panel under the **Events** section (shown below indicated by the green arrows). 
+**Bindable Events** are a way for UMG to mimic the behavior currently used by Slate which needs a single handler to tell it if the Event was handled. You can bind a function in your **Widget Blueprint** to the event from the **Details** panel under the **Events** section (shown below indicated by the yellow arrows). 
 
-![](EventBinding.png)(w:880)
+![](EventBinding.png)
 
 
 Indicated above by the yellow box, some Widgets supplement **Events** through the handling of **Interaction**. For the example above, in addition to the **OnClicked** Event for a Button Widget, you are able to specify the handling of click events through setting the **Click Method** or **Touch Method**. You can also specify through the **IsFocusable** option if a button should only be mouse-clickable and never keyboard focusable. 
