#UE4doc #docs - Removed bolding from term.

diff --git a/Engine/Documentation/Source/Programming/UnrealArchitecture/Reference/Classes/Specifiers/Within/ClassSpecifier_Within.INT.udn b/Engine/Documentation/Source/Programming/UnrealArchitecture/Reference/Classes/Specifiers/Within/ClassSpecifier_Within.INT.udn
--- a/Engine/Documentation/Source/Programming/UnrealArchitecture/Reference/Classes/Specifiers/Within/ClassSpecifier_Within.INT.udn
+++ b/Engine/Documentation/Source/Programming/UnrealArchitecture/Reference/Classes/Specifiers/Within/ClassSpecifier_Within.INT.udn
@@ -4,7 +4,7 @@ Description:
 Availability:Public
 Version: 4.9
 
-The **Within** class specifier indicates that objects of this class cannot exist _outside_ of an instance of 
+The `Within` class specifier indicates that objects of this class cannot exist _outside_ of an instance of 
 *ClassName*. In order to create an object of this class, you must specify an instance of *ClassName* as the 
 `Outer` object.
 
