#UE4doc #docs  - Minor edit. Added spacing.

diff --git a/Engine/Documentation/Source/Programming/UnrealArchitecture/Reference/Classes/Specifiers/NotBlueprintable/ClassSpecifier_NotBlueprintable.INT.udn b/Engine/Documentation/Source/Programming/UnrealArchitecture/Reference/Classes/Specifiers/NotBlueprintable/ClassSpecifier_NotBlueprintable.INT.udn
--- a/Engine/Documentation/Source/Programming/UnrealArchitecture/Reference/Classes/Specifiers/NotBlueprintable/ClassSpecifier_NotBlueprintable.INT.udn
+++ b/Engine/Documentation/Source/Programming/UnrealArchitecture/Reference/Classes/Specifiers/NotBlueprintable/ClassSpecifier_NotBlueprintable.INT.udn
@@ -4,4 +4,4 @@ Description:
 Availability:Public
 Version: 4.9
 
-Specifies that this class is NOT an acceptable base class for creating blueprints. The default is NotBlueprintable, unless inherited otherwise. This is inherited by subclasses.
+Specifies that this class is NOT an acceptable base class for creating blueprints. The default is `NotBlueprintable`, unless inherited otherwise. This is inherited by subclasses.
