#UE4 Docs: Jira UEDOC-2312 Animation: Driven Key Controller

diff --git a/Engine/Documentation/Source/Engine/Animation/NodeReference/SkeletalControls/SkeletalControls.INT.udn b/Engine/Documentation/Source/Engine/Animation/NodeReference/SkeletalControls/SkeletalControls.INT.udn
--- a/Engine/Documentation/Source/Engine/Animation/NodeReference/SkeletalControls/SkeletalControls.INT.udn
+++ b/Engine/Documentation/Source/Engine/Animation/NodeReference/SkeletalControls/SkeletalControls.INT.udn
@@ -101,6 +101,73 @@ Scale - from one bone to another.
 | **Rotation** | A Rotator specifying the Rotation to apply to the target bone. |
 | **Scale** | A Vector specifying the Scale to apply to the target bone. |
 
+
+## Bone Driven Controller
+
+The **Bone Driven Controller** allows a "Driver" bone to dynamically affect the motion of one or more "Driven" bones. This is useful for characters that have accessories attached as you can use this to avoid geometry from intersecting during an animation, even when a lot of blending is being used. 
+
+You can either set the "Driver" value directly with a **Multiplier**, remap it to a whole new range, or use a **Curve** asset to drive the motion. Using a Curve is usually the best approach as it lest you define the response naturally and interactively tweak points/tangents while seeing the changes reflected in real-time.  
+
+In the example below, the attached accessory (green) has no authored animations, and is being driven in two axes as a function of the character's thigh bone. This is all calculated at runtime, so a blend of animations ends up working quite well here, even without hand-authored tweaking. 
+
+[OBJECT:EmbeddedVideo]
+[PARAMLITERAL:width]
+640
+[/PARAMLITERAL]
+[PARAMLITERAL:height]
+360
+[/PARAMLITERAL]
+[PARAMLITERAL:videoid]
+CDG4WcQ7cIY
+[/PARAMLITERAL]
+[/OBJECT]
+
+![Bone Driven Controller](BoneDrivenController.png)
+
+| Property | Description |
+| -------- | ----------- |
+| [REGION:tablesection]**Source (Driver)**[/REGION] ||
+| **Source Bone** | The Bone to use as the controller input. |
+| **Source Component** | The Transform component to use as input.  |
+| [REGION:tablesection]**Mapping**[/REGION] ||
+| **Driving Curve** | Curve used to map from the source attribute to the driven attributes if present (otherwise the Multiplier will be used). |
+| **Remap Source** | Whether or not to clamp the driver value and remap it before scaling it. |
+| **Multiplier** | Multiplier to apply to the input value (Ignored when a Curve is used). |
+| [REGION:tablesection]**Destination (Driven)**[/REGION] ||
+| **Target Bone** | Bone to drive using controller input. |
+| **Target Component Space** | The Parent Bone Space used. [INCLUDE:#LocationSpace] |
+| **Modification Mode** | Set the type of modification to make to the destination component(s). [INCLUDE:#ModificationSettings]  |
+| [REGION:tablesection]**Settings**[/REGION] ||
+| **(As Pin) Alpha** | Enable/Disable the display of the Alpha pin on the node. |
+| **Alpha Scale Bias** | Set the minimal and maximal input scaling values. |
+
+
+<!--
+[EXCERPT:LocationSpace]
+[REGION:simpletable]
+| Value | Effect |
+| ----- | ------ |
+| **Translation** | The XYZ component of translation on the target bone. |
+| **Rotation** | The XYZ component of rotation on the target bone. |
+| **Scale** | The XYZ component of scale on the target bone. |
+[/REGION] 
+[/EXCERPT:LocationSpace]
+-->
+
+<!--
+[EXCERPT:ModificationSettings]
+[REGION:simpletable]
+| Value | Effect |
+| ----- | ------ |
+| **Add to Input** | Add the driven value to the input component value(s). |
+| **Replace Component** | Replace the input component value(s) with the driven value. |
+| **Add to Ref Pose** | Add the driven value to the reference pose value. |
+[/REGION] 
+[/EXCERPT:ModificationSettings]
+-->
+
+
+
 ## Spring Controller
 
 The **Spring Controller** applies a spring solver that can be used to limit how far a bone can stretch from 
@@ -142,19 +209,6 @@ The **Two Bone IK** control applies an Inverse Kinematic (IK) solver to a 3-join
 | **Joint Target Location Space** | The reference frame for the Joint Target Location. [INCLUDE:#LocationSpace] |
 | **Joint Target Space Bone Name** | The name of the bone to use as the joint target location when Joint Target Location Space is set to BCS_BoneSpace. |
 
-<!--
-[EXCERPT:LocationSpace]
-[REGION:simpletable]
-| Value | Effect |
-| ----- | ------ |
-| **BCS_WorldSpace** | The location is assumed to be relative to world space. |
-| **BCS_ComponentSpace** | The location is assumed to be relative to the SkeletalMeshComponent. |
-| **BCS_ParentSpace** | The location is assumed to be relative to the parent bone. |
-| **BCS_BoneSpace** | The location is assumed to be relative to an arbitrary bone (specified by name as the value of another property). |
-[/REGION] 
-[/EXCERPT:LocationSpace]
--->
-
 | Pin | Description |
 | --- | ----------- |
 | [REGION:tablesection]**Input Pins**[/REGION] ||
