added TOC.

diff --git a/Engine/Documentation/Source/Engine/Animation/NodeReference/SkeletalControls/SkeletalControls.INT.udn b/Engine/Documentation/Source/Engine/Animation/NodeReference/SkeletalControls/SkeletalControls.INT.udn
--- a/Engine/Documentation/Source/Engine/Animation/NodeReference/SkeletalControls/SkeletalControls.INT.udn
+++ b/Engine/Documentation/Source/Engine/Animation/NodeReference/SkeletalControls/SkeletalControls.INT.udn
@@ -17,6 +17,8 @@ playing a generic walk animation. Any sort of modification can be used to tweak
 Transforms applied by AnimationSequences.
 [/EXCERPT:Intro]
 
+[TOC(start:2 end:2)]
+
 ## Component Space
 
 [INCLUDE:Engine/Animation/NodeReference/SpaceConversion#intro]
