#UE4doc #docs - Capitalized term

diff --git a/Engine/Documentation/Source/Programming/UnrealArchitecture/Reference/Functions/Specifiers/BlueprintAuthorityOnly/FunctionSpecifier_BlueprintAuthorityOnly.INT.udn b/Engine/Documentation/Source/Programming/UnrealArchitecture/Reference/Functions/Specifiers/BlueprintAuthorityOnly/FunctionSpecifier_BlueprintAuthorityOnly.INT.udn
--- a/Engine/Documentation/Source/Programming/UnrealArchitecture/Reference/Functions/Specifiers/BlueprintAuthorityOnly/FunctionSpecifier_BlueprintAuthorityOnly.INT.udn
+++ b/Engine/Documentation/Source/Programming/UnrealArchitecture/Reference/Functions/Specifiers/BlueprintAuthorityOnly/FunctionSpecifier_BlueprintAuthorityOnly.INT.udn
@@ -4,4 +4,4 @@ Description:
 Availability:Public
 Version: 4.9
 
-This function will not execute from blueprint code if running on something without network authority.
+This function will not execute from Blueprint code if running on something without network authority.
