#UE4doc #docs - Added Engine version field. Updated multiple images to match latest UI. Bolded/Unbolded terms.

diff --git a/Engine/Documentation/Source/Engine/UMG/UserGuide/Fonts/Fonts.INT.udn b/Engine/Documentation/Source/Engine/UMG/UserGuide/Fonts/Fonts.INT.udn
--- a/Engine/Documentation/Source/Engine/UMG/UserGuide/Fonts/Fonts.INT.udn
+++ b/Engine/Documentation/Source/Engine/UMG/UserGuide/Fonts/Fonts.INT.udn
@@ -5,6 +5,7 @@ Description:Details the process of importing fonts, working with the Font Editor
 Related: Engine/UMG
 Related: Engine/UMG/UserGuide/WidgetBlueprints
 Related: Engine/UMG/UserGuide/WidgetTypeReference
+Version: 4.9
 
 [TOC(start:2 end:2)]
 
@@ -12,7 +13,7 @@ This page covers how to import your own fonts, using the **Font Editor** and fon
 
 ## Font Assets
 
-Fonts in Unreal Engine 4 are categorized as a **Font** asset and use two caching methods, **Runtime** which is in the form of a Composite Font or **Offline** which is the older pre-computed Font Atlas method. You can switch between the two methods by opening up a Font asset in the **Font Editor** (this provides a simple way to convert existing Font assets from the older method to the new composite method without having to replace them). 
+Fonts in Unreal Engine 4 are categorized as a **Font** asset and use two caching methods, **Runtime** which is in the form of a Composite Font or **Offline** which is the older pre-computed Font Atlas method. You can switch between the two methods by opening up a Font asset in the Font Editor (this provides a simple way to convert existing Font assets from the older method to the new composite method without having to replace them). 
 
 ## Creating & Importing Fonts
 
@@ -22,7 +23,7 @@ This section illustrates how to create a new Font asset or import TTF or OTF fil
 
 To create a Font asset, follow the steps below: 
 
-1.	Click the **Create** button inside the **Content Browser**, then under **User Interface** select the **Font** option. 
+1.	Click the **Add New** button inside the **Content Browser**, then under **User Interface** select the **Font** option. 
 
 	![](NewFont.png)
 
@@ -30,17 +31,16 @@ To create a Font asset, follow the steps below:
 
 	![](NameFont.png)
 
-1.	Once you enter a name, you will notice an asterisk appears on the asset, this indicates that the asset has not be saved. 
+1.	Once you enter a name, you will notice an asterisk appears on the asset, this indicates that the asset has not be saved. Click the **Save All** button to save your asset, then confirm the save in the pop-up menu that appears to finish.
 
 	![](SaveFont.png)
-
-1.	Click the **Save All** button to save your asset, then confirm the save in the pop-up menu that appears to finish. 
+	 
 
 [REGION:tip]
-Instead of clicking the **Create** button, you can instead **Right-click** in an empty space in the **Content Browser** to access the context menu. From there, you can click the **User Interface** and **Font** option to create a Font asset. Do not forget to save your assets! 
+Instead of clicking the **Add New** button, you can instead **Right-click** in an empty space in the **Content Browser** to access the context menu. From there, you can click the **User Interface** and **Font** option to create a Font asset. Do not forget to save your assets! 
 [/REGION]
 
-When you create a new Font asset inside the Editor, it will create an empty Composite Font asset which you can then define its parameters using the **Font Editor**. When you import a TTF or OTF file, a new Composite Font asset is generated using the font sample you provided (at which point you can define additional parameters by opening it up in the **Font Editor**). 
+When you create a new Font asset inside the Editor, it will create an empty Composite Font asset which you can then define its parameters using the Font Editor. When you import a TTF or OTF file, a new Composite Font asset is generated using the font sample you provided (at which point you can define additional parameters by opening it up in the Font Editor). 
 
 ### Importing Fonts
 
@@ -54,7 +54,7 @@ In the dialogue box that appears, navigate to the file to import then click on i
 
 ![](Import2.png)
 
-The file will then be added to the **Content Browser** as a Composite Font asset, using the file you have selected. 
+The file will then be added to the Content Browser as a Composite Font asset, using the file you have selected. 
 
 ![](Import3.png)
 
