#UE4doc #docs - Added Engine Version Field. Updated images callout box color. Bolded/Unbolded terms.

diff --git a/Engine/Documentation/Source/Engine/UMG/UserGuide/PropertyBinding/PropertyBinding.INT.udn b/Engine/Documentation/Source/Engine/UMG/UserGuide/PropertyBinding/PropertyBinding.INT.udn
--- a/Engine/Documentation/Source/Engine/UMG/UserGuide/PropertyBinding/PropertyBinding.INT.udn
+++ b/Engine/Documentation/Source/Engine/UMG/UserGuide/PropertyBinding/PropertyBinding.INT.udn
@@ -2,6 +2,7 @@ Availability:Public
 Title: Property Binding
 Crumbs:%ROOT%, Engine, Engine/UMG, Engine/UMG/UserGuide
 Description:A guide to binding properties inside UMG. 
+Version: 4.9
 
 [TOC(start:2 end:2)]
 
@@ -21,7 +22,7 @@ The **Return Value** is tied to the Percent's Value, so you can plug in a value
 
 ![](PlayerHealth.png)(w:820)
 
-Above, the Function gets a variable called **Player Health** from inside our Character Blueprint. Anytime the **Player Health** variable is updated, it is automatically passed to and reflected as the Percent Value on the Progress Bar. 
+Above, the Function gets a variable called Player Health from inside our Character Blueprint. Anytime the Player Health variable is updated, it is automatically passed to and reflected as the Percent Value on the Progress Bar. 
 
 ## Property Binding
 
