#UE4doc #docs - Added Engine Version Field. Updated images callout box color. Bolded/Unbolded terms.

diff --git a/Engine/Documentation/Source/Engine/UMG/UserGuide/Slots/Slots.INT.udn b/Engine/Documentation/Source/Engine/UMG/UserGuide/Slots/Slots.INT.udn
--- a/Engine/Documentation/Source/Engine/UMG/UserGuide/Slots/Slots.INT.udn
+++ b/Engine/Documentation/Source/Engine/UMG/UserGuide/Slots/Slots.INT.udn
@@ -2,6 +2,7 @@ Availability:Public
 Title: Slots
 Crumbs:%ROOT%, Engine, Engine/UMG, Engine/UMG/UserGuide
 Description:An overview of how slots are used in UMG.
+Version: 4.9
 
 [TOC(start:2 end:2)]
 
@@ -12,7 +13,7 @@ Moreover, each Slot is different. For example, if you were to place a control on
 
 ## Accessing Slots 
 
-By convention, all **Slot** related properties appear under the **Layout** category in the **Details** panel (yellow box below). 
+By convention, all Slot related properties appear under the **Layout** category in the **Details** panel (yellow box below). 
 
 ![](SlotLayout.png)
 
@@ -20,14 +21,14 @@ You will also notice that the type of Slot being used by your Widget is displaye
 
 ### Setting Layout Properties
 
-At runtime, to modify properties under **Layout**, you can access the Slot member of the Widget in Blueprint or C++ and then **Cast** it to the correct **Slot Type**. Once doing so, you will then be able to modify the properties, an example of which is indicted below. 
+At runtime, to modify properties under Layout, you can access the Slot member of the Widget in Blueprint or C++ and then **Cast** it to the correct **Slot Type**. Once doing so, you will then be able to modify the properties, an example of which is indicted below. 
 
 ![](SlotCast.png)
 
 Above, a **Vertical Box** entitled **GameTitleBox** has been placed on a **CanvasPanel**. By getting the Slot associated with the Vertical Box and **Casting** to the **CanvasPanelSlot** type, we are then able to set the position of the box when our "StartButton" is clicked. 
 
 [REGION:note]
-Currently in Blueprints, only SETTER nodes are exposed. If you need to GET properties from the Layout, you may want to create a Variable to store your property and upon **Event Construct**, **SET** your Layout property via your Variable so that you have reference to and can access it later.
+Currently in Blueprints, only SETTER nodes are exposed. If you need to GET properties from the Layout, you may want to create a Variable to store your property and upon **Event Construct**, and **SET** your Layout property via your Variable so that you have reference to and can access it later.
 
 [/REGION]
 
