#UE4doc #docs - Added Engine Version field. Bolded/Unbolded terms. Updated images to match callout color.

diff --git a/Engine/Documentation/Source/Engine/UMG/UserGuide/Styling/Styling.INT.udn b/Engine/Documentation/Source/Engine/UMG/UserGuide/Styling/Styling.INT.udn
--- a/Engine/Documentation/Source/Engine/UMG/UserGuide/Styling/Styling.INT.udn
+++ b/Engine/Documentation/Source/Engine/UMG/UserGuide/Styling/Styling.INT.udn
@@ -67,7 +67,7 @@ You can assign Textures as your Image asset, but there are some flags you will w
 _Click for full view._
 
 #### Material Based Images
-Materials can also be assigned as the Image asset for your Image States, but you will need to change a few settings. Inside the Material, in the **Details** panel under **Usage**, you will need to check the **Used with UI** checkbox. This enables compiling a different shader specifically for Slate. You also need to only output **Emissive Color**
+Materials can also be assigned as the Image asset for your Image States, but you will need to change a few settings. Inside the Material, in the **Details** panel under **Usage**, you will need to check the **Used with UI** checkbox. This enables compiling a different shader specifically for Slate. You also need to only output **Emissive Color**.
 
 [REGION:lightbox]
 	[![](UsedWithUI.png)(w:760)](UsedWithUI.png)
