#UE4doc #docs - Added engine version field. Bolded/unbolded terms. Fixed Table values to add new column for number corresponding to image section.

diff --git a/Engine/Documentation/Source/Engine/UMG/UserGuide/WidgetBlueprints/WidgetBlueprints.INT.udn b/Engine/Documentation/Source/Engine/UMG/UserGuide/WidgetBlueprints/WidgetBlueprints.INT.udn
--- a/Engine/Documentation/Source/Engine/UMG/UserGuide/WidgetBlueprints/WidgetBlueprints.INT.udn
+++ b/Engine/Documentation/Source/Engine/UMG/UserGuide/WidgetBlueprints/WidgetBlueprints.INT.udn
@@ -2,6 +2,7 @@ Availability:Public
 Title: Widget Blueprints
 Crumbs:%ROOT%, Engine, Engine/UMG, Engine/UMG/UserGuide
 Description:How to create a Widget Blueprint and Overview of the Widget Blueprint Interface.
+Version: 4.9
 
 In order to start working with Unreal Motion Graphics, you will first need to create a **Widget Blueprint** as seen below. 
 
@@ -29,7 +30,7 @@ In order to start working with Unreal Motion Graphics, you will first need to cr
 
 ## Widget Blueprint Editor
 
-By default, the **Widget Blueprint Editor** opens up on the **Designer** tab when a **Widget Blueprint** is opened. The **Designer** tab is the visual representation of the layout and will give you an idea of how the screen will appear in-game.  
+By default, the **Widget Blueprint Editor** opens up on the **Designer** tab when a **Widget Blueprint** is opened. The Designer tab is the visual representation of the layout and will give you an idea of how the screen will appear in-game.  
 
 [REGION:lightbox]
 	[![](WidgetBlueprintNumbers.png)(w:840 h:480)](WidgetBlueprintNumbers.png)
@@ -37,16 +38,16 @@ By default, the **Widget Blueprint Editor** opens up on the **Designer** tab whe
 
 _Click for full view._
 
-| Window | Description |
-| -------- | ----------- |
-| **(1) Menu Bar** | The common Menu bar. |
-| **(2) Tool Bar** | This contains a number of commonly used functions for the Blueprint Editor, such as Compile, Save, and Play. |
-| **(3) Editor Mode** | Switches the UMG Widget Blueprint Editor between Designer and Graph modes. |
-| **(4) Palette** | The list of Widgets that you can drag-and-drop into the Visual Designer. Displays any class inheriting from UWidget. |
-| **(5) Hierarchy** | Displays the parenting structure of the User Widget. You can also drag Widgets into this window. |
-| **(6) Visual Designer** | This is the visual representation of the layout. Widgets dragged into the Visual Designer can be manipulated in the window. |
-| **(7) Details** | Displays the properties of the currently selected Widget. |
-| **(8) Animations** | This is the animation track for UMG which allows you to keyframe animations for your Widgets. |
+|   Number    | Window | Description |
+| :-----: | -------- | ----------- |
+| **1** | **Menu Bar** | The common Menu bar. |
+| **2** | **Tool Bar** | This contains a number of commonly used functions for the Blueprint Editor, such as Compile, Save, and Play. |
+| **3** | **Editor Mode** | Switches the UMG Widget Blueprint Editor between Designer and Graph modes. |
+| **4** | **Palette** | The list of Widgets that you can drag-and-drop into the Visual Designer. Displays any class inheriting from UWidget. |
+| **5** | **Hierarchy** | Displays the parenting structure of the User Widget. You can also drag Widgets into this window. |
+| **6** | **Visual Designer** | This is the visual representation of the layout. Widgets dragged into the Visual Designer can be manipulated in the window. |
+| **7** | **Details** | Displays the properties of the currently selected Widget. |
+| **8** | **Animations** | This is the animation track for UMG which allows you to keyframe animations for your Widgets. |
 
 [REGION:note]
 The **Visual Designer** window by default is 1:1 scale however you can hold **Control** and **Mouse-Wheel** up to zoom in further.
