#UE4doc #docs - Added engine version field.

diff --git a/Engine/Documentation/Source/Engine/Matinee/HowTo/MatineeHowTo.INT.udn b/Engine/Documentation/Source/Engine/Matinee/HowTo/MatineeHowTo.INT.udn
--- a/Engine/Documentation/Source/Engine/Matinee/HowTo/MatineeHowTo.INT.udn
+++ b/Engine/Documentation/Source/Engine/Matinee/HowTo/MatineeHowTo.INT.udn
@@ -2,6 +2,7 @@ Availability: Public
 Title:Matinee How-To's
 Crumbs: %ROOT%, Engine, Engine/Matinee
 Description:A How To Guide for using different aspects of Matinee in Unreal Engine 4. 
+Version: 4.9
 
 
 The **Matinee - How-To...** pages will show you how you can set up simple scenarios using Matinee in a step-by-step process with images and instructions. 
diff --git a/Engine/Documentation/Source/Engine/Matinee/MatineeAndCinematics.INT.udn b/Engine/Documentation/Source/Engine/Matinee/MatineeAndCinematics.INT.udn
--- a/Engine/Documentation/Source/Engine/Matinee/MatineeAndCinematics.INT.udn
+++ b/Engine/Documentation/Source/Engine/Matinee/MatineeAndCinematics.INT.udn
@@ -2,6 +2,7 @@ Availability: Public
 Title:Matinee and Cinematics
 Crumbs: %ROOT%, Engine
 Description:Animation tool that provides the ability to animate the properties of Actors over time, either to create dynamic gameplay or cinematic in-game sequences.
+Version: 4.9
 
 [VAR:Topic]
 [OBJECT:Topic]
