#UE4doc #docs - Added Engine Version Field. Added Skill Level field. Updated image callout box color. Removed images in sentences and replaced with proper term.

diff --git a/Engine/Documentation/Source/Engine/Matinee/HowTo/MHT_1/MHT_1.INT.udn b/Engine/Documentation/Source/Engine/Matinee/HowTo/MHT_1/MHT_1.INT.udn
--- a/Engine/Documentation/Source/Engine/Matinee/HowTo/MHT_1/MHT_1.INT.udn
+++ b/Engine/Documentation/Source/Engine/Matinee/HowTo/MHT_1/MHT_1.INT.udn
@@ -2,6 +2,8 @@ Availability: Public
 Title:Setting Up a Door
 Crumbs: %ROOT%, Engine, Engine/Matinee, Engine/Matinee/HowTo
 Description:Setting up a door that opens/closes automatically through Matinee. 
+Version: 4.9
+SkillLevel: Beginner
 
 
 ---
@@ -41,8 +43,10 @@ In the following steps, we are going to create a simple door that opens/closes w
 
 	![](MHT1_DoorFrame.png)(w:480)
 
-	_Note: if you do not see the **Props** folder, refer to the **Creating the Project** step and make sure that you enabled **Include starter content**._
-	
+	[Region:note]
+	If you do not see the **Props** folder, refer to the **Creating the Project** step and make sure that you enabled **Include starter content**.
+	[/Region]
+
 1. **Left-Click** on the **SM_DoorFrame** Static Mesh and drag-and-drop it into the **Level Viewport**. 
 	
 	![](MHT_DoorFrameDragIn.jpg)(w:640)
@@ -57,13 +61,17 @@ In the following steps, we are going to create a simple door that opens/closes w
 
 	![](MHT_CollisionComplexity.jpg)
 
-	_Note: The step above will ensure that the door has collision applied to it. _
+	[Region:note]
+	The step above will ensure that the door has collision applied to it.
+	[/Region]
 
 1. With the **SM_Door** selected in the Level Viewport find the Details panel. Under **Transform** and the **Rotation** setting, set the **Z** value to **-180**. 
 
 	![](MHT1_RotateDoor.png)(w:640)
 
-	_Note: it will auto adjust to -179.99 which is okay._
+	[Region:note]
+	It will auto adjust to -179.99 which is okay.
+	[/Region]
 
 1. Using the **Translation** widget, position the door inside the door frame. 
 
@@ -114,7 +122,7 @@ In the following steps, we are going to create a simple door that opens/closes w
 
 	![](MHT1_TrackContextMenu.png)
 
-1. When the ![](MHT1_GroupName.png) appears, enter a **New Group Name** for the group (in this example, we are using the name **Door**). 
+1. When the window asking you to enter a **New Group Name** for the group appears enter a name for this group (in this example, we are using the name **Door**). 
 
 1. Minimize the Matinee window, then in the **Main Editor** find the **World Outliner** window and select the **SM_Door** so that it is highlighted. 
 
@@ -124,7 +132,9 @@ In the following steps, we are going to create a simple door that opens/closes w
 
 	![](MHT1_AddDoorToGroup.png)
 
-	_Note: This will add the **SM_Door** Actor to the **Door** group._
+	[Region:note]
+	This will add the **SM_Door** Actor to the **Door** group.
+	[/Region]
 
 1. **Right-Click** on the Group (Door) again to bring up the context menu, then select **Add New Movement Track**. 
 
@@ -148,7 +158,7 @@ In the following steps, we are going to create a simple door that opens/closes w
 
 	![](MHT1_DoorOpen.png)
 
-1. Return to the Matinee window and **Left-Click** on the ![](MHT1_EndPoint.png) **End Marker** and move it to the left to **0.50**. 
+1. Return to the Matinee window and **Left-Click** on the Red arrow **End Marker** and move it to the left to **0.50**. 
 
 	![](MHT1_EndPointMoved.png)
 
@@ -214,7 +224,7 @@ In the following steps, we are going to create a simple door that opens/closes w
 	
 1. The Blueprint is now complete and the **TriggerBox** is now hooked up to the Matinee Actor. 
 
-1. Minimize the Level Blueprint and **click** the ![](MHT1_PlayButton.png) **Play** Button to play in the Editor. 
+1. Minimize the Level Blueprint and click the **Play** Button in the main toolbar to play in the Editor. 
 
 	[OBJECT:EmbeddedVideo]
 	[PARAMLITERAL:width]
