#UE4doc #docs:
Added engine version field
Unbolded some hyperlinks
Replaced multiple images due to out of date UI

diff --git a/Engine/Documentation/Source/Engine/UMG/UserGuide/WidgetTypeReference/WidgetTypeReference.INT.udn b/Engine/Documentation/Source/Engine/UMG/UserGuide/WidgetTypeReference/WidgetTypeReference.INT.udn
--- a/Engine/Documentation/Source/Engine/UMG/UserGuide/WidgetTypeReference/WidgetTypeReference.INT.udn
+++ b/Engine/Documentation/Source/Engine/UMG/UserGuide/WidgetTypeReference/WidgetTypeReference.INT.udn
@@ -2,6 +2,7 @@ Availability:Public
 Title: Widget Type Reference
 Crumbs:%ROOT%, Engine, Engine/UMG, Engine/UMG/UserGuide
 Description:Reference guide to the different types of Widgets that can be used in a Widget Blueprint in UMG. 
+version: 4.9
 
 [TOC(start:2 end:2)]
 
@@ -26,7 +27,7 @@ The most frequently used Widgets will more than likely fall under this category.
 | **Button** | The button is a single-child, click-able primitive widget which enables basic interaction. You can place any other widget inside a button to make a more complex and interesting click-able element in your UI.  |
 | **Check Box** | The check box widget allows you to display a toggled state of 'unchecked', 'checked', and 'indeterminable'. You can use the check box for a classic check box, or as a toggle button, or as radio buttons.  |
 | **Image** | The image widget allows you to display a Slate Brush, Texture, or Material in the UI. |
-| [**Named Slot**](Engine\UMG\UserGuide\WidgetTypeReference\NamedSlot) | [INCLUDE:Engine\UMG\UserGuide\WidgetTypeReference\NamedSlot#desc] |
+| [Named Slot](Engine\UMG\UserGuide\WidgetTypeReference\NamedSlot) | [INCLUDE:Engine\UMG\UserGuide\WidgetTypeReference\NamedSlot#desc] |
 | **Progress Bar** | The progress bar widget is a simple bar that fills up and can be restyled to fit any number of uses such as Experience, Health, Points, etc. |
 | **Slider** | A simple widget that shows a sliding bar with a handle that allows you to control the value between 0-1. |
 | **Text** | The basic method for displaying text on screen and could be used for text descriptions of options or other UI elements. |
@@ -62,7 +63,7 @@ The **Panel** category contains useful Widgets for controlling layouts and placi
 | **Uniform Grid Panel** | A panel that evenly divides up available space between all of its children. |
 | **Vertical Box** | A vertical box widget is a layout panel allowing child widgets to be automatically laid out vertically. This is great for stacking widgets on top of each other and keeping widgets aligned vertically.|
 | **Widget Switcher** | A widget switcher is like a tab control but without the tabs which you can create on your own and combine with this to get a tabbed effect. At most one widget is visible at a time.  |
-| [**Wrap Box**](Engine\UMG\UserGuide\WidgetTypeReference\WrapBox) | [INCLUDE:Engine\UMG\UserGuide\WidgetTypeReference\WrapBox#desc] |
+| [Wrap Box](Engine\UMG\UserGuide\WidgetTypeReference\WrapBox) | [INCLUDE:Engine\UMG\UserGuide\WidgetTypeReference\WrapBox#desc] |
 
 
 ## Primitive Widgets
@@ -76,7 +77,7 @@ The **Widgets** contained in the **Primitive** category provide additional metho
 | **Circular Throbber** | A throbber widget that orients images in a spinning circle.  |
 | **Editable Text** | This a text field that allows for user input without a box background. This only supports a single line of editable text. |
 | **Editable Text (Multi-Line)** | Similar to Editable Text but supports multiple lines of text instead of a single line.|
-| [**Menu Anchor**](Engine\UMG\UserGuide\WidgetTypeReference\MenuAnchor) | [INCLUDE:Engine\UMG\UserGuide\WidgetTypeReference\MenuAnchor#desc] |
+| [Menu Anchor](Engine\UMG\UserGuide\WidgetTypeReference\MenuAnchor) | [INCLUDE:Engine\UMG\UserGuide\WidgetTypeReference\MenuAnchor#desc] |
 | **Native Widget Host** |This is a container widget that can contain one child Slate widget. This should be used when all you need is to nest a native widget inside a UMG widget. |
 | **Spacer** |A spacer widget provides custom padding between other widgets. The spacer does not have a visual representation and is invisible in game.|
 | **Throbber** | An animated throbber widget that shows several zooming circles in a row (could be used to signify loading for example). |
