#UE4doc #docs:
Added engine version field
Changed font settings for a few terms
Replaced a couple images due to outdated HUD

diff --git a/Engine/Documentation/Source/Engine/UMG/UserGuide/WidgetTypeReference/MenuAnchor/MenuAnchor.INT.udn b/Engine/Documentation/Source/Engine/UMG/UserGuide/WidgetTypeReference/MenuAnchor/MenuAnchor.INT.udn
--- a/Engine/Documentation/Source/Engine/UMG/UserGuide/WidgetTypeReference/MenuAnchor/MenuAnchor.INT.udn
+++ b/Engine/Documentation/Source/Engine/UMG/UserGuide/WidgetTypeReference/MenuAnchor/MenuAnchor.INT.udn
@@ -5,6 +5,7 @@ Description:Overview of using Menu Anchors for use with Popup Menus.
 Related: Engine/UMG/UserGuide
 Related: Engine/UMG/UserGuide/Styling
 Related: Engine/UMG/UserGuide/Slots
+version: 4.9
 
 [TOC(start:2 end:2)]
 
@@ -25,22 +26,26 @@ In the **Details** panel for a placed **Menu Anchor**, there are a couple of spe
 | **Menu Class** | The Widget Blueprint (pop-up) to spawn when the menu is called. Creates the widget freshly each time. |
 | **Placement** | Several placement options use to determine the location of the created widget. |
 | **OnGetMenuContent** | Allows for the customization of the pop-up by binding a function or property to the **OnGetMenuContent** Event which is called when the pop-up is summoned (see below). |
+| **OnMenuOpenChanged** | Allows for the customization of the pop-up by binding a function or property to the **OnMenuOpenChanged** Event which is called opened state of the menu changes. |
 
-For the **OnGetMenuContent** Event, you can create and bind a function directly from the **Details** panel which will be called whenever the Menu Anchor is called. This allows you to provide extra functionality to the creation of your pop-up menu. For example, below we've run a check to determine if the player is _In Battle_ and if not, we allow them to access the pop-up menu. If they are we do not allow them to access the pop-up menu. 
+For the **OnGetMenuContent** Event, you can create and bind a function directly from the **Details** panel which will be called whenever the Menu Anchor is called. This allows you to provide extra functionality to the creation of your pop-up menu. For example, below we've run a check to determine if the player is "In Battle" and if not, we allow them to access the pop-up menu. If they are we do not allow them to access the pop-up menu. 
 
 ![](ExampleGetMenu.png)(w:840)
 
 ### Menu Anchor Blueprint Functions
 
-The **Menu Anchor** widget has some widget specific functions that can be called through script which are described below. 
+The Menu Anchor widget has some widget specific functions that can be called through script which are described below. 
 
 ![](MenuAnchorNodes.png)
 
 | Option | Description |
 | -------- | ----------- |
 | **Close** | Closes the menu if it is currently open. |
+| **Get Menu Position** | Returns the menu's position in world space. |
+| **Has Open Sub Menus** | Returns wether the menu has open sub menus or not. |
 | **Is Open** | Checks to see if the Target Menu Anchor is currently open and returns a Boolean. |
 | **Open** | Opens the menu if it is currently closed. |
+| **Should Open Due to Click** | Returns wether the menu should be openable by clicking or not. |
 | **Toggle Open** | This toggles the open state for the menu and can be used instead of opening/closing. |
 
 ## Usage Example
@@ -73,17 +78,17 @@ After you have created your Inventory Menu, you would also add the **Menu Anchor
 
 ![](AddMenuAnchor.png)
 
-You would then position the **Menu Anchor** where you want it to appear, we positioned and Anchored it to the right/center position. 
+You would then position the Menu Anchor where you want it to appear, we positioned and Anchored it to the right/center position. 
 
 ![](PlacedAnchor.png)
 
-In the **Details** panel for the **Menu Anchor**, you would set the **Menu Class** (or Widget Blueprint to spawn) as well as its placement options.
+In the **Details** panel for the Menu Anchor, you would set the **Menu Class** (or Widget Blueprint to spawn) as well as its placement options.
 
 ![](MenuAnchorDetails.png)
 
-On the Graph tab, you would add **OnClick** Events for each of your Buttons and drag in the **Menu Anchor** widget. 
+On the Graph tab, you would add **OnClick** Events for each of your Buttons and drag in the Menu Anchor widget. 
 
-Dragging off the **Menu Anchor**, you can call the **Open** function to open and display your specified Menu Class. 
+Dragging off the Menu Anchor, you can call the **Open** function to open and display your specified Menu Class. 
 
 ![](HUD_Graph.png)
 
@@ -97,15 +102,15 @@ This could be through the **Level Blueprint** or (in our case) through the Defau
 
 ![](DisplayingInventoryMenu.png)
 
-Above in our **MyCharacter** Blueprint, when the game is started we create the _HUD_ Widget and store it as a variable called **HUDWidget** so we can access it later. We then set **I** to toggle between **Add to Viewport** and **Remove from Parent** to Display or Remove the Inventory Menu. We also set **Show Mouse Cursor** to enabled or disabled based on whether or not the Inventory Menu is open or not. 
+Above in our **MyCharacter** Blueprint, when the game is started we create the HUD Widget and store it as a variable called **HUDWidget** so we can access it later. We then set **I** to toggle between **Add to Viewport** and **Remove from Parent** to Display or Remove the Inventory Menu. We also set **Show Mouse Cursor** to enabled or disabled based on whether or not the Inventory Menu is open or not. 
 
-And finally inside the **MenuPopUp** Blueprint (your pop-up menu), adding some script to **Close** the Menu Anchor when the player presses a Button in the pop-up menu. Below when the pop-up is constructed, we get a reference to the **MyCharacter** Blueprint by **Casting** to it then promoting it to a variable called **MyCharacterReference**. With this reference we can then access the _HUD_
+And finally inside the **MenuPopUp** Blueprint (your pop-up menu), adding some script to **Close** the Menu Anchor when the player presses a Button in the pop-up menu. Below when the pop-up is constructed, we get a reference to the **MyCharacter** Blueprint by **Casting** to it then promoting it to a variable called **MyCharacterReference**. With this reference we can then access the HUD
 
 ![](PopUpMenuNetwork.png)(w:820)
 
 You would probably want to add some functionality for each of the buttons for your pop-up (instead of printing to the screen) that when pressed, does something. For example, perhaps the first button "uses" the item from the Inventory, while the second "drops" the item and the third "cancels" and returns the player to the Inventory Menu. 
 
-This is just one example, you could also change what causes the **Menu Anchor** to open. Instead of a button press, perhaps using the **IsHovered** function to determine if an object is being hovered over then opening the Menu Anchor and pop-up window to display a tool-tip or some other form of notification. In either case, this should get you started with using the **Menu Anchor** to display/remove pop-up widgets through script. 
+This is just one example, you could also change what causes the Menu Anchor to open. Instead of a button press, perhaps using the **IsHovered** function to determine if an object is being hovered over then opening the Menu Anchor and pop-up window to display a tool-tip or some other form of notification. In either case, this should get you started with using the Menu Anchor to display/remove pop-up widgets through script. 
 
 
 
