#UE4doc #docs:
Added engine version field
Changed font settings for multiple terms
Replaced a few images due to out of date UI
Added  a step as it needed to be explained

diff --git a/Engine/Documentation/Source/Engine/UMG/UserGuide/WidgetTypeReference/NamedSlot/NamedSlot.INT.udn b/Engine/Documentation/Source/Engine/UMG/UserGuide/WidgetTypeReference/NamedSlot/NamedSlot.INT.udn
--- a/Engine/Documentation/Source/Engine/UMG/UserGuide/WidgetTypeReference/NamedSlot/NamedSlot.INT.udn
+++ b/Engine/Documentation/Source/Engine/UMG/UserGuide/WidgetTypeReference/NamedSlot/NamedSlot.INT.udn
@@ -5,6 +5,7 @@ Description:Exposing an External Slot for User Created Widgets with a Named Slot
 Related: Engine/UMG/UserGuide
 Related: Engine/UMG/UserGuide/Styling
 Related: Engine/UMG/UserGuide/Slots
+version: 4.9
 
 [TOC(start:2 end:2)]
 
@@ -16,7 +17,7 @@ Allows you to expose an external slot for your User Widget. When implementing th
 
 ## Usage Example
 
-Below we created a **Widget Blueprint** called _Button Template_ and replaced the **Canvas Panel** with a **Horizontal Box**. Inside that Horizontal Box we placed a **Button** and **Text** widget and what we want is to allow the space next to the button to filled by someone else with whatever widgets they want or to be filled by us later from inside another Widget Blueprint. 
+Below we created a **Widget Blueprint** called ButtonTemplate and replaced the **Canvas Panel** with a **Horizontal Box**. Inside that Horizontal Box we placed a **Button** and **Text** widget and what we want is to allow the space next to the button to filled by someone else with whatever widgets they want or to be filled by us later from inside another Widget Blueprint. 
 
 ![](NamedSlot1.png)
 
@@ -24,7 +25,7 @@ We can do so by adding a **Named Slot** widget to the Horizontal Box.
 
 ![](NamedSlot2.png)
 
-In our example, for clarity, we have updated the name of the Named Slot, set it to **Fill** and adjusted the preview size. 
+In our example, for clarity, we have updated the name of the Named Slot, set it to **Fill** and adjusted the preview size. Also, to get the same size, select **Custom** from the **Fill Size** dropdown and change the **Width** to 400 and **Height** to 50.
 
 ![](NamedSlot3.png)
 
@@ -32,9 +33,9 @@ We have also updated the text on the Button.
 
 ![](NamedSlot4.png)
 
-We **Compiled** and **Saved**, closed the _Button Template_, and then created a new Widget Blueprint. 
+We **Compiled** and **Saved**, closed the ButtonTemplate, and then created a new Widget Blueprint. 
 
-In the new Widget Blueprint, under **User Created**, we drag our _Button Template_ User Widget onto our Canvas Panel. 
+In the new Widget Blueprint, under **User Created**, we drag our Button Template User Widget onto our Canvas Panel. 
 
 ![](NamedSlot5.png)
 
@@ -42,7 +43,7 @@ We can then add Widgets to the Named Slot, however we cannot drag-and-drop them
 
 ![](NamedSlot6.png)
 
-We need to use the **Hierarchy** to create a new "root" widget for the Named Slot. 
+We need to use the **Hierarchy** to create a new root widget for the Named Slot. 
 
 ![](NamedSlot7.png)
 
@@ -54,9 +55,9 @@ In order to drag-and-drop Widgets onto the graph, we need a container that can h
 
 ![](NamedSlot9.png)(w:800)
 
-Above we have created another _Button Template_ Widget, this one uses a Canvas Panel and we can now drag Widgets onto it. 
+Above we have created another ButtonTemplate Widget, this one uses a Canvas Panel and we can now drag Widgets onto it. 
 
-You add multiple Widgets to the Named Slot, allowing you to setup "templates" that can be populated with various widgets. 
+You add multiple Widgets to the Named Slot, allowing you to setup templates that can be populated with various widgets. 
 
 ![](NamedSlot10.png)(w:800)
 
