#UE4doc #docs:
Added engine version field
Changed font settings for a few terms

diff --git a/Engine/Documentation/Source/Engine/UMG/UserGuide/WidgetTypeReference/WrapBox/WrapBox.INT.udn b/Engine/Documentation/Source/Engine/UMG/UserGuide/WidgetTypeReference/WrapBox/WrapBox.INT.udn
--- a/Engine/Documentation/Source/Engine/UMG/UserGuide/WidgetTypeReference/WrapBox/WrapBox.INT.udn
+++ b/Engine/Documentation/Source/Engine/UMG/UserGuide/WidgetTypeReference/WrapBox/WrapBox.INT.udn
@@ -31,7 +31,7 @@ In the **Details** panel of a Widget that is placed inside a **Wrap Box**, there
 
 ## Usage Example
 
-The **Wrap Box** is useful when you want to arrange items within a container and have those items appear aligned and evenly spaced and automatically format the contents inside it based on the width of the container. 
+The Wrap Box is useful when you want to arrange items within a container and have those items appear aligned and evenly spaced and automatically format the contents inside it based on the width of the container. 
 
 Take for example the image below where we have created a "shop" screen of sorts where several items are presented to the player. 
 
@@ -81,7 +81,7 @@ It could also be used to keep things inside it arranged when resizing the window
 [/PARAMLITERAL]
 [/OBJECT]
 
-While the **Wrap Box** only checks width to determine when to arrange the Widgets within it, you could also use this in combination with a **Scroll Box** which will add a scroll bar when the amount of Widgets within the window becomes cut off allowing the user to scroll down to see the remaining Widgets. 
+While the Wrap Box only checks width to determine when to arrange the Widgets within it, you could also use this in combination with a Scroll Box which will add a scroll bar when the amount of Widgets within the window becomes cut off allowing the user to scroll down to see the remaining Widgets. 
 
 
 
