#UE4doc #docs:
Added engine version field

diff --git a/Engine/Documentation/Source/Gameplay/Gameplay.INT.udn b/Engine/Documentation/Source/Gameplay/Gameplay.INT.udn
--- a/Engine/Documentation/Source/Gameplay/Gameplay.INT.udn
+++ b/Engine/Documentation/Source/Gameplay/Gameplay.INT.udn
@@ -2,6 +2,7 @@ Title: Gameplay Guide
 Crumbs:%ROOT%
 Description:Overviews and examples of gameplay functionality for programmers and visual scripters.
 Availability:Public
+version: 4.9
 
 [VAR:Params]
 [PARAM:image]
