#UE4doc #docs:
Added engine version field

diff --git a/Engine/Documentation/Source/Gameplay/AI/AI.INT.udn b/Engine/Documentation/Source/Gameplay/AI/AI.INT.udn
--- a/Engine/Documentation/Source/Gameplay/AI/AI.INT.udn
+++ b/Engine/Documentation/Source/Gameplay/AI/AI.INT.udn
@@ -2,6 +2,7 @@ Availability:Public
 Title: AI and Behavior Trees
 Crumbs:%ROOT%, Engine, Gameplay
 Description:Information over Artificial Intelligence including the use of Behavior Trees and Navmesh.
+version: 4.9
 
 <!-- Q:\Daniel.Broder\BT Docs -->
 
