#UE4doc #docs:
Added engine version field

diff --git a/Engine/Documentation/Source/Gameplay/Analytics/Analytics.INT.udn b/Engine/Documentation/Source/Gameplay/Analytics/Analytics.INT.udn
--- a/Engine/Documentation/Source/Gameplay/Analytics/Analytics.INT.udn
+++ b/Engine/Documentation/Source/Gameplay/Analytics/Analytics.INT.udn
@@ -2,6 +2,7 @@ Availability:Public
 Title:In-Game Analytics
 Crumbs: %ROOT%, Gameplay
 Description:Using in-game analytics to track player engagement and find balance issues.
+version: 4.9
 
 
 In order to get data on how your game is performing, you must use an analytics provider to capture and process the data. For those without a homegrown solution, there are plenty of 
