#UE4doc #docs:
Added engine version field
Changed font settings for a term
Replaced an image as the UI was outdated
Edited a graph region to reflect the updates made to the node displayed

diff --git a/Engine/Documentation/Source/Gameplay/Analytics/Blueprints/BlueprintsPlugin.INT.udn b/Engine/Documentation/Source/Gameplay/Analytics/Blueprints/BlueprintsPlugin.INT.udn
--- a/Engine/Documentation/Source/Gameplay/Analytics/Blueprints/BlueprintsPlugin.INT.udn
+++ b/Engine/Documentation/Source/Gameplay/Analytics/Blueprints/BlueprintsPlugin.INT.udn
@@ -175,7 +175,7 @@ To record an event that has no associated attributes with it, use the **Record E
 
 ## Record Event With Attribute
  
-This records an event by name with a single attribute, so that you can compare the frequency of that event across the unique attributes. For instance, the `Item.Equipped` event below would record 
+This records an event by name with a single attribute, so that you can compare the frequency of that event across the unique attributes. For instance, the Item.Equipped event below would record 
 how often the user equips a weapon in your game, which will show which items are the most popular with players in your game.
 
 [REGION:graph]
@@ -184,7 +184,7 @@ how often the user equips a weapon in your game, which will show which items are
             function
         [/PARAMLITERAL]
         [PARAMLITERAL:title]
-            Record Event
+            Record Event with Attribute
         [/PARAMLITERAL]
         [PARAM:inputs]
             [OBJECT:BlueprintPin]
