#UE4doc #docs:
Added engine version field

diff --git a/Engine/Documentation/Source/Gameplay/ClassCreation/ClassCreationBasics.INT.udn b/Engine/Documentation/Source/Gameplay/ClassCreation/ClassCreationBasics.INT.udn
--- a/Engine/Documentation/Source/Gameplay/ClassCreation/ClassCreationBasics.INT.udn
+++ b/Engine/Documentation/Source/Gameplay/ClassCreation/ClassCreationBasics.INT.udn
@@ -2,6 +2,7 @@ Availability:Public
 Title:Class Creation Basics
 Crumbs: %ROOT%, Gameplay
 Description:Examples showing how to create classes with Blueprints alone, C++ alone, and a combination of C++ and Blueprints.
+version: 4.9
 
 ![](LightSwitchActor.png)(w:500 convert:false) 
 
