#DocFix

Updating multiuple docs from forum requests. (Verifing before submitting the change)

diff --git a/Engine/Documentation/Source/Engine/Blueprints/QuickStart/5/BPQS_5.INT.udn b/Engine/Documentation/Source/Engine/Blueprints/QuickStart/5/BPQS_5.INT.udn
--- a/Engine/Documentation/Source/Engine/Blueprints/QuickStart/5/BPQS_5.INT.udn
+++ b/Engine/Documentation/Source/Engine/Blueprints/QuickStart/5/BPQS_5.INT.udn
@@ -62,7 +62,7 @@ We will do that by working with the Other Actor output from the **OnComponentBeg
 1. Type "Branch" in the search, then select **Branch** from the context menu.
 
 	![](BranchAdded.png)
-1. Connect the output pin of the **Equals** node to the input Boolean pin on the **Branch** node.
+1. Connect the output pin of the **Equals** node to the input Boolean pin on the **Condition** node.
 
 	![](BranchComplete.png)
 
diff --git a/Engine/Documentation/Source/Resources/ContentExamples/MaterialProperties/1_5/1_5.INT.udn b/Engine/Documentation/Source/Resources/ContentExamples/MaterialProperties/1_5/1_5.INT.udn
--- a/Engine/Documentation/Source/Resources/ContentExamples/MaterialProperties/1_5/1_5.INT.udn
+++ b/Engine/Documentation/Source/Resources/ContentExamples/MaterialProperties/1_5/1_5.INT.udn
@@ -82,4 +82,5 @@ Related: Engine/Rendering/Materials/MaterialProperties
 | **Deferred Decals** | If the material is to be used by **Decal Actor**, it needs to be set to this or it will not work. |
 | **Light Function**  | Used by materials applied to lights as **LightFunctions**. |
 | **PostProcess**     | If the material is used as **Blendable Material** in a **PostProcessing** chain. |
+| **User Interface**     | This material will be used for **UMG** or **Slate UI** |
 
