#UE4doc #docs:
Added engine version field
Replaced multiple images due to out of date UI
Changed font settings for multiple terms

diff --git a/Engine/Documentation/Source/Gameplay/ClassCreation/BlueprintOnly/BlueprintOnly.INT.udn b/Engine/Documentation/Source/Gameplay/ClassCreation/BlueprintOnly/BlueprintOnly.INT.udn
--- a/Engine/Documentation/Source/Gameplay/ClassCreation/BlueprintOnly/BlueprintOnly.INT.udn
+++ b/Engine/Documentation/Source/Gameplay/ClassCreation/BlueprintOnly/BlueprintOnly.INT.udn
@@ -2,6 +2,7 @@ Availability:Public
 Title:Blueprints Only
 Crumbs: %ROOT%, Gameplay, Gameplay/ClassCreation
 Description:Introductory information for gameplay programmers getting started with Unreal Engine.
+version: 4.9
 
 A [Blueprints Class](Engine/Blueprints/UserGuide/Types/ClassBlueprint) allows you to set up new classes using the [](Engine/Blueprints) system.  After you [create a new Blueprint Class](Engine/Blueprints/UserGuide/Types/ClassBlueprint/Creation),
 you can add components, set up functions and other gameplay or design behavior with visual scripting, and set default values for class variables. The LightSwitch class set up using only Blueprints is named
@@ -9,14 +10,13 @@ you can add components, set up functions and other gameplay or design behavior w
 
 ## Class Setup
 
-The **LightSwitch_BPOnly** class was [created](Engine/Blueprints/UserGuide/Types/ClassBlueprint/Creation) in the **Content Browser**, and **Actor** was selected as the parent class. 
+The LightSwitch_BPOnly class was created in the **Content Browser**, and **Actor** was selected as the parent class. 
 
-Components are added to the Blueprint either statically in the **Blueprint Editor**'s [Components Window](Engine/Blueprints/UserGuide/Components), or dynamically through visual scripting in the [Graph Editor](Engine/Blueprints/Editor/UIComponents/GraphEditor).
-Because the LightSwitch class will always have the PointLightComponent and the SphereComponent, in this case the components have been added statically in [Components Window](Engine/Blueprints/UserGuide/Components).
+Components are added to the Blueprint either statically in the **Blueprint Editor**'s [Components](Engine/Blueprints/UserGuide/Components) tab, or dynamically through visual scripting in the [Graph Editor](Engine/Blueprints/Editor/UIComponents/GraphEditor).
+Because the LightSwitch class will always have the PointLightComponent and the SphereComponent, in this case the components have been added statically in **Components** tab.
 
-The **Components** tab was used to add the PointLightComponent, named **PointLight1**, as the root component, and the SphereComponent, named **Sphere1**, attached to the PointLightComponent. 
-Components added to the class in the **Blueprint Editor** have light blue icons, while components inherited from the parent class have dark blue icons. To learn more about adding
-and arranging components using the **Components** tab, see the [](Engine/Blueprints/UserGuide/Components) tab documentation.
+The **Components** tab was used to add the PointLightComponent, named PointLight1, as the root component, and the SphereComponent, named Sphere1, attached to the PointLightComponent. 
+Components added to the class in the **Blueprint Editor** have light blue icons, while components inherited from the parent class have dark blue icons.
 
 ![](BP_Only_ComponentList.png)  
 
@@ -25,19 +25,18 @@ In the **Blueprint Editor**, you can add new [variables](Engine/Blueprints/UserG
 contained within the Blueprint Class. In graphs, [nodes](Engine/Blueprints/UserGuide/Nodes) are wired together to create design-time and gameplay functionality that can be driven by class variables, gameplay events,
 and even an Actor's surroundings. 
 
-The float variable **DesiredIntensity** was added to the **LightSwitch_BPOnly** class using the **My Blueprint** tab.  The **My Blueprint** tab also displays the components added in **Components Window**, so they can be accessed
-in graphs if desired. To read more about adding variables using the **My Blueprint** tab, see the [Variables](Engine/Blueprints/UserGuide/Variables) documentation, or see the [](Engine/Blueprints/Editor/UIComponents/MyBlueprint) 
-documentation for general usage information.
+The float variable **DesiredIntensity** was added to the LightSwitch_BPOnly class using the **My Blueprint** tab.  The **My Blueprint** tab also displays the components added in **Components Window**, so they can be accessed
+in graphs if desired.
 
 ![](BP_Only_MyBlueprint.png)
 
-There are two graphs used to set up the **LightSwitch_BPOnly** class behavior. The first is the [](Engine/Blueprints/UserGuide/UserConstructionScript) function graph, which contains the **Construction Script** function entry node.
-This event executes when an Actor is added to the level, or when an existing Actor is moved within the level.  In the **LightSwitch_BPOnly** class, the **Construction Script** event is connected to a **Set 
-Brightness** node, so that the intensity of **Point Light 1** (the PointLightComponent) is set to the value of **DesiredIntensity** when the Actor is added to or moved within the level. 
+There are two graphs used to set up the LightSwitch_BPOnly class behavior. The first is the [](Engine/Blueprints/UserGuide/UserConstructionScript) function graph, which contains the **Construction Script** function entry node.
+This event executes when an Actor is added to the level, or when an existing Actor is moved within the level.  In the LightSwitch_BPOnly class, the **Construction Script** event is connected to a **Set 
+Intensity** node, so that the intensity of **Point Light 1** (the PointLightComponent) is set to the value of **DesiredIntensity** when the Actor is added to or moved within the level. 
 
 ![](BP_Only_ConstructionScript.png)
 
-The other graph set up in the **LightSwitch_BPOnly** class is the [](Engine/Blueprints/UserGuide/EventGraph). Execution in the EventGraph begins from events, such as the **OnComponentBeginOverlap** and 
+The other graph set up in the LightSwitch_BPOnly class is the [](Engine/Blueprints/UserGuide/EventGraph). Execution in the EventGraph begins from events, such as the **OnComponentBeginOverlap** and 
 **OnComponentEndOverlap** events. These events execute when some other Actor within the level overlaps or leaves the SphereComponent.  Both events are connected to the **Toggle Visibility** node, so
 the visibility of the PointLightComponent is toggled when those events execute. For more on events and working with graphs, see the [](Engine/Blueprints/UserGuide/Events), [](Engine/Blueprints/UserGuide/EventGraph),
 and [Graph Editor](Engine/Blueprints/Editor/UIComponents/GraphEditor) documentation.
@@ -46,12 +45,12 @@ and [Graph Editor](Engine/Blueprints/Editor/UIComponents/GraphEditor) documentat
 ![](BP_Only_EventGraph.png)
 
 [REGION:note]
-If you cannot find the Toggle Visibility from **right-click** menu un-check context sensitive or try searching for it in the Find a Node Menu. You can also get the same results by simply dragging
+If you cannot find the Toggle Visibility from **Right-click** menu un-check context sensitive or try searching for it in the Find a Node Menu. You can also get the same results by simply dragging
 out a line from the Point light variable node and then searching for Toggle Visibility. 
 [/REGION]
 [REGION:note]
 If you cannot find the OnComponentEndOverlap for Sphere1 make sure that you have the Sphere variable selected and then use the add event from the Event->Add Event or 
-from the **right-click** menu Add Event->Collision->View **OnComponentBeginOverlap / OnComponentEndOverlap**. 
+from the **Right-click** menu Add Event->Collision->View **OnComponentBeginOverlap / OnComponentEndOverlap**. 
 [/REGION]
 
 The **DesiredIntensity** variable is set to **Editable** in the variable's settings, so it is visible in the **Blueprint Editor**'s **Class Defaults** and can be edited.
@@ -60,7 +59,7 @@ This also means that the variable can be changed for each instance of the class,
 ![](BP_Only_Defaults.png)(convert:false)
 
 Blueprinted classes can be extended with other Blueprint Classes, either using the dropdown button next to the class in the **Class Viewer** to create a new Blueprint, or by
-**right-clicking** on the Blueprint and selecting **Create New Blueprint Based on This**.
+**Right-clicking** on the Blueprint and selecting **Create New Blueprint Based on This**.
 
-The Blueprint Class **LightSwitchBPOnly** is in the [](Engine/Content/Browser), and can be dragged into the level from there. It is also in the [](Engine/UI/ClassViewer).
+The Blueprint Class LightSwitch_BPOnly is in the [](Engine/Content/Browser), and can be dragged into the level from there. It is also in the [](Engine/UI/ClassViewer).
 For more information on placing Actors in a level using the **Content Browser** or **Class Viewer**, see the [](Engine/Actors/Placement) documentation.
