#UE4doc #docs:
Updated engine version field
Removed some sections of code as they were no longer needed
Removed some hyperlinks as they were redundant

diff --git a/Engine/Documentation/Source/Gameplay/ClassCreation/CodeOnly/CodeOnly.INT.udn b/Engine/Documentation/Source/Gameplay/ClassCreation/CodeOnly/CodeOnly.INT.udn
--- a/Engine/Documentation/Source/Gameplay/ClassCreation/CodeOnly/CodeOnly.INT.udn
+++ b/Engine/Documentation/Source/Gameplay/ClassCreation/CodeOnly/CodeOnly.INT.udn
@@ -2,7 +2,7 @@ Availability:Public
 Title:C++ Only
 Crumbs: %ROOT%, Gameplay, Gameplay/ClassCreation
 Description:Introductory information for gameplay programmers getting started with Unreal Engine.
-Version: 4.6
+Version: 4.9
 
 Using the [](Programming\Development\ManagingGameCode\CppClassWizard), you can easily add new C++ classes to your project. After you choose the class you would like your new class to derive from, the wizard
 sets up the necessary header file and source file. If this is the first code you have added to the project, the project is converted to a code project, which creates the game module that contains your 
@@ -43,7 +43,7 @@ This means that their properties will be visible in the **Details** tab of a Lig
 		class USphereComponent* Sphere1;
 * Declare your constructor, so you can set default values for your components and variables:
 
-		ALightSwitchCodeOnly(const FObjectInitializer& ObjectInitializer);
+		ALightSwitchCodeOnly();
 
 * Declare `OnOverlapBegin` and `OnOverlapEnd`, the functions that will be called when another Actor enters or leaves the SphereComponent. Note that these have different signatures.
 
@@ -99,7 +99,7 @@ LightSwitchCodeOnly.h
 			UPROPERTY(VisibleAnywhere, Category = "Switch Components")
 			class USphereComponent* Sphere1;
 
-			ALightSwitchCodeOnly(const FObjectInitializer& ObjectInitializer);
+			ALightSwitchCodeOnly();
 
 			/** called when something enters the sphere component */
 			UFUNCTION()
@@ -124,8 +124,7 @@ The **C++ Class Wizard** also creates the source file for the class, in this cas
 
 You will begin by adding the class constructor.
 
-		ALightSwitchCodeOnly::ALightSwitchCodeOnly(const FObjectInitializer& ObjectInitializer)
-			: Super(ObjectInitializer)
+		ALightSwitchCodeOnly::ALightSwitchCodeOnly()
 		{
 
 		}
@@ -136,14 +135,14 @@ In the `LightSwitchCodeOnly` constructor, C++ is used to:
 		DesiredIntensity = 3000.0f;
 * Create the PointLightComponent, set its variables (including setting its intensity to the value of `DesiredIntensity`), and make it the root component.
 
-		PointLight1 = ObjectInitializer.CreateDefaultSubobject<UPointLightComponent>(this, "PointLight1");
+		PointLight1 = CreateDefaultSubobject<UPointLightComponent>(TEXT("PointLight1"));
 		PointLight1->Intensity = DesiredIntensity;
 		PointLight1->bVisible = true;
 		RootComponent = PointLight1;
 
 * Create the SphereComponent, set its variables, and attach it to the PointLightComponent.
 
-		Sphere1 = ObjectInitializer.CreateDefaultSubobject<USphereComponent>(this, TEXT("Sphere1"));
+		Sphere1 = CreateDefaultSubobject<USphereComponent>(TEXT("Sphere1"));
 		Sphere1->InitSphereRadius(250.0f);
 		Sphere1->AttachParent = RootComponent;
 
@@ -163,17 +162,16 @@ LightSwitchCodeOnly.cpp
 	#include "BasicClasses.h"
 	#include "LightSwitchCodeOnly.h"
 
-	ALightSwitchCodeOnly::ALightSwitchCodeOnly(const FObjectInitializer& ObjectInitializer)
-		: Super(ObjectInitializer)
+	ALightSwitchCodeOnly::ALightSwitchCodeOnly()
 	{
 		DesiredIntensity = 3000.0f;
 
-		PointLight1 = ObjectInitializer.CreateDefaultSubobject<UPointLightComponent>(this, "PointLight1");
+		PointLight1 = CreateDefaultSubobject<UPointLightComponent>(TEXT("PointLight1"));
 		PointLight1->Intensity = DesiredIntensity;
 		PointLight1->bVisible = true;
 		RootComponent = PointLight1;
 
-		Sphere1 = ObjectInitializer.CreateDefaultSubobject<USphereComponent>(this, TEXT("Sphere1"));
+		Sphere1 = CreateDefaultSubobject<USphereComponent>(TEXT("Sphere1"));
 		Sphere1->InitSphereRadius(250.0f);
 		Sphere1->AttachParent = RootComponent;
 
