#UE4doc #docs:
Updated engine version field
Changed font settings for multiple terms
Changed multiple code snippets and steps as they were outdated
Replaced a couple images due to outdated UI elements

diff --git a/Engine/Documentation/Source/Gameplay/ClassCreation/CodeAndBlueprints/CodeAndBlueprints.INT.udn b/Engine/Documentation/Source/Gameplay/ClassCreation/CodeAndBlueprints/CodeAndBlueprints.INT.udn
--- a/Engine/Documentation/Source/Gameplay/ClassCreation/CodeAndBlueprints/CodeAndBlueprints.INT.udn
+++ b/Engine/Documentation/Source/Gameplay/ClassCreation/CodeAndBlueprints/CodeAndBlueprints.INT.udn
@@ -2,7 +2,7 @@ Availability:Public
 Title:C++ and Blueprints
 Crumbs: %ROOT%, Gameplay, Gameplay/ClassCreation
 Description:Introductory information for gameplay programmers getting started with Unreal Engine.
-Version: 4.6
+Version: 4.9
 
 C++ classes can be extended with Blueprints, allowing programmers to set up new gameplay classes in code that can be built upon and changed with Blueprints by level designers.
 There are a number of specifiers that change how a C++ class interacts with the Blueprint system, and some of them will be highlighted in this example.
@@ -15,10 +15,10 @@ Most of the code setup in the LightSwitchBoth class is similar to the code in th
 extend the LightSwitchCodeOnly class with a Blueprint, the components, variables, and functions created in that class would not be accessible to Blueprint graphs. This example will go over the `UPROPERTY()` 
 and `UFUNCTION()` specifiers that make LightSwitchBoth work as a template for the Blueprint that will derive from it.
 
-You may find it helpful to first refer to the [C++ Only LightSwitch example](Gameplay\ClassCreation\CodeOnly), to see how the header file and source file are set up
+You may find it helpful to first refer to the C++ Only LightSwitch example, to see how the header file and source file are set up
 to create the LightSwitchComponent, SphereComponent, DesiredIntensity variable, and OnOverlap function.
 
-This header file is adapted from the [C++ Only LightSwitch example](Gameplay\ClassCreation\CodeOnly) to add the following functionality: 
+This header file is adapted from the C++ Only LightSwitch example to add the following functionality: 
 
 * The PointLightComponent and the SphereComponent are BlueprintReadOnly, and will show up in the category **Switch Components** in the **My Blueprint** tab.
 * OnOverlapBegin and OnOverlapEnd are now BlueprintNativeEvents, and will show up in the category **Switch Functions** in the **My Blueprint** tab.
@@ -58,16 +58,20 @@ LightSwitchBoth.h
 			UPROPERTY(VisibleAnywhere, BlueprintReadOnly, Category="Switch Components")
 			class USphereComponent* Sphere1;
 
-			ALightSwitchBoth(const FObjectInitializer& ObjectInitializer);
+			ALightSwitchBoth();
 
 			/** called when something enters the sphere component */
 			UFUNCTION(BlueprintNativeEvent, Category="Switch Functions")
 			void OnOverlapBegin(class AActor* OtherActor, class UPrimitiveComponent* OtherComp, int32 OtherBodyIndex, bool bFromSweep, const FHitResult& SweepResult);
-	
+			
+			void OnOverlapBegin_Implementation(class AActor* OtherActor, class UPrimitiveComponent* OtherComp, int32 OtherBodyIndex, bool bFromSweep, const FHitResult& SweepResult);
+
 			/** called when something leaves the sphere component */
 			UFUNCTION(BlueprintNativeEvent, Category="Switch Functions")
 			void OnOverlapEnd(class AActor* OtherActor, class UPrimitiveComponent* OtherComp, int32 OtherBodyIndex);
 
+			void OnOverlapEnd_Implementation(class AActor* OtherActor, class UPrimitiveComponent* OtherComp, int32 OtherBodyIndex);
+
 			/** Toggles the light component's visibility*/
 			UFUNCTION()
 			void ToggleLight();
@@ -93,18 +97,17 @@ LightSwitchBoth.cpp
 	#include "BasicClasses.h"
 	#include "LightSwitchBoth.h"
 
-	ALightSwitchBoth::ALightSwitchBoth(const FObjectInitializer& ObjectInitializer)
-		: Super(ObjectInitializer)
+	ALightSwitchBoth::ALightSwitchBoth()
 	{
 
 		DesiredIntensity = 3000.0f;
 
-		PointLight1 = ObjectInitializer.CreateDefaultSubobject<UPointLightComponent>(this, "PointLight1");
+		PointLight1 = CreateDefaultSubobject<UPointLightComponent>(TEXT("PointLight1"));
 		PointLight1->Intensity = DesiredIntensity;
 		PointLight1->bVisible = true;
 		RootComponent = PointLight1;
 
-		Sphere1 = ObjectInitializer.CreateDefaultSubobject<USphereComponent>(this, TEXT("Sphere1"));
+		Sphere1 = CreateDefaultSubobject<USphereComponent>(this, TEXT("Sphere1"));
 		Sphere1->InitSphereRadius(250.0f);
 		Sphere1->AttachParent = RootComponent;
 
@@ -147,16 +150,16 @@ In this case, LightSwitchBoth is chosen as the parent class for the Blueprint, w
 
 The PointLightComponent and SphereComponent added in C++ also show in the **Components** tab in the **Blueprint Editor**.
 Their icons are dark blue to indicate that they are native components, inherited from the parent LightSwitchBoth class.  Any new components added just to the **LightSwitchBoth_BP** Blueprint 
-would have light blue icons instead. To learn more about adding and arranging components using the **Components** tab, see the [](Engine/Blueprints\UserGuide\Components) tab documentation.
+would have light blue icons instead. To learn more about adding and arranging components using the **Components** tab, see the [Components](Engine/Blueprints\UserGuide\Components) tab documentation.
 
 ![](Both_ComponentList.png)
 
-The **Blueprint Editor**'s **Graph Panel** is the core of Blueprint editing. In the **Graph Panel**, you can add new [variables](Engine/Blueprints\UserGuide\Variables), 
+The **Blueprint Editor**'s **Graph** panel is the core of Blueprint editing. In the **Graph** panel, you can add new [variables](Engine/Blueprints\UserGuide\Variables), 
 [functions](Engine/Blueprints\UserGuide\Functions), and [macros](Engine/Blueprints\UserGuide\Macros) in the [](Engine/Blueprints\Editor/UIComponents\MyBlueprint) tab.  You also have access to all the [graphs](Engine/Blueprints\UserGuide\Graphs) 
 contained within the Blueprint Class. In graphs, [nodes](Engine/Blueprints\UserGuide\Nodes) are wired together to create design-time and gameplay functionality that can be driven by class variables, gameplay events,
 and even an Actor's surroundings. 
 
-The **My Blueprint** tab in **Graph Panel** displays the PointLightComponent and the SphereComponent added in C++ to the LightSwitchBoth class. This is because of the `BlueprintReadOnly` specifier.
+The **My Blueprint** tab in **Graph** panel displays the PointLightComponent and the SphereComponent added in C++ to the LightSwitchBoth class. This is because of the `BlueprintReadOnly` specifier.
 Nodes for these components can be added to a graph by clicking and dragging from their names in the **My Blueprint** tab into the graph.  These nodes can then be connected to nodes that change
 variables like visibility or light color. The **DesiredIntensity** variable also appears in the **My Blueprint** tab.  Because it is a variable and not a component, the `BlueprintReadWrite` specifier
 was able to be used. This means it is possible to create nodes to both get and set the value of **DesiredIntensity** in a Blueprint graph. See the [](Engine/Blueprints\Editor/UIComponents\MyBlueprint) documentation for general usage information.
@@ -177,7 +180,7 @@ event. Without the Construction Script setup, a new LightSwitchBoth_BP Actor wou
 Actor is moved within the level, as well as when DesiredIntensity is changed. Using the Construction Script means that Actor variables exposed to Blueprints can be changed easily and the 
 effects of these changes can be seen quickly.
 
-In the **LightSwitchBoth_BP** class, the **Construction Script** event is connected to a **Set Brightness** node, so that the brightness of **Point Light 1** (the PointLightComponent) is set to 
+In the **LightSwitchBoth_BP** class, the **Construction Script** event is connected to a **Set Intensity** node, so that the brightness of **Point Light 1** (the PointLightComponent) is set to 
 the value of **DesiredIntensity** when the Actor is added to or moved within the level or DesiredIntensity changes. 
 
 ![](Both_ConstructionScript.png)
@@ -195,14 +198,14 @@ For more on events and working with graphs, see the [](Engine/Blueprints\UserGui
 
 ![](Both_EventGraph_2.png)
 
-The **DesiredIntensity** variable is set to EditAnywhere in the LightSwitchBoth header file, so it is visible in the **Blueprint Editor**'s Defaults and can be edited by clicking the ![](Engine\Blueprints\UserGuide\Variables\classDefaults_button.png) to show the Class Defaults in the Details Panel.
+The **DesiredIntensity** variable is set to EditAnywhere in the LightSwitchBoth header file, so it is visible in the **Blueprint Editor**'s Defaults and can be edited by clicking the **Class Defaults** button to show the Class Defaults in the Details Panel.
 This also means that the variable can be changed for each instance of the class, so each Actor can have its own DesiredIntensity. Because DesiredIntensity is also BlueprintReadWrite and used in the 
 Construction Script, updating it also causes the Construction Script to execute again.
 
 ![](Both_Defaults.png)
 
 Blueprinted classes can be extended with other Blueprint Classes, either using the dropdown button next to the class in the **Class Viewer** to create a new Blueprint, or by
-**right-clicking** on the Blueprint and selecting **Create New Blueprint Based on This**.
+**Right-clicking** on the Blueprint and selecting **Create New Blueprint Based on This**.
 
 The Blueprint Class LightSwitchBoth_BP is in the [](Engine/Content/Browser), and can be dragged into the level from there. It is also in the [](Engine/UI/ClassViewer).
 For more information on placing Actors in a level using the Content Browser or Class Viewer, see the [](Engine/Actors/Placement) documentation.
