#UE4doc #docs - Added engine version field.

diff --git a/Engine/Documentation/Source/Programming/Development/Tools/EngineProgrammingTools.INT.udn b/Engine/Documentation/Source/Programming/Development/Tools/EngineProgrammingTools.INT.udn
--- a/Engine/Documentation/Source/Programming/Development/Tools/EngineProgrammingTools.INT.udn
+++ b/Engine/Documentation/Source/Programming/Development/Tools/EngineProgrammingTools.INT.udn
@@ -2,6 +2,7 @@ Availability:Public
 Title:Programming Tools
 Crumbs: %ROOT%, Programming, Programming\Development
 Description:Applications, tools, and add-ins that are useful for programmers working with Unreal Engine.
+Version: 4.9
 
 
 [OBJECT:TopicList]
