#UE4doc #docs - Fixed grammatical error.

diff --git a/Engine/Documentation/Source/Programming/Development/Tools/ConsoleManager/ConsoleManager.INT.udn b/Engine/Documentation/Source/Programming/Development/Tools/ConsoleManager/ConsoleManager.INT.udn
--- a/Engine/Documentation/Source/Programming/Development/Tools/ConsoleManager/ConsoleManager.INT.udn
+++ b/Engine/Documentation/Source/Programming/Development/Tools/ConsoleManager/ConsoleManager.INT.udn
@@ -220,7 +220,7 @@ Other settings can come from the Scalability feature. Look at **Config/BaseScala
 
 ## Priority
 
-Console variables can be override from various sources e.g. user/editor/project settings, command line, consolevariables.ini, ...
+Console variables can be overridden from various sources e.g. user/editor/project settings, command line, consolevariables.ini, ...
 In order to be able to reapply some settings (e.g. project settings can be changed in editor UI) while keeping the specified overrides (e.g. from command line),
 we introduced a priority. Now all settings can be applied in any order. 
 
