#UE4doc #docs - Added Engine Version Field. Removed extraneous bolding of terms.

diff --git a/Engine/Documentation/Source/Programming/Development/Tools/SymbolDebugger/SymbolDebugger.INT.udn b/Engine/Documentation/Source/Programming/Development/Tools/SymbolDebugger/SymbolDebugger.INT.udn
--- a/Engine/Documentation/Source/Programming/Development/Tools/SymbolDebugger/SymbolDebugger.INT.udn
+++ b/Engine/Documentation/Source/Programming/Development/Tools/SymbolDebugger/SymbolDebugger.INT.udn
@@ -8,9 +8,9 @@ The **Symbol Debugger** program is used to assist in syncing the files required
 
 ![Symbol Debugger](symbol_debugger.png)
 
-The **Symbol Debugger** utilizes source indexing of builds generated by the build system and a local symbol store to allow debugging without requiring syncing source files to the time the build was made. 
+The Symbol Debugger utilizes source indexing of builds generated by the build system and a local symbol store to allow debugging without requiring syncing source files to the time the build was made. 
 
-Debugging a minidump file from a crash report is now as easy as simply dragging the dmp file and dropping it onto the Symbol Debugger app. (There is some setup required to properly utilize the **Symbol Debugger**. This will be explained later in this document.)
+Debugging a minidump file from a crash report is now as easy as simply dragging the dmp file and dropping it onto the Symbol Debugger app. (There is some setup required to properly utilize the Symbol Debugger. This will be explained later in this document.)
 
 ## UI Description
 
@@ -28,7 +28,7 @@ Debugging a minidump file from a crash report is now as easy as simply dragging
 
 ### Method Selection
 
-There are four methods allowed when running the **Symbol Debugger**. The current method is selected via the radio buttons in the section labeled **Method Selection**. Depending on the method selected, various controls will be enabled or disabled to allow for the user to enter more information. The following methods are supported:
+There are four methods allowed when running the Symbol Debugger. The current method is selected via the radio buttons in the section labeled **Method Selection**. Depending on the method selected, various controls will be enabled or disabled to allow for the user to enter more information. The following methods are supported:
 
 | Method | Description |
 | ------ | ----------- |
@@ -60,29 +60,29 @@ This value will also need to be utilized in [Setting up Visual Studio](#Settingu
 
 ### Platform
 
-The **Platform** section will contain the platform the build was made for. When the method is set to anything other than **CrashDump**, this value will have to be set by you to the proper platform string (ie, **Win32** or **Win64**).
+The **Platform** section will contain the platform the build was made for. When the method is set to anything other than CrashDump, this value will have to be set by you to the proper platform string (ie, **Win32** or **Win64**).
 
 ### Engine Version
 
-The **Engine Version** section will contain the engine version of the build of interest. If **CrashDump**, **Changelist**, or **SourceLabel** are the method, this field will not be editable but will contain the results determined by inspecting the selected method. If EngineVersion is the method, this field will be filled with the contents of the 'Method Input' text box upon commit.
+The **Engine Version** section will contain the engine version of the build of interest. If CrashDump, Changelist, or SourceLabel are the method, this field will not be editable but will contain the results determined by inspecting the selected method. If EngineVersion is the method, this field will be filled with the contents of the 'Method Input' text box upon commit.
 
 ### Changelist Number
 
-The **Changelist Number** section will contain the changelist number of the build of interest. If **CrashDump**, **EngineVersion**, or **SourceLabel** are the method, this field will not be editable but will contain the results determined by inspecting the selected method. If Changelist is the method, this field will be filled with the contents of the 'Method Input' text box upon commit.
+The **Changelist Number** section will contain the changelist number of the build of interest. If CrashDump, EngineVersion, or SourceLabel are the method, this field will not be editable but will contain the results determined by inspecting the selected method. If Changelist is the method, this field will be filled with the contents of the 'Method Input' text box upon commit.
 
 ### Source Control Label
 
-The **Source Control Label** section will contain the engine version of the build of interest. If **CrashDump**, **EngineVersion**, or **Changelist** are the method, this field will not be editable but will contain the results determined by inspecting the selected method. If SourceLabel is the method, this field will be filled with the contents of the 'Method Input' text box upon commit.
+The **Source Control Label** section will contain the engine version of the build of interest. If CrashDump, EngineVersion, or Changelist are the method, this field will not be editable but will contain the results determined by inspecting the selected method. If SourceLabel is the method, this field will be filled with the contents of the 'Method Input' text box upon commit.
 
 ### Action Buttons
 
-The 'Action Buttons' section contains the buttons that trigger various actions. They are enabled depending on the current status of the selected method. The following buttons are present:
+The **Action Buttons** section contains the buttons that trigger various actions. They are enabled depending on the current status of the selected method. The following buttons are present:
 
 | Button | Description |
 | ------ | ----------- |
-| Inspect | The Inspect button is used to evaluate the supplied information and determine the source control label of the build. It will be enabled if the entered information is enough for the current method to be evaluated. For example, if CrashDump is the selected method, simply having set a crash dump file name and local symbol store is all that is required. If any other method is active, however, the Platform field will have to be filled in as well. |
-| Sync | Once a valid build label has been found, the Sync button will become enabled. When clicked, all the files required for debugging will be synced from source control into the local symbol store location. |
-| Debug | Currently only valid for CrashDump evaluation, when clicked, the Debug button will launch the debugger with the supplied minidump. |
+| **Inspect** | The Inspect button is used to evaluate the supplied information and determine the source control label of the build. It will be enabled if the entered information is enough for the current method to be evaluated. For example, if CrashDump is the selected method, simply having set a crash dump file name and local symbol store is all that is required. If any other method is active, however, the Platform field will have to be filled in as well. |
+| **Sync** | Once a valid build label has been found, the Sync button will become enabled. When clicked, all the files required for debugging will be synced from source control into the local symbol store location. |
+| **Debug** | Currently only valid for CrashDump evaluation, when clicked, the Debug button will launch the debugger with the supplied minidump. |
 
 ## Setting up Visual Studio
 
