#UE4doc #docs - Added engine version field.

diff --git a/Engine/Documentation/Source/Programming/Gameplay/DataDriven/DataDrivenGameplay.INT.udn b/Engine/Documentation/Source/Programming/Gameplay/DataDriven/DataDrivenGameplay.INT.udn
--- a/Engine/Documentation/Source/Programming/Gameplay/DataDriven/DataDrivenGameplay.INT.udn
+++ b/Engine/Documentation/Source/Programming/Gameplay/DataDriven/DataDrivenGameplay.INT.udn
@@ -2,6 +2,7 @@ Availability:Public
 Title: Data Driven Gameplay Elements
 Crumbs:%ROOT%, Programming
 Description:Driving gameplay elements using externally stored data.
+Version: 4.9
 
 [OBJECT:Redirect]
 [PARAM:link]
