#UE4doc #docs - Added engine version field.

diff --git a/Engine/Documentation/Source/Programming/Gameplay/Framework/GameplayFramework.INT.udn b/Engine/Documentation/Source/Programming/Gameplay/Framework/GameplayFramework.INT.udn
--- a/Engine/Documentation/Source/Programming/Gameplay/Framework/GameplayFramework.INT.udn
+++ b/Engine/Documentation/Source/Programming/Gameplay/Framework/GameplayFramework.INT.udn
@@ -2,6 +2,7 @@ Availability:Public
 Title:Gameplay Framework
 Crumbs:%ROOT%, Programming
 Description:Core systems, such as game rules, player input and controls, cameras, and user interfaces.
+Version: 4.9
 
 
 [OBJECT:Redirect]
