#UE4doc #docs - Added engine version field.

diff --git a/Engine/Documentation/Source/Programming/Gameplay/Framework/Camera/Camera.INT.udn b/Engine/Documentation/Source/Programming/Gameplay/Framework/Camera/Camera.INT.udn
--- a/Engine/Documentation/Source/Programming/Gameplay/Framework/Camera/Camera.INT.udn
+++ b/Engine/Documentation/Source/Programming/Gameplay/Framework/Camera/Camera.INT.udn
@@ -2,6 +2,7 @@ Availability:Public
 Title:Camera
 Crumbs:%ROOT%, Gameplay, Gameplay/Framework
 Description:The **Camera** represents the player's point of view; how the player sees the world.
+Version: 4.9
 
 [OBJECT:Redirect]
 [PARAM:link]
diff --git a/Engine/Documentation/Source/Programming/Gameplay/Framework/Controller/Controller.INT.udn b/Engine/Documentation/Source/Programming/Gameplay/Framework/Controller/Controller.INT.udn
--- a/Engine/Documentation/Source/Programming/Gameplay/Framework/Controller/Controller.INT.udn
+++ b/Engine/Documentation/Source/Programming/Gameplay/Framework/Controller/Controller.INT.udn
@@ -2,6 +2,7 @@ Availability:Public
 Title:Controller
 Crumbs:%ROOT%, Gameplay, Gameplay/Framework
 Description: In the context of a player or AI entity, the **Controller** is essentially the brain.
+Version: 4.9
 
 [OBJECT:Redirect]
 [PARAM:link]
