#UE4doc #docs - added engine version field.

diff --git a/Engine/Documentation/Source/Programming/Gameplay/Framework/Controller/PlayerController/PlayerController.INT.udn b/Engine/Documentation/Source/Programming/Gameplay/Framework/Controller/PlayerController/PlayerController.INT.udn
--- a/Engine/Documentation/Source/Programming/Gameplay/Framework/Controller/PlayerController/PlayerController.INT.udn
+++ b/Engine/Documentation/Source/Programming/Gameplay/Framework/Controller/PlayerController/PlayerController.INT.udn
@@ -2,6 +2,7 @@ Availability:Public
 Title:PlayerController
 Crumbs:%ROOT%, Gameplay, Gameplay/Framework, Gameplay/Framework/Controller
 Description:The PlayerController implements functionality for taking the input data from the player and translating that into actions, such as movement, using items, firing weapons, etc.
+Version: 4.9
 
 [OBJECT:Redirect]
 [PARAM:link]
