#UE4doc #docs - added engine version field.

diff --git a/Engine/Documentation/Source/Programming/Gameplay/Framework/GameMode/GameMode.INT.udn b/Engine/Documentation/Source/Programming/Gameplay/Framework/GameMode/GameMode.INT.udn
--- a/Engine/Documentation/Source/Programming/Gameplay/Framework/GameMode/GameMode.INT.udn
+++ b/Engine/Documentation/Source/Programming/Gameplay/Framework/GameMode/GameMode.INT.udn
@@ -2,6 +2,7 @@ Availability:Public
 Title:GameMode
 Crumbs:%ROOT%, Gameplay/Framework
 Description:The GameMode is responsible for setting the rules of the game.
+Version: 4.9
 
 [OBJECT:Redirect]
 [PARAM:link]
diff --git a/Engine/Documentation/Source/Programming/Gameplay/Framework/GameState/GameState.INT.udn b/Engine/Documentation/Source/Programming/Gameplay/Framework/GameState/GameState.INT.udn
--- a/Engine/Documentation/Source/Programming/Gameplay/Framework/GameState/GameState.INT.udn
+++ b/Engine/Documentation/Source/Programming/Gameplay/Framework/GameState/GameState.INT.udn
@@ -2,6 +2,7 @@ Availability:Public
 Title:GameState
 Crumbs:%ROOT%, Gameplay/Framework
 Description:The GameState is responsible for monitoring the state of the game on the client.
+Version: 4.9
 
 [OBJECT:Redirect]
 [PARAM:link]
diff --git a/Engine/Documentation/Source/Programming/Gameplay/Framework/Input/Input.INT.udn b/Engine/Documentation/Source/Programming/Gameplay/Framework/Input/Input.INT.udn
--- a/Engine/Documentation/Source/Programming/Gameplay/Framework/Input/Input.INT.udn
+++ b/Engine/Documentation/Source/Programming/Gameplay/Framework/Input/Input.INT.udn
@@ -2,6 +2,7 @@ Availability: Public
 Title:Input
 Crumbs:%ROOT%, Gameplay
 Description:The **Input** object is responsible for converting input from the player into data in a form Actors can understand and make use of.
+Version: 4.9
 
 [OBJECT:Redirect]
 [PARAM:link]
diff --git a/Engine/Documentation/Source/Programming/Gameplay/Framework/Pawn/Pawn.INT.udn b/Engine/Documentation/Source/Programming/Gameplay/Framework/Pawn/Pawn.INT.udn
--- a/Engine/Documentation/Source/Programming/Gameplay/Framework/Pawn/Pawn.INT.udn
+++ b/Engine/Documentation/Source/Programming/Gameplay/Framework/Pawn/Pawn.INT.udn
@@ -2,6 +2,7 @@ Availability:Public
 Title:Pawn
 Crumbs:%ROOT%, Gameplay, Gameplay/Framework
 Description:The **Pawn** is the physical representation of a player within the world.
+Version: 4.9
 
 [OBJECT:Redirect]
 [PARAM:link]
