#UE4doc #docs - added engine version field.

diff --git a/Engine/Documentation/Source/Programming/Gameplay/Framework/Pawn/Character/Character.INT.udn b/Engine/Documentation/Source/Programming/Gameplay/Framework/Pawn/Character/Character.INT.udn
--- a/Engine/Documentation/Source/Programming/Gameplay/Framework/Pawn/Character/Character.INT.udn
+++ b/Engine/Documentation/Source/Programming/Gameplay/Framework/Pawn/Character/Character.INT.udn
@@ -2,6 +2,7 @@ Availability:Public
 Title:Character
 Crumbs:%ROOT%, Gameplay, Gameplay/Framework, Gameplay/Framework/Pawn
 Description:A **Character** is a Pawn which has some basic bipedal movement functionality by default.
+Version: 4.9
 
 
 [OBJECT:Redirect]
diff --git a/Engine/Documentation/Source/Programming/Gameplay/Framework/QuickReference/QuickReference.INT.udn b/Engine/Documentation/Source/Programming/Gameplay/Framework/QuickReference/QuickReference.INT.udn
--- a/Engine/Documentation/Source/Programming/Gameplay/Framework/QuickReference/QuickReference.INT.udn
+++ b/Engine/Documentation/Source/Programming/Gameplay/Framework/QuickReference/QuickReference.INT.udn
@@ -2,6 +2,7 @@ Availability:Public
 Title:Gameplay Framework Quick Reference
 Crumbs:%ROOT%, Gameplay, Gameplay/Framework
 Description:Short overview of classes for game rules, characters, controllers, user interfaces, etc., that make up the framework for the game.
+Version: 4.9
 
 [OBJECT:Redirect]
 [PARAM:link]
diff --git a/Engine/Documentation/Source/Programming/Gameplay/Framework/UIAndHUD/UIAndHUD.INT.udn b/Engine/Documentation/Source/Programming/Gameplay/Framework/UIAndHUD/UIAndHUD.INT.udn
--- a/Engine/Documentation/Source/Programming/Gameplay/Framework/UIAndHUD/UIAndHUD.INT.udn
+++ b/Engine/Documentation/Source/Programming/Gameplay/Framework/UIAndHUD/UIAndHUD.INT.udn
@@ -2,7 +2,7 @@ Availability: Public
 Title:User Interfaces & HUDs
 Crumbs:%ROOT%, Gameplay, Gameplay/Framework
 Description:Guides and information for artists and programmers creating user interfaces such as menus and HUDs.
-
+Version: 4.9
 
 [OBJECT:Redirect]
 [PARAM:link]
