#UE4doc #docs - Updated image.

diff --git a/Engine/Documentation/Source/Programming/UnrealBuildSystem/Configuration/BuildConfiguration.INT.udn b/Engine/Documentation/Source/Programming/UnrealBuildSystem/Configuration/BuildConfiguration.INT.udn
--- a/Engine/Documentation/Source/Programming/UnrealBuildSystem/Configuration/BuildConfiguration.INT.udn
+++ b/Engine/Documentation/Source/Programming/UnrealBuildSystem/Configuration/BuildConfiguration.INT.udn
@@ -18,7 +18,7 @@ $ bOmitPCDebugInfoInDevelopment : If true, debug info is not included when build
 $ bUseUnityBuild : Set to true when rebuilding the engine as a whole. This should be disabled when using fast iteration mode. See [Build Modes](#BuildModes) for more information.
 $ bEnableFastIteration : If true, compile and link times are dramatically decreased. This should be disabled when creating unity builds. See [Build Modes](#BuildModes) for more information.
     * Setting to _true_ gives fastest iteration times.
-    * If you encounter weird "RPC" or "PDB" link errors, then you should set `bUseIncrementalLinking = false` in **ValidateConfiguration()**
+    * If you encounter weird "RPC" or "PDB" link errors, then you should set `bUseIncrementalLinking = false` in `ValidateConfiguration()`
 
 
 ## Build Modes
