#UE4doc #docs - Added engine version field. Fixed Chinese text that was carried over to Title and Description for Video Tutorials in English page.

diff --git a/Engine/Documentation/Source/Resources/Resources.INT.udn b/Engine/Documentation/Source/Resources/Resources.INT.udn
--- a/Engine/Documentation/Source/Resources/Resources.INT.udn
+++ b/Engine/Documentation/Source/Resources/Resources.INT.udn
@@ -3,6 +3,7 @@ Title:Samples and Tutorials
 Crumbs:%ROOT%
 Description:Links to various example scenes, sample games, and tutorials.
 Navigation:topic
+Version: 4.9
 
 [VAR:Params]
 [PARAM:image]
diff --git a/Engine/Documentation/Source/Videos/VideoTutorials.INT.udn b/Engine/Documentation/Source/Videos/VideoTutorials.INT.udn
--- a/Engine/Documentation/Source/Videos/VideoTutorials.INT.udn
+++ b/Engine/Documentation/Source/Videos/VideoTutorials.INT.udn
@@ -1,5 +1,5 @@
-Title:视频教学
-Description:展示虚幻引擎4中理念和技巧的视频教程
+Title:Video Tutorials
+Description: Video tutorials demonstrating concepts and techniques in Unreal Engine 4.
 Crumbs:%ROOT%
 Availability:Public
 
