#UE4 Docs: Adding in a section about lightmass and base material overrides.

diff --git a/Engine/Documentation/Source/Engine/Rendering/Materials/MaterialInstances/Editor/MaterialInstanceEditor.INT.udn b/Engine/Documentation/Source/Engine/Rendering/Materials/MaterialInstances/Editor/MaterialInstanceEditor.INT.udn
--- a/Engine/Documentation/Source/Engine/Rendering/Materials/MaterialInstances/Editor/MaterialInstanceEditor.INT.udn
+++ b/Engine/Documentation/Source/Engine/Rendering/Materials/MaterialInstances/Editor/MaterialInstanceEditor.INT.udn
@@ -71,8 +71,6 @@ The **Details** pane in the Material Instance editor is where all of the 'work'
 * **Parameter Groups** - In this section you can adjust all of the values that have been exposed to the Material Instance from the Parent Material. 
 * **General** - The General section allows you to change the Parent Material, Phys Material, adjust how this Material Instance will influence the Lightmass build and override some of the properties inherited from the parent Material. 
 
-
-
 ### Parent List
 
 ![](MIE_Parent.png)
@@ -118,15 +116,31 @@ You can also enable / disable real time preview by pressing **CTRL + R** at the
 
 [INCLUDE:Engine/Rendering/Materials/MaterialInstances#MakingAnInstance]
 
-## Overriding Parameters
+## Overriding Material Instance Parameters
 
-The checkboxes to the left side of each of the parameter names specify whether or not an overridden value should be used. So below **Base Metal, Base Roughness, Chrome Roughness,** and **OPACITY Chrome** are being 
-used while **Base Scratch Roughness, Base UV Tile,** and **Chrome Metal** are not.  Only checked values are stored for memory conservation reasons, so unchecked values will be lost when 
-the Material Instance Editor window is closed.
+The checkboxes to the left side of each of the parameter names specify whether or not an overridden value should be used.
+So below **Base Metal, Base Roughness, Chrome Roughness,** and **OPACITY Chrome** are being used while **Base Scratch Roughness, Base UV Tile,** and **Chrome Metal** are not.
+Only checked values are stored for memory conservation reasons, so unchecked values will be lost when the Material Instance Editor window is closed.
 
 ![](MIE_Overriding.png)
 
+### Overriding Material Instance Lightmass Settings
+You can also override parameters that deal with how the Material Instance interacts with Lightmass by expanding the **Lightmass Settings** section and then enabling the property you want to override by clicking on it.
+Once a Lightmass property override is enabled, you can then input the new data you want to use in the **Parameter Value** section. 
+In the following image we have taken the Material Instance for the sky box and overridden it's **Emissive Boost** to be 2.0 instead of 1.0. When the level lighting is built this will make the Emissive parts of the Material cast more light into the Lightmass build making the levels lighting brighter overall.
+
+![](Override_Lightmass_Settings.png)
+
+### Overriding Parent Material Settings
+In some cases the Material Instances that are available might not fit your exact needs. 
+For example say you need the Material Instance to be two sided instead of single sided. 
+To fix this you could go to the Parent Material for that Material Instance and set it to be two sided but this would make any other Material Instances that use that Parent Material two sided which is not what we want. 
+Instead you can override certain Parent Material properties within a Material Instance by going to the **Material Property Overrides** section of the Material Instance and then enabling the property you want to override.
+
+![](Overried_Parent_Material_Settings.png)
 
+In the image above we have set the Material Instance for the sky box  to be Two Sided instead of single sided by enabling the **Two Sided** option that is under the Material Property Overrides section. 
+This will make the Material Instance always render even if it is put onto geometry that has its faces turned so that you can not see them.
 
 ## Workflow
 
