#UE4doc #docs -  Added Engine version field. Added Skill Level Field. Broke up single images to be used with comparison slider and image tables.

diff --git a/Engine/Documentation/Source/Resources/Community/SciFiBunk_MaterialInstancing/MasterMaterialsandMaterialInstancing.INT.udn b/Engine/Documentation/Source/Resources/Community/SciFiBunk_MaterialInstancing/MasterMaterialsandMaterialInstancing.INT.udn
--- a/Engine/Documentation/Source/Resources/Community/SciFiBunk_MaterialInstancing/MasterMaterialsandMaterialInstancing.INT.udn
+++ b/Engine/Documentation/Source/Resources/Community/SciFiBunk_MaterialInstancing/MasterMaterialsandMaterialInstancing.INT.udn
@@ -1,6 +1,8 @@
 Availability: Public
 Title:SciFi Bunk: Material Instancing
 Description: This document covers how the SciFi Bunk project uses Master Materials and Material Instancing in the building of its environments.
+Version: 4.9
+SkillLevel: Intermediate
 
 [VAR:Params]
 [PARAM:image]
@@ -64,13 +66,20 @@ The end goal of setting up this master material is to create a material that can
 
 The screenshot below shows just how much flexibility this will give your environments, allowing you to quickly change the way your scene looks by tweaking a few values. 
 
-![](image_0.jpg)
+[OBJECT:ComparisonSlider]
+	[PARAM:before]
+	![Material Instance Color: Orange](Image_0_Left.png)(w:400)
+	[/PARAM]
+	[PARAM:after]
+	![Material Instance Color: White](Image_0_Right.png)(w:400)
+	[/PARAM]
+[/OBJECT]
 
 ## Master Material Setup
 
 We will start by creating your master material. This is what all of the materials in your scene will be instanced from, so it needs to be set up correctly and efficiently as it likely get quite complex. One advantage to using a master material is that if you want to add more features to your materials later on, you need only apply it to the master material and it will propagate down to the instances.
 
-Here is an example showing the final master material setup. All of the input textures and constants in this material are parameterized to allow access in the instanced material. To do this, you can right click on the value or map and click **Convert to parameter**, and then you can assign it a name. This will allow you to change it quickly in the material instance.
+Here is an example showing the final master material setup. All of the input textures and constants in this material are parameterized to allow access in the instanced material. To do this, you can right-click on the value or map and click **Convert to parameter**, and then you can assign it a name. This will allow you to change it quickly in the material instance.
 
 Also, the **Static Switch** parameter is extensively used in this material, which allows you to choose between using certain setups such as detail normal maps, parallax etc. This will make your final material instance cheaper as you are able to switch off un-needed features. More about this can be read here: [Parameter Expressions](Engine/Rendering/Materials/ExpressionReference/Parameters)
 
@@ -88,6 +97,12 @@ To the left of the DiffMap texture input, you can see that a node has been plugg
 
 The following example shows the flexibility of this system. On the left, the material is using a single color for the albedo map. In the middle, it is using just an albedo map. On the right, it is using an albedo map that has been tinted orange. This shows how you can use the same maps to get a lot of variety in the engine.
 
+[REGION:imagetable]
+| ![](image_3_1.png)(w:300) | ![](image_3_2.png)(w:300) | ![](image_3_3.png)(w:300) |
+| ------ | ------ | ------ |
+| Single Color for the Albedo Map  | Albedo Map  | Albedo Map with Color |
+[/REGION]
+
 ![](image_3.jpg)
 
 ### Metallic Setup
@@ -100,7 +115,11 @@ This is useful because it allows you to choose between using a single value, whi
 
 The following example shows this in action. The material on the left uses a metallic value of 0, the middle a value of 1, and the one on the right uses a specific metallic map.
 
-![](image_5.jpg)
+[REGION:imagetable]
+| ![](image_5_1.png)(w:300) | ![](image_5_2.png)(w:300) | ![](image_5_3.png)(w:300) |
+| ------ | ------ | ------ |
+| Metallic Value - 0 | Metallic Value - 1 | Metallic Texture Map |
+[/REGION]
 
 ### Roughness Setup
 
@@ -110,13 +129,21 @@ Similar to the albedo setup, this allows you to choose between using a roughness
 
 The following example shows this working. The material on the left uses a roughness value of 1, the middle uses a roughness value of 0, and the one on the right uses a specific roughness map.
 
-![](image_7.jpg)
+[REGION:imagetable]
+| ![](image_7_1.png)(w:300) | ![](image_7_2.png)(w:300) | ![](image_7_3.png)(w:300) |
+| ------ | ------ | ------ |
+| Roughness Value - 1 | Roughness Value - 0 | Roughness Texture Map |
+[/REGION]
 
 The **UseDetailRoughness?** static switch lets you use a detail roughness map which is plugged into the **DetailTiling** constant which is multiplied with a TexCoord, to allow you to tile your detail roughness map. This value also plugs into the detail normal map. This is useful for materials that require a texture that can only be viewed up close and can have a very high tiling rate. 
 
 The following example shows this in action. The material on the left does not use a detail roughness/normal map, the one in the middle uses a detail map with a tiling rate of 1, and the one on the right uses a tiling rate of 4.
 
-![](image_75.jpg)
+[REGION:imagetable]
+| ![](image_75_1.png)(w:300) | ![](image_75_2.png)(w:300) | ![](image_75_3.png)(w:300) |
+| ------ | ------ | ------ |
+| Does not use a Normal Map | With Normal Map and Tiling set to 1 | With Normal Map and Tiling set to 4 |
+[/REGION]
 
 ### Occlusion Setup
 
@@ -132,7 +159,11 @@ This setup lets you choose between using an emissive on your material or not (by
 
 The following example shows how this emissive setup works. The material on the left has an emissive scale set to 0, the middle set to 1, and the one on the right is set to 2. 
 
-![](image_10.jpg)
+[REGION:imagetable]
+| ![](image_10_1.png)(w:300) | ![](image_10_2.png)(w:300) | ![](image_10_3.png)(w:300) |
+| ------ | ------ | ------ |
+| Emissive Value - 0 | Emissive Value - 1 | Emissive Value - 2 |
+[/REGION]
 
 ### Normal Setup
 
@@ -152,7 +183,15 @@ This setup is plugged into the **UV** slots of all the aforementioned sections,
 
 The following example shows how the Parallax effect works. On the left, Parallax is disabled, and on the right, Parallax is set to 0.007. It's a fairly subtle effect but may work well depending on the material in question.
 
-![](image_13.jpg)
+[OBJECT:ComparisonSlider]
+	[PARAM:before]
+	![Parallax Disabled](Image_13_Left.png)(w:400)
+	[/PARAM]
+	[PARAM:after]
+	![Parallax Enabled](Image_13_Right.png)(w:400)
+	[/PARAM]
+[/OBJECT]
+
 
 ## Material Instancing
 
