#UE4doc #docs - Added engine version field.

diff --git a/Engine/Documentation/Source/Resources/Concepts/GameplayConceptExamples.INT.udn b/Engine/Documentation/Source/Resources/Concepts/GameplayConceptExamples.INT.udn
--- a/Engine/Documentation/Source/Resources/Concepts/GameplayConceptExamples.INT.udn
+++ b/Engine/Documentation/Source/Resources/Concepts/GameplayConceptExamples.INT.udn
@@ -3,6 +3,8 @@ Title:Gameplay Concept Examples
 Crumbs:%ROOT%
 Description:These are framework examples, aimed at showing you how to add a certain design tool or game feature.
 Navigation:topic
+Version: 4.9
+
 
 [VAR:Params]
 [PARAM:image]
