#UE4doc #docs:
Added engine version field
Changed font settings for multiple terms
Replaced multiple images

diff --git a/Engine/Documentation/Source/Gameplay/Input/Input.INT.udn b/Engine/Documentation/Source/Gameplay/Input/Input.INT.udn
--- a/Engine/Documentation/Source/Gameplay/Input/Input.INT.udn
+++ b/Engine/Documentation/Source/Gameplay/Input/Input.INT.udn
@@ -2,6 +2,7 @@ Availability: Public
 Title:Input
 Crumbs:%ROOT%, Gameplay
 Description:The Input object is responsible for converting input from the player into data in a form Actors can understand and make use of.
+version: 4.9
 
 
 The **PlayerInput** object is responsible for converting input from the player into data in a form that Actors, like PlayerControllers or Pawns, 
@@ -20,7 +21,7 @@ Hardware input from a player is very straightforward.  It most commonly includes
 ## PlayerInput
 [EXCERPT:PlayerInput]
 PlayerInput is a UObject within the PlayerController class that manages player input. It is only spawned on the client.  Two structs are
-defined within PlayerInput.  The first, `FInputActionKeyMapping`, defines an ActionMapping. The other, `FInputAxisKeyMapping`, defines an AxisMapping. 
+defined within PlayerInput.  The first, **FInputActionKeyMapping**, defines an ActionMapping. The other, **FInputAxisKeyMapping**, defines an AxisMapping. 
 The hardware input definitions used in both ActionMappings and AxisMappings are established in InputCoreTypes.
 
 [EXCERPT:ActionMapping]
