#UE4doc #docs:
Added engine version field
Changed font settings for multiple terms
Replaced an image

diff --git a/Engine/Documentation/Source/Gameplay/DataDriven/DataDrivenGameplay.INT.udn b/Engine/Documentation/Source/Gameplay/DataDriven/DataDrivenGameplay.INT.udn
--- a/Engine/Documentation/Source/Gameplay/DataDriven/DataDrivenGameplay.INT.udn
+++ b/Engine/Documentation/Source/Gameplay/DataDriven/DataDrivenGameplay.INT.udn
@@ -2,26 +2,27 @@ Availability:Public
 Title: Data Driven Gameplay Elements
 Crumbs:%ROOT%, Gameplay
 Description:Driving gameplay elements using externally stored data.
+version: 4.9
 
 [TOC(start:2)]
 
 Data driven gameplay helps mitigate the amount of work and complexity involved, as well as providing 
 the ability to visualize and parameterize data creation and progression, for games that have an extended 
 lifetime well beyond that of a typical boxed game, and require constant tweaking and balancing of data 
-based on player feedback. The ability to move data out to **Excel** or other spreadsheet documents that can be 
+based on player feedback. The ability to move data out to **Microsoft Excel** or other spreadsheet documents that can be 
 maintained using proven tools and then imported to automatically take effect in the game.
 
 There are two new ways that data can be imported into UE4 through Excel documents:
 
-* DataTables
-* CurveTables
+* **DataTables**
+* **CurveTables**
 
 
-These Excel documents are `.xlsm` (macro enabled **Excel** documents) that have macro based export 
-buttons to enable easily exporting to the intermediate data format, comma separated variables (`.csv`). 
+These Excel documents are .xlsm (macro enabled Excel documents) that have macro based export 
+buttons to enable easily exporting to the intermediate data format, comma separated variables (.csv). 
 These documents all reside in a single location to help make data easy to find and modify.
 
-* [](Example.xlsm) (**right-click** and **Save as...** to download)
+* [](Example.xlsm) (**Right-click** and **Save as...** to download)
 
 
 ## Data Tables
@@ -30,9 +31,9 @@ DataTables, as the name implies, is a table of miscellaneous but related data gr
 and useful way, where the data fields can be any valid UObject property, including asset references. Before a designer can import a 
 CSV file into a DataTable, a programmer has to create the row container telling the engine how to interpret the data. 
 These tables consist of column names that have a 1:1 mapping to a given code based UStruct and its variables, 
-which must inherit from FTableRowBase to be recognized by the importer. 
+which must inherit from **FTableRowBase** to be recognized by the importer. 
 
-The first column is expected to be named `Name` and contains the names by which each row is accessed for use in game. Subsequent columns 
+The first column is expected to be named Name and contains the names by which each row is accessed for use in game. Subsequent columns 
 have a variable name for a heading and below, in the same column, the data for that row/column intersection. Given this 
 format, a single row maps 1:1 directly to the structure inherited from FTableRowBase.
 
@@ -84,8 +85,8 @@ The double quotes around the asset type are important for the property importing
 
 ## Data Curves
 
-Data curves work in a somewhat similar manner to DataTables, however they only support floating point 
-types. As with DataTables, the first column is named `Name` and contains the names by which each row 
+Data Curves work in a somewhat similar manner to DataTables, however they only support floating point 
+types. As with DataTables, the first column is named Name and contains the names by which each row 
 is accessed for use in game. Each column heading after the first stores the X-axis variable for a 
 curve to be plotted. The data under this heading is the Y-axis value for the given row. Given this 
 format, a single row corresponds to a curve that code can access and interpolate data along.
@@ -107,7 +108,7 @@ Here is an example table for damage progression:
 
 Importing a CSV file is a simple process:
 
-1.	Save your file out of Excel or another spreadsheet software with the `.csv` extension.
+1.	Save your file out of Excel or another spreadsheet software with the .csv extension.
 2.	Open Unreal Editor, and click on **Import** in the **Content Browser**.
 3.  Navigate and select the CSV file you want to import as a DataTable.
 4.  Finally, choose the row representation from the drop down list, in this case, LevelUpData.
@@ -118,7 +119,7 @@ Importing a CSV file is a simple process:
 
 	![](datatableobject.png)
 
-	You can view the contents of the DataTable in the editor by **double-clicking** on the object. You can update the object by **right-clicking** on the object and selecting **Reimport** from the menu. 
+	You can view the contents of the DataTable in the editor by **Double-clicking** on the object. You can update the object by **Right-clicking** on the object and selecting **Reimport** from the menu. 
 	Note that the original file path is used when reimporting the object. Here is the view of the data just after importing:
 	
 	![](datatableview.jpg)
@@ -126,7 +127,7 @@ Importing a CSV file is a simple process:
 ## Data Hookup
 
 Hooking up data from these tables is quite easy. From a programmer's perspective, all you must do is 
-place a Blueprint-exposed variable of either `FDataTableRowHandle` or `FCurveTableRowHandle`, 
+place a Blueprint-exposed variable of either **FDataTableRowHandle** or **FCurveTableRowHandle**, 
 depending on if you want a DataTable or CurveTable, respectively. From a content provider's point of 
 view, this will expose a data field with two subfields:
 
@@ -137,8 +138,8 @@ view, this will expose a data field with two subfields:
 ## Data Usage (Programmers Only)
 
 Once you have the data hooked up, using the data is quite simple. The handle structures provide helper 
-functions (`FindRow()` and `GetCurve()`) that allow you to retrieve either a struct or curve filled with data. 
-In the case of a `FCurveTableRowHandle`, an `FRichCurve` pointer is returned. The `FDataTableRowHandle`, however, 
+functions (**FindRow()** and **GetCurve()**) that allow you to retrieve either a struct or curve filled with data. 
+In the case of a FCurveTableRowHandle, an FRichCurve pointer is returned. The FDataTableRowHandle, however, 
 allows you to specify a structure in the template function call. This structure can be the final structure 
 or any of its parents in the inheritance hierarchy.
 
