#UE4doc #docs:
Added engine version field
Replaced and recolored multiple images
Changed a few file path references

diff --git a/Engine/Documentation/Source/Gameplay/HowTo/UseOnHit/Blueprints/UseOnHit_BP.INT.udn b/Engine/Documentation/Source/Gameplay/HowTo/UseOnHit/Blueprints/UseOnHit_BP.INT.udn
--- a/Engine/Documentation/Source/Gameplay/HowTo/UseOnHit/Blueprints/UseOnHit_BP.INT.udn
+++ b/Engine/Documentation/Source/Gameplay/HowTo/UseOnHit/Blueprints/UseOnHit_BP.INT.udn
@@ -5,17 +5,18 @@ Description:Using the OnHit Events to fracture a destructible mesh as well as ap
 Related: Engine/Blueprints
 Related: Engine/Blueprints/UserGuide/Events
 Related: Engine/Blueprints/UserGuide/CastNodes
+version: 4.9
 
 
 On this page we will use the **OnComponentHit** and **Event Hit** [](Engine/Blueprints/UserGuide/Events) to apply damage to and fracture a Destructible Mesh as well as apply an impulse at the hit location to push the destructible, simulating the effects of being hit by a projectile and applying force at the hit location. 
 
 [REGION:note]
-_For this How To, we are using the **Blueprint First Person Template** project with **Starter Content** enabled._
+For this How To, we are using the **Blueprint First Person Template** project with **Starter Content** enabled.
 [/REGION]
 
-1.	In your project in the **Content Browser**, go to the **Content/FirstPersonBP/Meshes** folder. 
+1.	In your project in the **Content Browser**, go to the **Content/Geometry/Meshes** folder. 
 
-1.	**Right-click** on **FirstPersonCube_Rounded** Static Mesh and select **Create Destructible Mesh**. 
+1.	**Right-click** on **1M_Cube** Static Mesh and select **Create Destructible Mesh**. 
 
 	![](OnHit1.png)
 
