#UE4doc #docs:
Added engine version field

diff --git a/Engine/Documentation/Source/Gameplay/HowTo/UseRaycasts/GHT_2_Overview.INT.udn b/Engine/Documentation/Source/Gameplay/HowTo/UseRaycasts/GHT_2_Overview.INT.udn
--- a/Engine/Documentation/Source/Gameplay/HowTo/UseRaycasts/GHT_2_Overview.INT.udn
+++ b/Engine/Documentation/Source/Gameplay/HowTo/UseRaycasts/GHT_2_Overview.INT.udn
@@ -2,6 +2,7 @@ Availability: Public
 Title:Using Raycasts (Tracing)
 Crumbs: %ROOT%, Gameplay
 Description:A How To Guide for using Raycasts (Tracing) for gameplay purposes. 
+version: 4.9
 
 There may be instances in your game where you want to determine if the player character is looking at something, and if so, alter the game state in some way (for example, highlight something when a player looks at it). Or maybe you want to determine if an enemy can see the player character, and if so, start shooting or engaging them in some way. You can accomplish both scenarios by using Raycasts (or Tracing) to "shoot" out an invisible ray which will detect geometry between two points and if geometry is hit, return what was hit so that you may then do something with it. 
 
