#UE4doc #docs:
Added engine version field
Replaced many images
Replaced multiple mentions of terms that have changed names

diff --git a/Engine/Documentation/Source/Gameplay/HowTo/UseRaycasts/Blueprints/GHT_2B.INT.udn b/Engine/Documentation/Source/Gameplay/HowTo/UseRaycasts/Blueprints/GHT_2B.INT.udn
--- a/Engine/Documentation/Source/Gameplay/HowTo/UseRaycasts/Blueprints/GHT_2B.INT.udn
+++ b/Engine/Documentation/Source/Gameplay/HowTo/UseRaycasts/Blueprints/GHT_2B.INT.udn
@@ -2,16 +2,17 @@ Availability: Public
 Title:Using Raycasts (Tracing) in Blueprints
 Crumbs: %ROOT%, Gameplay, Gameplay/HowTo/UseRaycasts
 Description:A How To Guide for using Raycasts (Tracing) in Blueprints. 
+version: 4.9
 
 [TOC (start:2 end:2)]
 
-On this page, you will learn the difference between **Single** and **Multi** Raycasts as well as how to set up each which could be used to retrieve collision information. 
+On this page, you will learn the difference between **Single** and **Multi** **Raycasts** as well as how to set up each which could be used to retrieve collision information. 
 
 ## Single Line Trace
 
-### Single Line Trace by Channel
+### LineTraceByChannel
 
-The **Single Line Trace by Channel** will perform a collision trace along a given line and return the first Object (that responds to the Trace Channel) the trace hits. 
+The **LineTraceByChannel** will perform a collision trace along a given line and return the first Object (that responds to the Trace Channel) the trace hits. 
 
 Below you will find steps for setting up a Single Line Trace by Channel. 
 
@@ -25,11 +26,11 @@ Below you will find steps for setting up a Single Line Trace by Channel.
 
 	This will cause the trace to run every frame.
 	
-1.  Drag off the execute pin, then search for **Single Line** and add the **Single Line Trace by Channel** node. 
+1.  Drag off the execute pin, then search for the **LineTraceByChannel** node. 
 
 	![](GHT2B_2.png)
 
-1.  While holding down the **Control** key, drag in the **FirstPersonCamera** component. 
+1.  While holding down the **Ctrl** key, drag in the **FirstPersonCamera** component. 
 
 	![](GHT2B_3.png)
 
@@ -89,17 +90,17 @@ Below you will find steps for setting up a Single Line Trace by Channel.
 
 The example above returns all Objects that are set to respond to the provided Trace Channel, however there may be instances when you want to only return certain objects. In the example above, you can use the Actors to Ignore pin which takes in an Array of Actors which should be ignored by the trace, but this would mean you have to specify each individual Actor to ignore. 
 
-You can also perform a Single Line Trace by Object where only the specified ObjectTypes are returned. This will allow you to target a specific set of Objects only to be included in the trace. 
+You can also perform a LineTraceByObject where only the specified ObjectTypes are returned. This will allow you to target a specific set of Objects only to be included in the trace. 
 
-### Single Line Trace by Object
+### LineTraceForObjects
 
-The **Single Line Trace by Object** will perform a collision trace along a given line and return the first Object (that responds to the ObjectType) the trace hits. 
+The **LineTraceForObjects** will perform a collision trace along a given line and return the first Object (that responds to the ObjectType) the trace hits. 
 
-To set up a Single Line Trace by Object trace, do the following: 
+To set up a LineTraceForObjects trace, do the following: 
 
-1.  Follow the steps used for the **Single Line Trace by Channel** example to set up a trace. 
+1.  Follow the steps used for the **LineTraceByChannel** example to set up a trace. 
 
-1.  Replace the **Single Line Trace by Channel** node with the **Single Line Trace by Object** node. 
+1.  Replace the **LineTraceByChannel** node with the **LineTraceForObjects** node. 
 
 1.  Drag off the **Object Types** pin and add the **Make Array** node. 
 
@@ -111,7 +112,7 @@ To set up a Single Line Trace by Object trace, do the following:
 
 	Here we are tracing for **WorldDyanmic** Objects. You can add more by clicking the **Add Pin** button.  
 
-1.  You can set the rest of the trace up the same way a **Single Line Trace by Channel** is setup. 
+1.  You can set the rest of the trace up the same way a **LineTraceByChannel** is setup. 
 
 	[REGION:lightbox]
 	[![](GHT2B_15.png)(w:800)](GHT2B_15.png)
@@ -128,15 +129,15 @@ To set up a Single Line Trace by Object trace, do the following:
 
 ## Multi Line Trace
 
-### Multi Line Trace by Channel
+### MultiLineTraceByChannel
 
-The **Multi Line Trace by Channel** will perform a collision trace along a given line and return all hits encountered up to and including the first blocking hit, returning only objects that respond to the specified Trace Channel. 
+The **MultiLineTraceByChannel** will perform a collision trace along a given line and return all hits encountered up to and including the first blocking hit, returning only objects that respond to the specified Trace Channel. 
 
-Below you will find steps for setting up a **Multi Line Trace by Channel**. 
+Below you will find steps for setting up a MultiLineTraceByChannel. 
 
-1.  Follow the steps used for the **Single Line Trace by Channel** example to set up a trace. 
+1.  Follow the steps used for the **LineTraceByChannel** example to set up a trace. 
 
-1.  Replace the **Single Line Trace by Channel** node with the **Multi Line Trace by Channel** node. 
+1.  Replace the **LineTraceByChannel** node with the **MultiLineTraceByChannel** node. 
 
 1.  Drag off the **Out Hits** pin and add a **ForEachLoop** node. 
 
@@ -158,19 +159,19 @@ Below you will find steps for setting up a **Multi Line Trace by Channel**.
 
 	![](GHT2B_17.png)
 
-	The Glass Window is a **Destructible Mesh** and its **Trace Response** to **Visibility** in its collision settings is set to **Overlap**, while the Physics Actor (Cube) has its **Visibility** setting to **Block**. 
+	The Glass Window is a **Destructible Mesh** and we have set its **Trace Response** to **Visibility** in its collision settings is set to **Overlap**, while the Physics Actor (Cube) has its **Visibility** setting to **Block**. 
 
 	You could use this for situations where you may want to shoot through an object (destroying it) and also hitting a player or enemy. 
 
-### Multi Line Trace by Object
+### MultiLineTraceForObjects
 
-The **Multi Line Trace by Object** will perform a collision trace along a given line and return all hits encountered, returning only objects that respond to the specified Object Type. 
+The **MultiLineTraceForObjects** will perform a collision trace along a given line and return all hits encountered, returning only objects that respond to the specified Object Type. 
 
-Below you will find steps for setting up a **Multi Line Trace by Object**. 
+Below you will find steps for setting up a MultiLineTraceForObjects. 
 
-1.  Follow the steps used for the **Single Line Trace by Channel** example to set up a trace. 
+1.  Follow the steps used for the **LineTraceByChannel** example to set up a trace. 
 
-1.  Replace the **Single Line Trace by Channel** node with the **Multi Line Trace by Object** node. 
+1.  Replace the **LineTraceByChannel** node with the **MultiLineTraceForObjects** node. 
 
 1.  Drag off the **Object Types** pin and add a **Make Array** node, then use the drop-down menus to add Objects to the Array. 
 
