#UE4doc #docs:
Added engine version field

diff --git a/Engine/Documentation/Source/Gameplay/HowTo/UseTimers/UseTimersOverview.INT.udn b/Engine/Documentation/Source/Gameplay/HowTo/UseTimers/UseTimersOverview.INT.udn
--- a/Engine/Documentation/Source/Gameplay/HowTo/UseTimers/UseTimersOverview.INT.udn
+++ b/Engine/Documentation/Source/Gameplay/HowTo/UseTimers/UseTimersOverview.INT.udn
@@ -2,6 +2,7 @@ Availability: Public
 Title:Using Timers
 Crumbs: %ROOT%, Gameplay
 Description:A How To Guide for Using Timers for Gameplay in Unreal Engine 4. 
+version: 4.9
 
 **Timers** in Unreal Engine 4 are a useful tool that can be used to perform an action once (or repeatedly) at specified intervals. Timers can be used to set actions to happen on a delay or over a period of time. For example, you may want to make the player invincible after obtaining a certain item and then revert after 10 seconds. Or, you may want to apply damage every second to a player that is standing in fire. Each of these scenarios can be accomplished with the use of Timers. 
 
