#UE4doc #docs - added engine version field.

diff --git a/Engine/Documentation/Source/Resources/ContentExamples/ContentExamples.INT.udn b/Engine/Documentation/Source/Resources/ContentExamples/ContentExamples.INT.udn
--- a/Engine/Documentation/Source/Resources/ContentExamples/ContentExamples.INT.udn
+++ b/Engine/Documentation/Source/Resources/ContentExamples/ContentExamples.INT.udn
@@ -2,6 +2,7 @@ Availability: Public
 Title:Content Examples
 Crumbs:%ROOT%, Resources
 Description:Example assets that demonstrate a specific concept or technique with corresponding documentation.
+Version: 4.9
 
 [VAR:Topic]
 [OBJECT:Topic]
