#UE4doc #doc - Added engine version field. Commented out redundant paragraph and condensed into introduction. Removed ToC since there was only a single object on the page.

diff --git a/Engine/Documentation/Source/Resources/ContentExamples/Audio/AudioContentExamples.INT.udn b/Engine/Documentation/Source/Resources/ContentExamples/Audio/AudioContentExamples.INT.udn
--- a/Engine/Documentation/Source/Resources/ContentExamples/Audio/AudioContentExamples.INT.udn
+++ b/Engine/Documentation/Source/Resources/ContentExamples/Audio/AudioContentExamples.INT.udn
@@ -42,20 +42,23 @@ Related: Engine/Audio/DistanceModelAttenuation
 	[/PARAM]
 [/OBJECT]
 [/VAR]
-[TOC(start:2 end:2)]
 
+<!---
+[TOC(start:2 end:2)]
+--->
 
 
 
-The **Audio Map** located inside the **Content Examples** project is an excellent source to see firsthand how all the audio systems come together, while demonstrating them in bite-sized practical examples. This page is intended to give you an idea of what to expect from the audio Content Examples as well as where to find specific examples based on the audio system(s) you want to work with. 
+The **Audio** map located inside the Content Examples project is an excellent source to see firsthand how all the audio systems come together, while demonstrating them in bite-sized practical examples. This page is intended to give you an idea of what to expect from the audio Content Examples as well as where to find specific examples based on the audio system(s) you want to work with. Whether you are new to audio in Unreal Engine 4 or need a refresher course, this map will provide an excellent foundation.
 
 
 ## Audio Map
 
 ![](contentExamplesStarterMap.png)(w:640)
 
+<!---
 The **Audio** map gives you the core concepts for using audio in Unreal Engine 4. This map covers many topics including the different types of audio assets and how to manipulate them to produce different effects based on your needs. Whether you are new to audio in Unreal Engine 4 or need a refresher course, this map will provide an excellent foundation. 
-
+--->
 | Example | What is Demonstrated |
 | ------- | ------------ |
 | **1.1 AmbientSound Actor** | Audio that plays as a one-off event. |
@@ -64,10 +67,10 @@ The **Audio** map gives you the core concepts for using audio in Unreal Engine 4
 | **1.4 AmbientSound Attenuation** | Audio that plays from the sound's location and gets louder or softer based on distance to the sound. |
 | **1.5 Attenuation with Low Pass Filter (LPF)** | Audio that plays from the sound's location and becomes more muffled the further the player is from the sound.  |
 | **1.6 Attenuation Shapes** | Audio that plays from the sounds location within defined shapes (e.g. Sphere, Capsule, Box or Cone shapes). |
-| **1.7 Shared Attenuation Settings** | Showcases a **Sound Attenuation** asset that affects the **Attenuation** of multiple **Sound Actors**. |
-| **1.8 AmbientSound Actor using SoundCues** | Showcases how a **Sound Cue** can be used instead of a **Sound Wave** to create more complex sounds or to mix multiple sounds into one. |
-| **1.9 Audio Volume** | Audio is affected by an **Audio Volume** where its effects can only be heard when the player is inside the volume (e.g. Dialogue plays and sounds as if you are in an auditorium when inside the volume. When outside, dialogue plays normally). |
-| **1.10 Audio Volume using Ambient Zone** | Two sounds are affected by an **Ambient Zone** where one sound is placed inside an Audio Volume and the other is not. The inside sound is quiet while outside the Audio Volume but louder when inside, while the outside sound is quiet when inside the Audio Volume and louder while standing outside it (e.g. audio when moving from an indoor to outdoor environment at night). |
+| **1.7 Shared Attenuation Settings** | Showcases a Sound Attenuation asset that affects the Attenuation of multiple Sound Actors. |
+| **1.8 AmbientSound Actor using SoundCues** | Showcases how a Sound Cue can be used instead of a Sound Wave to create more complex sounds or to mix multiple sounds into one. |
+| **1.9 Audio Volume** | Audio is affected by an Audio Volume where its effects can only be heard when the player is inside the volume (e.g. Dialogue plays and sounds as if you are in an auditorium when inside the volume. When outside, dialogue plays normally). |
+| **1.10 Audio Volume using Ambient Zone** | Two sounds are affected by an Ambient Zone where one sound is placed inside an Audio Volume and the other is not. The inside sound is quiet while outside the Audio Volume but louder when inside, while the outside sound is quiet when inside the Audio Volume and louder while standing outside it (e.g. audio when moving from an indoor to outdoor environment at night). |
 
 <!--
 ### Sound Cue Node Usage
