#UE4doc #docs - Added engine version field. Added section to call out the Blueprint Splines map and link to that CE page. Removed/added bolded terms where needed.

diff --git a/Engine/Documentation/Source/Resources/ContentExamples/Blueprints/Blueprints.INT.udn b/Engine/Documentation/Source/Resources/ContentExamples/Blueprints/Blueprints.INT.udn
--- a/Engine/Documentation/Source/Resources/ContentExamples/Blueprints/Blueprints.INT.udn
+++ b/Engine/Documentation/Source/Resources/ContentExamples/Blueprints/Blueprints.INT.udn
@@ -9,6 +9,8 @@ Related: Engine/Blueprints/UserGuide/Types/MacroLibrary
 Related: Engine/Animation/AnimBlueprints/EventGraph
 Related: Engine/Blueprints/UserGuide/UserConstructionScript
 Related: Gameplay/HowTo/SetUpInput/Blueprints
+Version: 4.9
+
 
 
 [REGION:banner]
@@ -51,7 +53,7 @@ Related: Gameplay/HowTo/SetUpInput/Blueprints
 
 
 
-There are several **Blueprint Maps** located inside the **Content Examples** project that demonstrate the power and flexibility of the **Blueprints** system. Whether you are looking to define player controls, implement a health/pick-up system, or simply have a door that opens when a player approaches it, the examples illustrated in these maps should give you some ideas on how to approach whatever it is you are after. 
+There are several **Blueprint** maps located inside the Content Examples project that demonstrate the power and flexibility of the **Blueprints** system. Whether you are looking to define player controls, implement a health/pick-up system, or simply have a door that opens when a player approaches it, the examples illustrated in these maps should give you some ideas on how to approach whatever it is you are after. 
 
 This page indicates what examples are included in each map and where to find them so that you may jump to your area of interest and begin scripting right away. If you are new to Blueprints however, it is suggested that you begin with the **Blueprints_Overview** map as it will cover the core aspects of the system and will help get you started with using them. 
 
@@ -70,11 +72,11 @@ Take your time going over these sections as you will use these elements quite of
 | ------- | ------------ |
 | **[1.1 Hello World Example](Resources\ContentExamples\Blueprints\1_1)** | A basic Blueprint is used to display floating text in the level (e.g. floating text "Hello World" is displayed when playing in the level). |
 | **[1.2 Components Example](Resources\ContentExamples\Blueprints\1_2)** | Examples of Components and how multiple Components can be used in a single Blueprint (e.g. a sprite, Static Mesh, and light all used in one Blueprint). |
-| **[1.3 Variables Example](Resources\ContentExamples\Blueprints\1_3)** | Shows how Variables can be set to public or private and how to edit them (e.g. a light color that can be changed via the **Details** panel or only changed in the Blueprint). |
-| **[1.4 Construction Script Example](Resources\ContentExamples\Blueprints\1_4)** | Shows how Components of a Blueprint can be disabled/enabled with the **Construction Script** (e.g. use checkboxes to turn on/off Components in a Blueprint). |
-| **[1.5 Event Graph Example](Resources\ContentExamples\Blueprints\1_5)** | The **Construction Script** and **Event Graph** are used to allow custom colors for a mesh without editing the material (e.g. user defines the color of a rocket via the **Details** panel).  |
-| **[1.6 Simple Math Example](Resources\ContentExamples\Blueprints\1_6)** | An **Event Tick** is used in this Blueprint to perform simple addition over time (e.g. a timer is displayed counting upwards every second from 0).  |
-| **[1.7 Flow Example](Resources\ContentExamples\Blueprints\1_7)** | A Blueprint that uses different types of Flow Control to modify movement, color and effects of an object (e.g. the direction a coin moves, its color, and if effects are applied via the **Details** panel). |
+| **[1.3 Variables Example](Resources\ContentExamples\Blueprints\1_3)** | Shows how Variables can be set to public or private and how to edit them (e.g. a light color that can be changed via the Details panel or only changed in the Blueprint). |
+| **[1.4 Construction Script Example](Resources\ContentExamples\Blueprints\1_4)** | Shows how Components of a Blueprint can be disabled/enabled with the Construction Script (e.g. use checkboxes to turn on/off Components in a Blueprint). |
+| **[1.5 Event Graph Example](Resources\ContentExamples\Blueprints\1_5)** | The Construction Script and Event Graph are used to allow custom colors for a mesh without editing the material (e.g. user defines the color of a rocket via the Details panel).  |
+| **[1.6 Simple Math Example](Resources\ContentExamples\Blueprints\1_6)** | An Event Tick is used in this Blueprint to perform simple addition over time (e.g. a timer is displayed counting upwards every second from 0).  |
+| **[1.7 Flow Example](Resources\ContentExamples\Blueprints\1_7)** | A Blueprint that uses different types of Flow Control to modify movement, color and effects of an object (e.g. the direction a coin moves, its color, and if effects are applied via the Details panel). |
 
 
 ## Blueprints_Advanced Map
@@ -91,7 +93,7 @@ The **Construction Script** section illustrates how to: randomly spawn meshes in
 | **[1.1 Randomly Scattered Meshes Example](Resources\ContentExamples\Blueprints_Advanced\1_1)** | How to randomly spawn Static Meshes within a specified radius (e.g. randomly spawn ground clutter or rocks in a radius). |
 | **[1.2 Automatic Wall Example](Resources\ContentExamples\Blueprints_Advanced\1_2)** | How to extend a row of meshes from an Actor's pivot point (e.g. a fence that can be scaled in length). |
 | **[1.3 Automatic Ring Example](Resources\ContentExamples\Blueprints_Advanced\1_3)** | How to space out a specified number of meshes in a circle, with rotation determined by each mesh's position (e.g. an evenly spaced and orientated ring of meshes). |
-| **[1.4 Customizable Spotlight Example](Resources\ContentExamples\Blueprints_Advanced\1_4)** | How to use the **Construction Script** to combine all the components of a spotlight into one Actor (e.g. a spotlight that can be customized including on/off state, light color, orientation and radius).  |
+| **[1.4 Customizable Spotlight Example](Resources\ContentExamples\Blueprints_Advanced\1_4)** | How to use the Construction Script to combine all the components of a spotlight into one Actor (e.g. a spotlight that can be customized including on/off state, light color, orientation and radius).  |
 
 
 ### Event Graph Examples
@@ -99,11 +101,11 @@ The **Event Graph** examples will show you how to: animate a Static Mesh with a
 
 | Example | What is Demonstrated |
 | ------- | ------------ |
-| **[2.1 Timeline Animation Example](Resources\ContentExamples\Blueprints_Advanced\2_1)** | How to use a **Timeline** to animate a Static Mesh wth an **Event Track** to spawn an emitter at a specified point (e.g. creating an animated bouncing ball). |
+| **[2.1 Timeline Animation Example](Resources\ContentExamples\Blueprints_Advanced\2_1)** | How to use a Timeline to animate a Static Mesh wth an Event Track to spawn an emitter at a specified point (e.g. creating an animated bouncing ball). |
 | **[2.2 Actor Tracking Example](Resources\ContentExamples\Blueprints_Advanced\2_2)** | How to make an Actor turn to face the player when coming within a specified range (e.g. a spotlight that reacts to and follows the player).  |
-| **[2.3 Button Trigger (Blueprint Interface) Example](Resources\ContentExamples\Blueprints_Advanced\2_3)** | How to pass a variable from one Blueprint to another using **Blueprint Interfaces** (e.g. an on/off variable is passed from a button to a set of gears that animate). |
-| **[2.4 Opening Door (Trigger) Example](Resources\ContentExamples\Blueprints_Advanced\2_4)** | How to open/close a door automatically with a **Trigger** and a Timeline. |
-| **[2.5 Child Blueprints Example](Resources\ContentExamples\Blueprints_Advanced\2_5)** | How to implement different types of pick-ups derived from the same **Parent Class** (e.g. a pick-up for health, a coin, and a cupcake each with different functionality when picked up). |
+| **[2.3 Button Trigger (Blueprint Interface) Example](Resources\ContentExamples\Blueprints_Advanced\2_3)** | How to pass a variable from one Blueprint to another using Blueprint Interfaces (e.g. an on/off variable is passed from a button to a set of gears that animate). |
+| **[2.4 Opening Door (Trigger) Example](Resources\ContentExamples\Blueprints_Advanced\2_4)** | How to open/close a door automatically with a Trigger and a Timeline. |
+| **[2.5 Child Blueprints Example](Resources\ContentExamples\Blueprints_Advanced\2_5)** | How to implement different types of pick-ups derived from the same Parent Class (e.g. a pick-up for health, a coin, and a cupcake each with different functionality when picked up). |
 
 
 ## Blueprint_Communication Map
@@ -120,7 +122,7 @@ The **Direct Blueprint Communication** or **Blueprint Communication** section sh
 | ------- | ------------ |
 | **1.1 Basic Communication (Target Blueprint)** | An example of a button Blueprint that toggles a light Blueprint on and off through a custom event. |
 | **1.2 Basic Communication (Actor Casting)** | How to use an overlap event to detect a Blueprint of a particular class (e.g. a light bulb is turned on/off by connecting a battery).|
-| **1.3 Communication (Actor Casting To Child Blueprints)** | How use **Child Blueprints** which share functionality of their parent but have independent functionality as well (e.g. turning lights on/off by connecting batteries each with different energy values).  |
+| **1.3 Communication (Actor Casting To Child Blueprints)** | How use Child Blueprints which share functionality of their parent but have independent functionality as well (e.g. turning lights on/off by connecting batteries each with different energy values).  |
 | **1.4 Communication with All Actors of a Specific Class** | How to communicate with several Actors of a particular class (e.g. turn on several lights at once each with different functionality).  |
 
 
@@ -130,7 +132,7 @@ In addition, by **Binding** one or more events to an **Event Dispatcher**, you c
 
 | Example | What is Demonstrated |
 | ------- | ------------ |
-| **2.1 Using Event Dispatchers** | How to use an **Event Dispatcher** to call an event inside the **Level Blueprint** (e.g. toggling a light on/off using an Event Dispatcher). |
+| **2.1 Using Event Dispatchers** | How to use an Event Dispatcher to call an event inside the Level Blueprint (e.g. toggling a light on/off using an Event Dispatcher). |
 | **2.2 Binding Event Dispatchers to A Custom Event** | Illustrates binding an Event Dispatcher to a custom event inside of the Level Blueprint (e.g. toggling a button that rotates an object over time). |
 | **2.3 Binding Event Dispatchers to A Custom Event on Spawn** | Illustrates how to enable/disable a trigger with an Event Dispatcher that is called when a spawned Actor is destroyed (e.g. press a button to spawn a bomb, no other bombs can be spawned until the previous bomb explodes). |
 
@@ -168,8 +170,11 @@ The example given allows you to click-and-drag objects around on screen that rea
 
 More detailed information for this map can be found here: 
 
-* [1.1 Blueprint Mouse Interaction](Resources\ContentExamples\MouseInterface)
+<!---* [1.1 Blueprint Mouse Interaction](Resources\ContentExamples\MouseInterface)--->
 
+| Example | What is Demonstrated |
+| ------- | ------------ |
+| **[1.1 Blueprint Mouse Interaction](Resources\ContentExamples\MouseInterface)** | Demonstrates how to utilize mouse or touch inputs in your games. |
 
 ## Blueprint_HUD Map
 
@@ -179,7 +184,45 @@ The **Blueprint_HUD** map provides an example on how to generate HUD elements wi
 
 More detailed information for this map can be found here: 
 
-* [1.1 Blueprints HUD Example](Resources\ContentExamples\Blueprints_HUD\1_1)
+<!---* [1.1 Blueprints HUD Example](Resources\ContentExamples\Blueprints_HUD\1_1)--->
+
+| Example | What is Demonstrated |
+| ------- | ------------ |
+| **[1.1 Blueprints HUD Example](Resources\ContentExamples\Blueprints_HUD\1_1)** | Demonstrates a simple in-game HUD that also includes clickable buttons in the forum of a basic pause menu. |
+
+## Blueprint_Splines Map
+
+![](blueprintSplines_ce_header.png)(convert:false)
+
+The **Blueprint Splines** map in Content Examples will show how you can use **Spline Components** with **Blueprints** to create paths that can be used in a number of different ways. By simply adding a Spline Component to your Blueprint and selecting an instance of that Blueprint, you will gain access to editing splines within your level. You will then be able to select spline keys inside the viewport and manipulate them like any other Actor using the Translation, Rotation, and Scale widgets. 
+
+The Content Examples provided will illustrate how **Blueprint Splines** can be used with a Construction Script to create paths of Static Mesh objects or how to dynamically alter splines paths or Spline Meshes during gameplay to create content or animations. 
+
+ 
+### Construction Script Examples
+
+The examples below showcase how a Spline Component can be used with a Blueprint's Construction Script to paths or scalable Static Meshes inside the viewport in real-time. 
+
+| Example | What is Demonstrated |
+| ------- | ------------ |
+| **1.1 Blueprint Spline Component** | Demonstrates how a Spline Component can be edited easily in the viewport (e.g. move four spline points of a Spline Component inside the viewport). |
+| **1.2 Placing Components Along Splines** | Shows how a Blueprint can retrieve various kinds of information from a Spline Component (e.g. shows how Static Mesh Components can be placed along the length of a spline where their orientation is determined by the direction of the spline in which they were placed). |
+| **1.3 Spline Mesh Components** | Spline Meshes which are another type of Spline Component used to stretch and bend a Static Mesh asset (e.g. a pipe that can be stretched and orientated using widgets in the viewport). |
+
+
+### Dynamic Examples
+
+The **Dynamic Examples** show how a Spline Component can be updated dynamically during gameplay to provide ambient background movement of objects or to animate objects or even a playable character! 
+
+| Example | What is Demonstrated |
+| ------- | ------------ |
+| **2.1 Using a Spline as a Path** | An example of a Spline Component used as a path for a Particle Effect to follow. |
+| **2.2 Changing Splines While Simulating** | Shows how Spline Components can be created dynamically during gameplay (e.g. a string of lanterns animate along a spline when playing or simulating). |
+| **2.3 Animating a Spline Mesh** | A Spline Component can also be used as a path for a Spline Mesh to move along (e.g. a vine grows out of the ground along a spline path). |
+| **2.4 Spline Mesh for Pawn Animation** | How a Spline Mesh can be used for animating a pawn (e.g. a Spline Mesh is used to animate a playable swimming fish). |
+
+
+
 
 <!--
 [EXCERPT:BlueprintsCEInstructions]
