#UE4doc #docs - Added engine version field. Increased image size in tip bubble.

diff --git a/Engine/Documentation/Source/Resources/ContentExamples/Blueprints/1_1/1_1.INT.udn b/Engine/Documentation/Source/Resources/ContentExamples/Blueprints/1_1/1_1.INT.udn
--- a/Engine/Documentation/Source/Resources/ContentExamples/Blueprints/1_1/1_1.INT.udn
+++ b/Engine/Documentation/Source/Resources/ContentExamples/Blueprints/1_1/1_1.INT.udn
@@ -9,6 +9,7 @@ Related: Engine/Blueprints/UserGuide/Types/MacroLibrary
 Related: Engine/Animation/AnimBlueprints/EventGraph
 Related: Engine/Blueprints/UserGuide/UserConstructionScript
 Related: Gameplay/HowTo/SetUpInput/Blueprints
+Version: 4.9
 
 [VAR:Steps]
 [OBJECT:Navigation]
@@ -52,13 +53,11 @@ Related: Gameplay/HowTo/SetUpInput/Blueprints
 [REGION:tip]
 This example is best seen during Play In Editor or Simulation. To Simulate, click the **Simulate** button located at the top of the screen in the editor Toolbar, or press **Alt-S**. To Play in the Editor, click the **Play** button, also located in the toolbar, or press **Alt-P**.
 
-![Simulate](button_Simulate.png)(h:60) ![Play In Editor](button_PlayIn.png)(h:60)
+![Simulate](button_Simulate.png)(h:150) ![Play In Editor](button_PlayIn.png)(h:150)
 [/REGION]
 
 
 
-
-
 This example simply shows how to make an update to a Component's property using the Event Graph by using the classic programmer introductory project of "Hello World." Simply put, the Blueprint script for the object sets the Text property of a Text Render Component to "Hello World" as soon as gameplay begins. Prior to gameplay, such as when using the editor, the Text property will be blank so no text will appear. This is one of the simplest and most basic approaches to demonstrating the Event Graph and Blueprint setup.
 
 As with all examples in this level, the object displayed is a Blueprint Actor, meaning it is a Blueprint created with the Actor class as a parent. 
