#UE4doc #docs - Added engine version field.

diff --git a/Engine/Documentation/Source/Resources/ContentExamples/Blueprints/1_2/1_2.INT.udn b/Engine/Documentation/Source/Resources/ContentExamples/Blueprints/1_2/1_2.INT.udn
--- a/Engine/Documentation/Source/Resources/ContentExamples/Blueprints/1_2/1_2.INT.udn
+++ b/Engine/Documentation/Source/Resources/ContentExamples/Blueprints/1_2/1_2.INT.udn
@@ -10,6 +10,7 @@ Related: Engine/Blueprints/UserGuide/Types/MacroLibrary
 Related: Engine/Animation/AnimBlueprints/EventGraph
 Related: Engine/Blueprints/UserGuide/UserConstructionScript
 Related: Gameplay/HowTo/SetUpInput/Blueprints
+Version: 4.9
 
 [VAR:Steps]
 [OBJECT:Navigation]
