#UE4doc #docs - Added engine version field. Bolded term.

diff --git a/Engine/Documentation/Source/Resources/ContentExamples/Blueprints/1_5/1_5.INT.udn b/Engine/Documentation/Source/Resources/ContentExamples/Blueprints/1_5/1_5.INT.udn
--- a/Engine/Documentation/Source/Resources/ContentExamples/Blueprints/1_5/1_5.INT.udn
+++ b/Engine/Documentation/Source/Resources/ContentExamples/Blueprints/1_5/1_5.INT.udn
@@ -10,6 +10,7 @@ Related: Engine/Blueprints/UserGuide/Types/MacroLibrary
 Related: Engine/Animation/AnimBlueprints/EventGraph
 Related: Engine/Blueprints/UserGuide/UserConstructionScript
 Related: Gameplay/HowTo/SetUpInput/Blueprints
+Version: 4.9
 
 [VAR:Steps]
 [OBJECT:Navigation]
@@ -63,7 +64,7 @@ In the Construction Script for the object, we create and apply our Dynamic Mater
 [/REGION]
 
 ## Event Graph
-In the Event Graph, we have a Timeline that blends from the color set in the Construction Script to the one set for the Event Graph variable, which is also public and visible in the **Details** panel. The Timeline is used just to give a smooth transition from one color to the other. For more information on Timelines, please see [the Timelines documentation](Engine/Blueprints\UserGuide\Timelines).
+In the Event Graph, we have a **Timeline** that blends from the color set in the Construction Script to the one set for the Event Graph variable, which is also public and visible in the **Details** panel. The Timeline is used just to give a smooth transition from one color to the other. For more information on Timelines, please see [the Timelines documentation](Engine/Blueprints\UserGuide\Timelines).
 
 [REGION:lightbox]
 [![](1.5EventGraph.png)](1.5EventGraph.png)
