#UE4doc #docs - Added engine version field. Removed Overview header.

diff --git a/Engine/Documentation/Source/Resources/ContentExamples/Blueprints/1_6/1_6.INT.udn b/Engine/Documentation/Source/Resources/ContentExamples/Blueprints/1_6/1_6.INT.udn
--- a/Engine/Documentation/Source/Resources/ContentExamples/Blueprints/1_6/1_6.INT.udn
+++ b/Engine/Documentation/Source/Resources/ContentExamples/Blueprints/1_6/1_6.INT.udn
@@ -10,6 +10,7 @@ Related: Engine/Blueprints/UserGuide/Types/MacroLibrary
 Related: Engine/Animation/AnimBlueprints/EventGraph
 Related: Engine/Blueprints/UserGuide/UserConstructionScript
 Related: Gameplay/HowTo/SetUpInput/Blueprints
+Version: 4.9
 
 [VAR:Steps]
 [OBJECT:Navigation]
@@ -48,12 +49,12 @@ Related: Gameplay/HowTo/SetUpInput/Blueprints
 
 [REGION:fullwidth]
 
-## Overview
+
 
 This example demonstrates how simple math equations can be performed on data within Blueprints. In this case, we add up the amount of time the game has run and apply the number of seconds to a Text Render Component. 
 
 ## Components
-The Components for this example consist of two Text Render Components. One simply displays the number of seconds that have passed, the other provides a label for that text.
+The Components for this example consist of two **Text Render** Components. One simply displays the number of seconds that have passed, the other provides a label for that text.
 
 ![](1.6Components.png)(w:500)
 
