#UE4doc #docs - Added engine version field. Added/Removed bolding from terms.

diff --git a/Engine/Documentation/Source/Resources/ContentExamples/Blueprints/1_7/1_7.INT.udn b/Engine/Documentation/Source/Resources/ContentExamples/Blueprints/1_7/1_7.INT.udn
--- a/Engine/Documentation/Source/Resources/ContentExamples/Blueprints/1_7/1_7.INT.udn
+++ b/Engine/Documentation/Source/Resources/ContentExamples/Blueprints/1_7/1_7.INT.udn
@@ -10,6 +10,7 @@ Related: Engine/Blueprints/UserGuide/Types/MacroLibrary
 Related: Engine/Animation/AnimBlueprints/EventGraph
 Related: Engine/Blueprints/UserGuide/UserConstructionScript
 Related: Gameplay/HowTo/SetUpInput/Blueprints
+Version: 4.9
 
 [VAR:Steps]
 [OBJECT:Navigation]
@@ -71,10 +72,10 @@ The Construction Script for this example creates and sets a Dynamic Material Ins
 
 The Event Graph uses several different types of flow control. These include:
 
-* A Sequence, which executes all other operations in order. This is useful when you have a one-to-many relationship in which one event needs to execute many different things.
-* A Branch, which is controlled via the Branch Condition variable.
-* A Switch on Enum, which controls the direction the coin moves (horizontally or vertically) based on the result of an enumeration.
-* A Do N, spawns a particle effect a specific number of times.
+* A **Sequence**, which executes all other operations in order. This is useful when you have a one-to-many relationship in which one event needs to execute many different things.
+* A **Branch**, which is controlled via the Branch Condition variable.
+* A **Switch on Enum**, which controls the direction the coin moves (horizontally or vertically) based on the result of an enumeration.
+* A **Do N**, spawns a particle effect a specific number of times.
 
 ### Creating Enumerations
 
@@ -84,7 +85,7 @@ The Event Graph uses several different types of flow control. These include:
 * Crouched
 * Prone
 
-Enumerations are created in the **Content Browser** by **right-clicking** and choosing **Blueprints > Enumeration** from the Context Menu. They can then be edited within the **Enum Editor** by **double-clicking** them. Within, you can add the various entries that need to appear within your enumeration.
+Enumerations are created in the Content Browser by **right-clicking** and choosing **Blueprints > Enumeration** from the Context Menu. They can then be edited within the **Enum Editor** by **double-clicking** them. Within, you can add the various entries that need to appear within your enumeration.
 
 ![](EnumEditor.png)
 
