#UE4doc #docs - added engine version field.

diff --git a/Engine/Documentation/Source/Resources/ContentExamples/Blueprints_Advanced/1_3/1_3.INT.udn b/Engine/Documentation/Source/Resources/ContentExamples/Blueprints_Advanced/1_3/1_3.INT.udn
--- a/Engine/Documentation/Source/Resources/ContentExamples/Blueprints_Advanced/1_3/1_3.INT.udn
+++ b/Engine/Documentation/Source/Resources/ContentExamples/Blueprints_Advanced/1_3/1_3.INT.udn
@@ -10,6 +10,7 @@ Related: Engine/Blueprints/UserGuide/Types/MacroLibrary
 Related: Engine/Animation/AnimBlueprints/EventGraph
 Related: Engine/Blueprints/UserGuide/UserConstructionScript
 Related: Gameplay/HowTo/SetUpInput/Blueprints
+Version: 4.9
 
 [VAR:Steps]
 [OBJECT:Navigation]
@@ -63,7 +64,7 @@ The **BP_Radial_Array** Blueprint uses the Construction Script to generate an ev
 [![](AdvBP_1_3_ConstructionScript.png)](AdvBP_1_3_ConstructionScript.png)
 [/REGION]
 [REGION:caption]
-Click for full size or **right-click** and Save As.
+_Click for full size or **right-click** and Save As._
 [/REGION]
 
 
