#UE4doc #docs - Added engine version field.

diff --git a/Engine/Documentation/Source/Resources/ContentExamples/Blueprints_Advanced/2_1/2_1.INT.udn b/Engine/Documentation/Source/Resources/ContentExamples/Blueprints_Advanced/2_1/2_1.INT.udn
--- a/Engine/Documentation/Source/Resources/ContentExamples/Blueprints_Advanced/2_1/2_1.INT.udn
+++ b/Engine/Documentation/Source/Resources/ContentExamples/Blueprints_Advanced/2_1/2_1.INT.udn
@@ -10,6 +10,7 @@ Related: Engine/Blueprints/UserGuide/Types/MacroLibrary
 Related: Engine/Animation/AnimBlueprints/EventGraph
 Related: Engine/Blueprints/UserGuide/UserConstructionScript
 Related: Gameplay/HowTo/SetUpInput/Blueprints
+Version: 4.9
 
 [VAR:Steps]
 [OBJECT:Navigation]
@@ -51,7 +52,7 @@ Related: Gameplay/HowTo/SetUpInput/Blueprints
 
 
 
-The **BP_Timeline_Ball** Blueprint demonstrates a basic [**Timeline**](Engine/Blueprints/UserGuide/Timelines) in the Event Graph, which automatically plays and loops, triggering an event at a certain point along the Timeline that spawns a particle effect.
+The **BP_Timeline_Ball** Blueprint demonstrates a basic [Timeline](Engine/Blueprints/UserGuide/Timelines) in the Event Graph, which automatically plays and loops, triggering an event at a certain point along the Timeline that spawns a particle effect.
 
 
 
