#UE4doc #docs - Added engine version field.

diff --git a/Engine/Documentation/Source/Resources/ContentExamples/Blueprints_Advanced/2_2/2_2.INT.udn b/Engine/Documentation/Source/Resources/ContentExamples/Blueprints_Advanced/2_2/2_2.INT.udn
--- a/Engine/Documentation/Source/Resources/ContentExamples/Blueprints_Advanced/2_2/2_2.INT.udn
+++ b/Engine/Documentation/Source/Resources/ContentExamples/Blueprints_Advanced/2_2/2_2.INT.udn
@@ -9,6 +9,7 @@ Related: Engine/Blueprints/UserGuide/Types/MacroLibrary
 Related: Engine/Animation/AnimBlueprints/EventGraph
 Related: Engine/Blueprints/UserGuide/UserConstructionScript
 Related: Gameplay/HowTo/SetUpInput/Blueprints
+Version: 4.9
 
 [VAR:Steps]
 [OBJECT:Navigation]
