#UE4doc #docs - Added engine version field. Removed bolding from hyperlink.

diff --git a/Engine/Documentation/Source/Resources/ContentExamples/Blueprints_Advanced/2_3/2_3.INT.udn b/Engine/Documentation/Source/Resources/ContentExamples/Blueprints_Advanced/2_3/2_3.INT.udn
--- a/Engine/Documentation/Source/Resources/ContentExamples/Blueprints_Advanced/2_3/2_3.INT.udn
+++ b/Engine/Documentation/Source/Resources/ContentExamples/Blueprints_Advanced/2_3/2_3.INT.udn
@@ -10,6 +10,7 @@ Related: Engine/Blueprints/UserGuide/Types/MacroLibrary
 Related: Engine/Animation/AnimBlueprints/EventGraph
 Related: Engine/Blueprints/UserGuide/UserConstructionScript
 Related: Gameplay/HowTo/SetUpInput/Blueprints
+Version: 4.9
 
 [VAR:Steps]
 [OBJECT:Navigation]
@@ -53,7 +54,7 @@ Related: Gameplay/HowTo/SetUpInput/Blueprints
 
 
 
-[**Blueprint Interfaces**](Engine/Blueprints\UserGuide\Types\Interface) allow one Blueprint to trigger an event or function in another Blueprint. Both of these Blueprints must implement the Blueprint Interface in their **Blueprint Properties**.
+[Blueprint Interfaces](Engine/Blueprints\UserGuide\Types\Interface) allow one Blueprint to trigger an event or function in another Blueprint. Both of these Blueprints must implement the Blueprint Interface in their **Blueprint Properties**.
 In this example, the **BP_Button** Blueprint causes the gears in the **BP_Gears** Blueprint to start or stop spinning.
 
 
