#UE4doc #docs - Added engine version field. Removed bolding from hyperlink.

diff --git a/Engine/Documentation/Source/Resources/ContentExamples/Blueprints_Advanced/2_4/2_4.INT.udn b/Engine/Documentation/Source/Resources/ContentExamples/Blueprints_Advanced/2_4/2_4.INT.udn
--- a/Engine/Documentation/Source/Resources/ContentExamples/Blueprints_Advanced/2_4/2_4.INT.udn
+++ b/Engine/Documentation/Source/Resources/ContentExamples/Blueprints_Advanced/2_4/2_4.INT.udn
@@ -11,6 +11,7 @@ Related: Engine/Blueprints/UserGuide/Types/MacroLibrary
 Related: Engine/Animation/AnimBlueprints/EventGraph
 Related: Engine/Blueprints/UserGuide/UserConstructionScript
 Related: Gameplay/HowTo/SetUpInput/Blueprints
+Version: 4.9
 
 [VAR:Steps]
 [OBJECT:Navigation]
@@ -53,7 +54,7 @@ Related: Gameplay/HowTo/SetUpInput/Blueprints
 
 
 
-The **BP_Door** Blueprint creates a door that opens with a trigger Box Component, and is another example of a Blueprint that might commonly be used in the creation of a game, with script and animation built in that needs no additional setup. 
+The **BP_Door** Blueprint creates a door that opens with a trigger **Box Component**, and is another example of a Blueprint that might commonly be used in the creation of a game, with script and animation built in that needs no additional setup. 
 
 When a player touches the BoxComponent, the Blueprint retrieves the player's location relative to the door and uses this information in conjunction with the Actor's rotation to determine which side of the door the player is on. Then, a [Timeline](Engine/Blueprints\UserGuide\Timelines) plays that opens the door in the correct direction, so that it swings away from the player.
 
