#UE4doc #docs - Added engine version field.

diff --git a/Engine/Documentation/Source/Resources/ContentExamples/Blueprints_HUD/1_1/1_1.INT.udn b/Engine/Documentation/Source/Resources/ContentExamples/Blueprints_HUD/1_1/1_1.INT.udn
--- a/Engine/Documentation/Source/Resources/ContentExamples/Blueprints_HUD/1_1/1_1.INT.udn
+++ b/Engine/Documentation/Source/Resources/ContentExamples/Blueprints_HUD/1_1/1_1.INT.udn
@@ -9,6 +9,7 @@ Related: Engine/Blueprints/UserGuide/Types/MacroLibrary
 Related: Engine/Animation/AnimBlueprints/EventGraph
 Related: Engine/Blueprints/UserGuide/UserConstructionScript
 Related: Gameplay/HowTo/SetUpInput/Blueprints
+Version: 4.9
 
 [VAR:Steps]
 [OBJECT:Navigation]
@@ -47,7 +48,7 @@ Related: Gameplay/HowTo/SetUpInput/Blueprints
 [REGION:fullwidth]
 
 This example demonstrates a simple in-game HUD, and also includes clickable buttons in the form of a basic pause menu. A health bar material is drawn on the HUD and uses a Blueprint Interface function to get the health value from the Character, while an objective marker is drawn over an objective along with the objective's name and distance from the player.
-The HUD Blueprint can be found by searching for **BP_HUD_Example** in the **Content Browser**.
+The HUD Blueprint can be found by searching for **BP_HUD_Example** in the Content Browser.
 
 ## GameMode
 
