#UE4doc #docs - Corrected errors and added skill level tags to pages. References Review for UEDOC-2333

diff --git a/Engine/Documentation/Source/Platforms/Oculus/QuickStart/1/OculusQuickStart_1.INT.udn b/Engine/Documentation/Source/Platforms/Oculus/QuickStart/1/OculusQuickStart_1.INT.udn
--- a/Engine/Documentation/Source/Platforms/Oculus/QuickStart/1/OculusQuickStart_1.INT.udn
+++ b/Engine/Documentation/Source/Platforms/Oculus/QuickStart/1/OculusQuickStart_1.INT.udn
@@ -3,6 +3,7 @@ Crumbs:%ROOT%, Platforms, Platforms/Oculus, Platforms/Oculus/QuickStart
 Description:Register as an Oculus developer to gain access to necessary programs and documentation
 Availability:Public
 Version: 4.9
+SkillLevel: Beginner
 
 [VAR:Steps]
 [OBJECT:Navigation]
@@ -24,7 +25,7 @@ While it is not mandatory to register as an Oculus developer, registering will a
 
 %Steps%
 
-Visit the [Oculus Developer Center](https://developer.oculus.com/) and click on person icon in the upper right hand corner of the page to log in if you already have an Oculus account. 
+Visit the [Oculus Developer Center](https://developer.oculus.com/) and click on the person icon in the upper right hand corner of the page to log in if you already have an Oculus account. 
 Please note that if you already have an account you can skip ahead to the next step.
 	
 ![](Oculus_Downloads.png)
diff --git a/Engine/Documentation/Source/Platforms/Oculus/QuickStart/2/OculusQuickStart_2.INT.udn b/Engine/Documentation/Source/Platforms/Oculus/QuickStart/2/OculusQuickStart_2.INT.udn
--- a/Engine/Documentation/Source/Platforms/Oculus/QuickStart/2/OculusQuickStart_2.INT.udn
+++ b/Engine/Documentation/Source/Platforms/Oculus/QuickStart/2/OculusQuickStart_2.INT.udn
@@ -3,6 +3,7 @@ Crumbs:%ROOT%, Platforms, Platforms/Oculus, Platforms/Oculus/QuickStart
 Description:Download the needed software for running games on the Oculus Rift
 Availability:Public
 Version: 4.9
+SkillLevel: Beginner
 
 [VAR:Steps]
 [OBJECT:Navigation]
@@ -20,7 +21,7 @@ Version: 4.9
 [/OBJECT]
 [/VAR]
 
-In this section we will go over which version of the Oculus PC Runtime drivers you need to download and were you can download them form.
+In this section we will go over which version of the Oculus PC Runtime drivers you need to download and where you can download them from.
 
 %Steps%
 
@@ -30,11 +31,11 @@ The version of the Oculus Runtime that is currently compatible with the latest v
 [/REGION]
 
 
-1. From the web browser of your choice got the Oculus Developer webpage by clicking on the link below.
+1. From the web browser of your choice go to the Oculus Developer webpage by clicking on the link below.
 	
 	[https://developer.oculus.com/](https://developer.oculus.com/) 
 
-1. Download the [PC RUNTIME( v0.7.0.0-BETA)](https://developer.oculus.com/downloads/pc/0.7.0.0-beta/Oculus_Runtime_for_Windows/) by clicking on the PC RUNTIME(V0.7.0-BETA) button.
+1. Download the [PC RUNTIME( v0.7.0.0-BETA)](https://developer.oculus.com/downloads/pc/0.7.0.0-beta/Oculus_Runtime_for_Windows/) by clicking on the **PC RUNTIME(V0.7.0-BETA)** button.
 
 	![](Downloads.png)
 
diff --git a/Engine/Documentation/Source/Platforms/Oculus/QuickStart/3/OculusQuickStart_3.INT.udn b/Engine/Documentation/Source/Platforms/Oculus/QuickStart/3/OculusQuickStart_3.INT.udn
--- a/Engine/Documentation/Source/Platforms/Oculus/QuickStart/3/OculusQuickStart_3.INT.udn
+++ b/Engine/Documentation/Source/Platforms/Oculus/QuickStart/3/OculusQuickStart_3.INT.udn
@@ -3,6 +3,7 @@ Crumbs:%ROOT%, Platforms, Platforms/Oculus, Platforms/Oculus/QuickStart
 Description:Install the needed software for running games on the Oculus Rift
 Availability:Public
 Version: 4.9
+SkillLevel: Beginner
 
 In the section we will go over how you can install the Oculus drivers once they have been downloaded to your development PC.
 
@@ -25,7 +26,7 @@ In the section we will go over how you can install the Oculus drivers once they
 %Steps%
 
 
-1. Locate where the **0.7.0 BETA** drivers that were download in the last step and double-click on them to being the install process.
+1. Locate where the **0.7.0 BETA** drivers that were downloaded in the last step and **Double-click** on them to being the install process.
 
 1. When the installer starts, Click the **Next** button to being the install.
 	
@@ -59,7 +60,7 @@ In the section we will go over how you can install the Oculus drivers once they
 
 	 ![](oculus_setup_9.png)
 
-1. Reboot your computer when prompted.
+1. **Reboot** your computer when prompted.
 
 	![](oculus_setup_10.png)
 
diff --git a/Engine/Documentation/Source/Platforms/Oculus/QuickStart/4/OculusQuickStart_4.INT.udn b/Engine/Documentation/Source/Platforms/Oculus/QuickStart/4/OculusQuickStart_4.INT.udn
--- a/Engine/Documentation/Source/Platforms/Oculus/QuickStart/4/OculusQuickStart_4.INT.udn
+++ b/Engine/Documentation/Source/Platforms/Oculus/QuickStart/4/OculusQuickStart_4.INT.udn
@@ -3,6 +3,7 @@ Crumbs:%ROOT%, Platforms, Platforms/Oculus, Platforms/Oculus/QuickStart
 Description:Connect and update your Oculus Rift headset.
 Availability:Public
 Version: 4.9
+SkillLevel: Beginner
 
 [VAR:Steps]
 [OBJECT:Navigation]
diff --git a/Engine/Documentation/Source/Platforms/Oculus/QuickStart/5/OculusQuickStart_5.INT.udn b/Engine/Documentation/Source/Platforms/Oculus/QuickStart/5/OculusQuickStart_5.INT.udn
--- a/Engine/Documentation/Source/Platforms/Oculus/QuickStart/5/OculusQuickStart_5.INT.udn
+++ b/Engine/Documentation/Source/Platforms/Oculus/QuickStart/5/OculusQuickStart_5.INT.udn
@@ -3,8 +3,9 @@ Crumbs:%ROOT%, Platforms, Platforms/Oculus, Platforms/Oculus/QuickStart
 Description:Load and play content designed for the Oculus Rift using Unreal Engine 4.
 Availability:Public
 Version: 4.9
+SkillLevel: Beginner
 
-In this section we will cover how you can test to see if your Dk2 is working inside Unreal Engine 4.
+In this section we will cover how you can test to see if your DK2 is working inside Unreal Engine 4.
 
 [VAR:Steps]
 [OBJECT:Navigation]
diff --git a/Engine/Documentation/Source/Platforms/Oculus/QuickStart/6/OculusQuickStart_6.INT.udn b/Engine/Documentation/Source/Platforms/Oculus/QuickStart/6/OculusQuickStart_6.INT.udn
--- a/Engine/Documentation/Source/Platforms/Oculus/QuickStart/6/OculusQuickStart_6.INT.udn
+++ b/Engine/Documentation/Source/Platforms/Oculus/QuickStart/6/OculusQuickStart_6.INT.udn
@@ -3,6 +3,7 @@ Crumbs:%ROOT%, Platforms, Platforms/Oculus, Platforms/Oculus/QuickStart
 Description:Read more about Oculus development with Unreal Engine.
 Availability:Public
 Version: 4.9
+SkillLevel: Beginner
 
 Here are some additional resources that provide useful information for developing VR projects inside of Unreal Engine 4.
 
diff --git a/Engine/Documentation/Source/Platforms/Oculus/QuickStart/OculusQuickStart.INT.udn b/Engine/Documentation/Source/Platforms/Oculus/QuickStart/OculusQuickStart.INT.udn
--- a/Engine/Documentation/Source/Platforms/Oculus/QuickStart/OculusQuickStart.INT.udn
+++ b/Engine/Documentation/Source/Platforms/Oculus/QuickStart/OculusQuickStart.INT.udn
@@ -3,6 +3,7 @@ Crumbs:%ROOT%, Platforms, Platforms/Oculus
 Description:Instructions on setting up and launching games on the Oculus Rift.
 Availability:Public
 Version: 4.9
+SkillLevel: Beginner
 
 [VAR:Topic]
 [OBJECT:Topic]
