#UE4doc #docs - Added engine version field. Added/Removed bolded terms.

diff --git a/Engine/Documentation/Source/Resources/ContentExamples/Blueprint_Splines/Blueprint_Splines.INT.udn b/Engine/Documentation/Source/Resources/ContentExamples/Blueprint_Splines/Blueprint_Splines.INT.udn
--- a/Engine/Documentation/Source/Resources/ContentExamples/Blueprint_Splines/Blueprint_Splines.INT.udn
+++ b/Engine/Documentation/Source/Resources/ContentExamples/Blueprint_Splines/Blueprint_Splines.INT.udn
@@ -3,6 +3,7 @@ Title:Blueprint Splines Content Examples
 Crumbs:%ROOT%, Resources, Resources/ContentExamples
 Description:Overview of the samples provided in the Blueprint Splines Content Examples.
 Related: Engine/Rendering/Materials
+Version: 4.9
 
 [REGION:banner]
 ![](blueprintSplines_ce_header.png)(convert:false)
@@ -43,7 +44,7 @@ Related: Engine/Rendering/Materials
 
 
 
-The Blueprint Splines Content Examples show how you can use **Spline Components** with **Blueprints** to create paths that can be used in a number of different ways. By simply adding a Spline Component to your Blueprint and selecting an instance of that Blueprint, you will gain access to editing splines within your level. You will then be able to select spline keys inside the viewport and manipulate them like any other Actor using the Translation, Rotation, and Scale widgets. 
+The **Blueprint Splines** map in Content Examples will show how you can use **Spline Components** with Blueprints to create paths that can be used in a number of different ways. By simply adding a Spline Component to your Blueprint and selecting an instance of that Blueprint, you will gain access to editing splines within your level. You will then be able to select spline keys inside the viewport and manipulate them like any other Actor using the Translation, Rotation, and Scale widgets. 
 
 The Content Examples provided on this page illustrate how **Blueprint Splines** can be used with a Construction Script to create paths of Static Mesh objects or how to dynamically alter splines paths or Spline Meshes during gameplay to create content or animations. 
 
@@ -69,10 +70,10 @@ The **Dynamic Examples** show how a Spline Component can be updated dynamically
 
 | Example | What is Demonstrated |
 | ------- | ------------ |
-| **2.1 Using a Spline as a Path** | An example of a **Spline Component** used as a path for a **Particle Effect** to follow. |
-| **2.2 Changing Splines While Simulating** | Shows how **Spline Components** can be created dynamically during gameplay (e.g. a string of lanterns animate along a spline when playing or simulating). |
-| **2.3 Animating a Spline Mesh** | A **Spline Component** can also be used as a path for a **Spline Mesh** to move along (e.g. a vine grows out of the ground along a spline path). |
-| **2.4 Spline Mesh for Pawn Animation** | How a **Spline Mesh** can be used for animating a pawn (e.g. a **Spline Mesh** is used to animate a playable swimming fish). |
+| **2.1 Using a Spline as a Path** | An example of a Spline Component used as a path for a Particle Effect to follow. |
+| **2.2 Changing Splines While Simulating** | Shows how Spline Components can be created dynamically during gameplay (e.g. a string of lanterns animate along a spline when playing or simulating). |
+| **2.3 Animating a Spline Mesh** | A Spline Component can also be used as a path for a Spline Mesh to move along (e.g. a vine grows out of the ground along a spline path). |
+| **2.4 Spline Mesh for Pawn Animation** | How a Spline Mesh can be used for animating a pawn (e.g. a Spline Mesh is used to animate a playable swimming fish). |
 
 
 
diff --git a/Engine/Documentation/Source/Resources/ContentExamples/Brushes/BSP.INT.udn b/Engine/Documentation/Source/Resources/ContentExamples/Brushes/BSP.INT.udn
--- a/Engine/Documentation/Source/Resources/ContentExamples/Brushes/BSP.INT.udn
+++ b/Engine/Documentation/Source/Resources/ContentExamples/Brushes/BSP.INT.udn
@@ -4,6 +4,7 @@ Crumbs:%ROOT%, Resources, Resources/ContentExamples
 Description:An overview of the Geometry Editing Content Example map.
 Related: Engine/Actors/Brushes
 Related: Engine/UI/LevelEditor/Modes
+Version: 4.9
 
 
 [VAR:Topic]
@@ -52,7 +53,7 @@ Related: Engine/UI/LevelEditor/Modes
 ![](BSP_Header.png)
 [/REGION]
 
-[TOC(start:2 end:2)]
+<!---[TOC(start:2 end:2)]--->
 
 
 
@@ -74,8 +75,8 @@ Listed below are the examples provided in the Geometry_Editing map:
 | ------- | ------------ |
 | **1.1 Additive and Subtractive Geometry Example** | Shows how Additive and Subtractive Brushes interact with one another. |
 | **1.2 Order Example** | Shows how the order in which you place Additive or Subtractive Brushes can affect how they interact with one another. |
-| **1.3 Solidity Example** | Shows how a Brush can collide differently with objects based on its **Solidity** setting. (e.g. whether a Brush causes cuts to be created in surrounding brushes or if it collides with objects in the world). |
-| **1.4 Surface Properties Example** | Shows how to control how textures that are placed across your surfaces as well as lightmap resolution via the **Surface Properties** options. |
-| **1.5 Edit Geometry Example** | This example demonstrates how to edit a Brush's shape with **Edit Mode** by clicking on the Edit (**Shift+5**) button. |
+| **1.3 Solidity Example** | Shows how a Brush can collide differently with objects based on its Solidity setting. (e.g. whether a Brush causes cuts to be created in surrounding brushes or if it collides with objects in the world). |
+| **1.4 Surface Properties Example** | Shows how to control how textures that are placed across your surfaces as well as lightmap resolution via the Surface Properties options. |
+| **1.5 Edit Geometry Example** | This example demonstrates how to edit a Brush's shape with Edit Mode by clicking on the Edit (Shift+5) button. |
 | **1.6 Example House** | This is an example of using Additive and Subtractive Brushes to construct a house with windows and a door opening. |
 
