#UE4doc #docs - Added engine version field. Bolded term.

diff --git a/Engine/Documentation/Source/Resources/ContentExamples/Cloth/Cloth.INT.udn b/Engine/Documentation/Source/Resources/ContentExamples/Cloth/Cloth.INT.udn
--- a/Engine/Documentation/Source/Resources/ContentExamples/Cloth/Cloth.INT.udn
+++ b/Engine/Documentation/Source/Resources/ContentExamples/Cloth/Cloth.INT.udn
@@ -3,7 +3,7 @@ Title:Cloth Content Examples
 Crumbs: %ROOT%, Resources
 Description:Overview of the samples provided in the Cloth Content Examples.
 Related: Engine/Physics/Apex
-
+Version: 4.9
 
 [REGION:banner]
 ![](Cloth_Header.png)(convert:false)
@@ -52,7 +52,7 @@ The **Cloth** map inside the Content Examples project shows how to create and ma
 
 ![](Cloth_Topic.png)(w:640)
 
-Listed below are the examples provided inside the Cloth map:
+Listed below are the examples provided inside the **Cloth** map:
 
 | Example | What is Demonstrated |
 | ------- | ------------ |
