#UE4doc #docs - Added engine version field. Removed ToC and Related section to replace with Related tag.

diff --git a/Engine/Documentation/Source/Resources/ContentExamples/Decals/Decals.INT.udn b/Engine/Documentation/Source/Resources/ContentExamples/Decals/Decals.INT.udn
--- a/Engine/Documentation/Source/Resources/ContentExamples/Decals/Decals.INT.udn
+++ b/Engine/Documentation/Source/Resources/ContentExamples/Decals/Decals.INT.udn
@@ -2,8 +2,9 @@ Availability: Public
 Title:Decals
 Crumbs:%ROOT%, Resources, Resources/ContentExamples
 Description:An overview of the Decals example level
-
-[TOC(start:2 end:2)]
+Version: 4.9
+Related: Engine/Actors/DecalActor
+<!---[TOC(start:2 end:2)]--->
 
 [VAR:Topic]
 [OBJECT:Topic]
@@ -71,7 +72,8 @@ You then only need to apply a Decal Material, which is really just a regular Mat
 %Resources/ContentExamples/Decals/1_7:TopicCompact%
 [/REGION]
 
-
+<!---
 ## Related Topics and Further Reading
 
 * [Decals](Engine/Actors/DecalActor)
+--->
\ No newline at end of file
