#UE4 Docs: Adding the SteamVR and Motion Controller docs to the landing page.

diff --git a/Engine/Documentation/Source/Platforms/VR/VR.INT.udn b/Engine/Documentation/Source/Platforms/VR/VR.INT.udn
--- a/Engine/Documentation/Source/Platforms/VR/VR.INT.udn
+++ b/Engine/Documentation/Source/Platforms/VR/VR.INT.udn
@@ -57,6 +57,7 @@ vr
 	[PARAM:links]
 		* [](Platforms/VR/CheatSheet "%Platforms/VR/CheatSheet:description%")
 		* [](Platforms/VR/ContentSetup "%Platforms/VR/ContentSetup:description%")
+		* [](Platforms/VR/MotionController "%Platforms/VR/ContentSetup:description%")
 	[/PARAM]
 [/OBJECT]
 
@@ -89,3 +90,19 @@ vr
 		* [](Platforms/GearVR/QuickStart "%Platforms/GearVR/QuickStart:description%")
 	[/PARAM]
 [/OBJECT]
+
+[OBJECT:TopicButtonList]
+	[PARAM:icon]
+		![](%ROOT%/SteamVR_Icon.png)(convert:false)
+	[/PARAM]
+	[PARAM:title]
+		[](Platforms/SteamVR "%Platforms/GearVR:description%")
+	[/PARAM]
+	[PARAM:description]
+	[/PARAM]
+	[PARAM:links]
+		* [](Platforms/SteamVR/BestPractices   "%Platforms/GearVR/BestPractices:description%")
+		* [](Platforms/SteamVR/QuickStart      "%Platforms/GearVR/QuickStart:description%")
+		* [](Platforms/SteamVR/Tips_And_Tricks "%Platforms/GearVR/QuickStart:description%")
+	[/PARAM]
+[/OBJECT]
