#UE4 #DOCS PhysMat link replacement

diff --git a/Engine/Documentation/Source/Engine/Physics/PhysicalMaterials/Reference/PhysMatReference.INT.udn b/Engine/Documentation/Source/Engine/Physics/PhysicalMaterials/Reference/PhysMatReference.INT.udn
--- a/Engine/Documentation/Source/Engine/Physics/PhysicalMaterials/Reference/PhysMatReference.INT.udn
+++ b/Engine/Documentation/Source/Engine/Physics/PhysicalMaterials/Reference/PhysMatReference.INT.udn
@@ -6,19 +6,19 @@ Description:Assets applied to physically simulated primitives directly or via ma
 [VAR:Topic]
 [OBJECT:Topic]
 	[PARAM:image]
-		![%Engine/Physics/PhysicalMaterials/PhysMatReference:title%](Engine/Physics/physics_topic.png)
+		![%Engine/Physics/PhysicalMaterials/Reference:title%](Engine/Physics/physics_topic.png)
 	[/PARAM]
 	[PARAM:icon]
 		![](%ROOT%/physics_icon.png)(convert:false)
 	[/PARAM]
 	[PARAM:title]
-		%Engine/Physics/PhysicalMaterials/PhysMatReference:title%
+		%Engine/Physics/PhysicalMaterials/Reference:title%
 	[/PARAM]
 	[PARAM:description]
-		%Engine/Physics/PhysicalMaterials/PhysMatReference:description%
+		%Engine/Physics/PhysicalMaterials/Reference:description%
 	[/PARAM]
 	[PARAM:path]
-		[RELATIVE:Engine/Physics/PhysicalMaterials/PhysMatReference]
+		[RELATIVE:Engine/Physics/PhysicalMaterials/Reference]
 	[/PARAM]
 [/OBJECT]
 [/VAR]
@@ -30,13 +30,13 @@ Description:Assets applied to physically simulated primitives directly or via ma
 		![](%ROOT%/physics_icon.png)(convert:false)
 	[/PARAM]
 	[PARAM:title]
-		%Engine/Physics/PhysicalMaterials/PhysMatReference:title%
+		%Engine/Physics/PhysicalMaterials/Reference:title%
 	[/PARAM]
 	[PARAM:description]
-		%Engine/Physics/PhysicalMaterials/PhysMatReference:description%
+		%Engine/Physics/PhysicalMaterials/Reference:description%
 	[/PARAM]
 	[PARAM:path]
-		[RELATIVE:Engine/Physics/PhysicalMaterials/PhysMatReference]
+		[RELATIVE:Engine/Physics/PhysicalMaterials/Reference]
 	[/PARAM]
 [/OBJECT]
 [/VAR]
diff --git a/Engine/Documentation/Source/Resources/ContentExamples/Physics/1_1/1_1.INT.udn b/Engine/Documentation/Source/Resources/ContentExamples/Physics/1_1/1_1.INT.udn
--- a/Engine/Documentation/Source/Resources/ContentExamples/Physics/1_1/1_1.INT.udn
+++ b/Engine/Documentation/Source/Resources/ContentExamples/Physics/1_1/1_1.INT.udn
@@ -7,7 +7,7 @@ Related: Engine/Physics/PhAT/UserGuide
 Related: Engine/Physics/Constraints/ConstraintsReference
 Related: Engine/Physics/Constraints/ConstraintsUserGuide
 Related: Engine/Physics/Collision
-Related: Engine/Physics/PhysicalMaterials/PhysMatReference
+Related: Engine/Physics/PhysicalMaterials/Reference
 
 [VAR:Steps]
 [OBJECT:Navigation]
diff --git a/Engine/Documentation/Source/Resources/ContentExamples/Physics/1_2/1_2.INT.udn b/Engine/Documentation/Source/Resources/ContentExamples/Physics/1_2/1_2.INT.udn
--- a/Engine/Documentation/Source/Resources/ContentExamples/Physics/1_2/1_2.INT.udn
+++ b/Engine/Documentation/Source/Resources/ContentExamples/Physics/1_2/1_2.INT.udn
@@ -7,7 +7,7 @@ Related: Engine/Physics/PhAT/UserGuide
 Related: Engine/Physics/Constraints/ConstraintsReference
 Related: Engine/Physics/Constraints/ConstraintsUserGuide
 Related: Engine/Physics/Collision
-Related: Engine/Physics/PhysicalMaterials/PhysMatReference
+Related: Engine/Physics/PhysicalMaterials/Reference
 
 [VAR:Steps]
 [OBJECT:Navigation]
diff --git a/Engine/Documentation/Source/Resources/ContentExamples/Physics/1_3/1_3.INT.udn b/Engine/Documentation/Source/Resources/ContentExamples/Physics/1_3/1_3.INT.udn
--- a/Engine/Documentation/Source/Resources/ContentExamples/Physics/1_3/1_3.INT.udn
+++ b/Engine/Documentation/Source/Resources/ContentExamples/Physics/1_3/1_3.INT.udn
@@ -7,7 +7,7 @@ Related: Engine/Physics/PhAT/UserGuide
 Related: Engine/Physics/Constraints/ConstraintsReference
 Related: Engine/Physics/Constraints/ConstraintsUserGuide
 Related: Engine/Physics/Collision
-Related: Engine/Physics/PhysicalMaterials/PhysMatReference
+Related: Engine/Physics/PhysicalMaterials/Reference
 
 [VAR:Steps]
 [OBJECT:Navigation]
diff --git a/Engine/Documentation/Source/Resources/ContentExamples/Physics/1_4/1_4.INT.udn b/Engine/Documentation/Source/Resources/ContentExamples/Physics/1_4/1_4.INT.udn
--- a/Engine/Documentation/Source/Resources/ContentExamples/Physics/1_4/1_4.INT.udn
+++ b/Engine/Documentation/Source/Resources/ContentExamples/Physics/1_4/1_4.INT.udn
@@ -7,7 +7,7 @@ Related: Engine/Physics/PhAT/UserGuide
 Related: Engine/Physics/Constraints/ConstraintsReference
 Related: Engine/Physics/Constraints/ConstraintsUserGuide
 Related: Engine/Physics/Collision
-Related: Engine/Physics/PhysicalMaterials/PhysMatReference
+Related: Engine/Physics/PhysicalMaterials/Reference
 
 [VAR:Steps]
 [OBJECT:Navigation]
diff --git a/Engine/Documentation/Source/Resources/ContentExamples/Physics/1_6/1_6.INT.udn b/Engine/Documentation/Source/Resources/ContentExamples/Physics/1_6/1_6.INT.udn
--- a/Engine/Documentation/Source/Resources/ContentExamples/Physics/1_6/1_6.INT.udn
+++ b/Engine/Documentation/Source/Resources/ContentExamples/Physics/1_6/1_6.INT.udn
@@ -7,7 +7,7 @@ Related: Engine/Physics/PhAT/UserGuide
 Related: Engine/Physics/Constraints/ConstraintsReference
 Related: Engine/Physics/Constraints/ConstraintsUserGuide
 Related: Engine/Physics/Collision
-Related: Engine/Physics/PhysicalMaterials/PhysMatReference
+Related: Engine/Physics/PhysicalMaterials/Reference
 
 [VAR:Steps]
 [OBJECT:Navigation]
diff --git a/Engine/Documentation/Source/Resources/ContentExamples/Physics/1_7/1_7.INT.udn b/Engine/Documentation/Source/Resources/ContentExamples/Physics/1_7/1_7.INT.udn
--- a/Engine/Documentation/Source/Resources/ContentExamples/Physics/1_7/1_7.INT.udn
+++ b/Engine/Documentation/Source/Resources/ContentExamples/Physics/1_7/1_7.INT.udn
@@ -7,7 +7,7 @@ Related: Engine/Physics/PhAT/UserGuide
 Related: Engine/Physics/Constraints/ConstraintsReference
 Related: Engine/Physics/Constraints/ConstraintsUserGuide
 Related: Engine/Physics/Collision
-Related: Engine/Physics/PhysicalMaterials/PhysMatReference
+Related: Engine/Physics/PhysicalMaterials/Reference
 
 [VAR:Steps]
 [OBJECT:Navigation]
diff --git a/Engine/Documentation/Source/Resources/ContentExamples/Physics/1_8/1_8.INT.udn b/Engine/Documentation/Source/Resources/ContentExamples/Physics/1_8/1_8.INT.udn
--- a/Engine/Documentation/Source/Resources/ContentExamples/Physics/1_8/1_8.INT.udn
+++ b/Engine/Documentation/Source/Resources/ContentExamples/Physics/1_8/1_8.INT.udn
@@ -7,7 +7,7 @@ Related: Engine/Physics/PhAT/UserGuide
 Related: Engine/Physics/Constraints/ConstraintsReference
 Related: Engine/Physics/Constraints/ConstraintsUserGuide
 Related: Engine/Physics/Collision
-Related: Engine/Physics/PhysicalMaterials/PhysMatReference
+Related: Engine/Physics/PhysicalMaterials/Reference
 
 [VAR:Steps]
 [OBJECT:Navigation]
diff --git a/Engine/Documentation/Source/Resources/ContentExamples/Physics/Physics.INT.udn b/Engine/Documentation/Source/Resources/ContentExamples/Physics/Physics.INT.udn
--- a/Engine/Documentation/Source/Resources/ContentExamples/Physics/Physics.INT.udn
+++ b/Engine/Documentation/Source/Resources/ContentExamples/Physics/Physics.INT.udn
@@ -7,7 +7,7 @@ Related: Engine/Physics/PhAT/UserGuide
 Related: Engine/Physics/Constraints/ConstraintsReference
 Related: Engine/Physics/Constraints/ConstraintsUserGuide
 Related: Engine/Physics/Collision
-Related: Engine/Physics/PhysicalMaterials/PhysMatReference
+Related: Engine/Physics/PhysicalMaterials/Reference
 
 [REGION:banner]
 ![](Physics_Header.png)(convert:false)
