#UE4 #DOCS: Physics Updates. Moving from one workspace to another, reference updates, link updates

diff --git a/Engine/Documentation/Source/Engine/Physics/Collision/Collision.INT.udn b/Engine/Documentation/Source/Engine/Physics/Collision/Collision.INT.udn
--- a/Engine/Documentation/Source/Engine/Physics/Collision/Collision.INT.udn
+++ b/Engine/Documentation/Source/Engine/Physics/Collision/Collision.INT.udn
@@ -1,8 +1,7 @@
 Availability: Public
-Title:Collision Responses
+Title:Collision
 Crumbs: %ROOT%, Engine, Engine/Physics
 Description:All about collision, collision responses, and collision presets.
-Version: 4.9
 
 [VAR:Topic]
 [OBJECT:Topic]
@@ -24,208 +23,16 @@ Version: 4.9
 [/OBJECT]
 [/VAR]
 
-[TOC(start:2 end:3)]
-
-
-**Collision Responses** and **Trace Responses** form the basis for how Unreal Engine 4 handles collision and ray casting during run time. Every object that can collide 
-gets an **Object Type** and a series of responses that define how it interacts will all other object types. When a collision or overlap event occurs, both (or all) 
-objects involved can be set to affect or be affected by blocking, overlapping, or ignoring each other.
-
-Trace Responses work basically the same way, except the trace (ray cast) itself can be defined as one of the Trace Response types, thusly allowing Actors to either 
-block it or ignore it based on _their_ Trace Responses.
-
-## Properties
-
-![](COL_collisionBasics.png)
-
-In the properties list, there are a large number of physics related key words.  The two that are not really clear are **Simple** and **Complex Collision**. In short, 
-Simple Collision shapes are your [**boxes, KDOP, spheres, capsules, and convex shapes**](Engine/Physics\Collision\Reference), all of which can be generated in the Static
-Mesh editor. Complex Collision is per-poly, and it is never used for an Actor that is simulating physics (it will just fall through the world).
-
-
-| **Property** | **Description** |
-| -------- | ----------- |
-| **Simulation Generates Hit Events** | This flag indicates if **Event Hit** and **On Component Hit** should be called for physics collisions involving this PrimitiveComponent. In more simple terms, it lets this object "know" it has actually had a physics collision. However, it is just a notification, and a dynamic object set to collide with the static environment will still do so regardless of this setting; it just will not fire those events. |
-| **Phys Material Override** | This allows you to override the Phsyical Material to use for simple collision on this body. |
-| **Generate Overlap Events** | If you want an object to generate overlap events, like **Event Actor Begin Overlap** or **Event Actor End Overlap**, then this flag needs to be **true**. Further, for it to be of any use, the object in question will need to be able to overlap another. |
-| [REGION:tablesection] &nbsp; [/REGION]||
-| **Collision Presets** | These are a collection of presets for the **Collision Responses** below. There is an INI file you can edit to change or add profiles. **To edit anything below this point, you need to set this to "Custom..."**|
-| **Collision Enabled** | The 3 settings in here are as follows: |
-| ^ | _No Collision_ - No collision is performed against this object, neither trace nor physics. |
-| ^ | _Query Only (No Physics Collision)_ - This body is only used for raycasts, sweeps, and overlaps. |
-| ^ | _Physics Only (No Query Collision)_ - This body is used only for physics collision. |
-| ^ | _Collision Enabled (Query and Physics)_ - This body is used for physics simulations and collision queries. |
-| **Object Type** | This defines what the object should be considered when it moves. It has more to do with the interactions you see below than it does with defining an object as static or dynamic. |
-| [REGION:tablesection] &nbsp; [/REGION]||
-| **Collision Responses** | This sets all the collision responses below to the value checked. So if you check "Block" here, all of the rows below will also check "Block." |
-| **Trace Responses** | These define how an object should behave when interacting with a trace (usually just done with a ray cast). An object can choose to **Block**, **Overlap**, or simply **Ignore** a trace from the given source. There is no difference in the default trace responses except for their names. In code or Blueprint, you can use either for tracing a ray cast from one thing to another. But, for organization of traces, they should be used as their name implies. The default responses for **Trace Responses** are: |
-| ^ | _Visibility_ - Trace from one location to another, if an object that `blocks` this type of trace is in the way, the trace is `blocked`. |
-| ^ | _Camera_ - Exactly the same as _Visibility_, but should be used for setting when tracing from the camera. |
-| **Object Responses** | An object can choose to Block, Overlap, or simply Ignore an object it is interacting with. The default responses for **Object Responses** are: |
-| ^ | _WorldStatic_ - Volumes and world geometry should be _WorldStatic_. | 
-| ^ | _WorldDynamic_ - Moving Actors outside of _Pawns_, _PhysicsBodies_, _Vehicles_, and _Destructible_ Actors. Like an elevator. | 
-| ^ | _Pawn_ - Characters. | 
-| ^ | _PhysicsBody_ - Any physics object that is or can be simulated in the world. | 
-| ^ | _Vehicle_ - A good response to have so you can have _Pawns_ jump into them. | 
-| ^ | _Destructible_ - Destructible Actors. |
-| [REGION:tablesection] &nbsp; [/REGION]||
-| **Use CCD** | Whether or not to use **Continuous Collision Detection** for this component. Increases accuracy of collision detection. |
-| **Always Create Physics State** | Indicates if you would like to create a **physics state** (its collision properties, mass, if it is sleeping, etc... all go into a physics state) all the time (for collision and simulation). If set to _false_, the state will still be generated if collision or simulation are activated on this Actor, say through a Blueprint or code, but there will be some overhead as it has to work the new object into the full physics simulation. Setting this to **true** can improve performance by removing the overhead for calculating the physics state of the object once it actually needs one (on collision or simulation). | 
-| **Multi Body Overlap** | If **true**, this component will generate individual overlaps for each overlapping physics body if it is a multi-body component. Think about skeletal physics assets, with all their individual collision shapes. With this enabled, a character's hand would generate _its own_ overlap event. Thus giving more control over what is reported and what to do in that specific case. |
-| **Check Async Scene on Move** | If this is flagged as **true**, the component will look for collisions in both of the physics scenes (synchronous and asynchronous). The asynchronous scene is primarily used by the broken pieces of destructible Actors. |
-| **Trace Complex on Move** | If this is set to **true**, component sweeps with this component will trace against complex collision during movement. Complex Collision is, in short, per face collision, where as Simple Collision are your spheres, capsules, boxes, and generated convex shapes. |
-| **Return Material on Move ** | Setting this to **true** will return the Physics Material in the **Hit Info**. |
-| **Can Ever Affect Navigation ** | Setting this to **true** can potentially influence navigation. |
-
-
-Some components will have additional or fewer options in the collision category, and many times these very specific options will only be shown in the appropriate editor 
-for that component.
-
-|Property | Description |
-| -------- | ----------- |
-| [REGION:tablesection] **Physics Body in Physics Asset** [/REGION] ||
-| **Collision Response** | Whether or not a specific physics body in a Physics Asset will collide at all. |
-
-## Interactions
-
-There are a few rules to keep in mind with how collisions are handled:
-
-* Blocking will naturally occur between two (or more) Actors set to Block. However, Simulation Generates Hit Events needs to be enabled to execute Event Hit, which is used in Blueprints, Destructible Actors, Triggers, etc...
-* Setting Actors to **Overlap** will often look like they **Ignore** each other, and without **Generate Overlap Events**, they are essentially the same.
-* For two or more simulating objects to block each other, they both need to be set to block their respective object types.
-* For two or more simulating objects: if one is set to overlap an object, and the second object is set to block the other, the overlap will occur but not the block.
-* Overlap events can be generated even if an object Blocks another, especially if traveling at high speeds.
-	* It is not recommended for an object to have both collision and overlap events. Though possible, there is much that needs manual handling.
-* If one object is set to ignore and the other is set to overlap, no overlap events will be fired.
-* The Object Responses listed are just the defaults and have no actual bearing on what an object _is_. See below for more information on setting up custom [**Object Types and Trace Channels**](#customobjectchannelsandtracechannels).
-
-For purposes of testing levels and looking around your worlds:
-
-* The default **Play In Editor** camera is a pawn. Thusly can be blocked by anything set to block pawns.
-* The **Simulate in Editor** camera, before possessing anything, is not a pawn. It can freely clip through everything and will not create any collision or overlap events.
-
-
-## Common Collision Interactions
-
-[REGION:note]
-These interactions assume that all the objects have **Collision Enabled** set to **Collision Enabled** so they are set to fully collide with everything. If collision is 
-disabled, it is as if _Ignore_ has been set for all Collision Responses.
-[/REGION]
-
-For the following section, this will be the setup used to explain what is happening:
-
-![](COL_setup.png)(w:940 convert:false)
-
-The sphere is a **PhysicsBody** and the box is **WorldDynamic**, and by changing their collision settings we can get a number of behaviors.
-
-### Collision
-
-<!--
-[EXCERPT:spacerEX]
-[REGION:raw]
-![](spacer.png)
-[/REGION]
-[/EXCERPT:spacerEX]
- -->
-
-[VAR:spacer]
-[INCLUDE:#spacerEX]
-[/VAR]
-
-By setting both of their collision settings to block each other, you get a collision, great for having objects interact with each other:
-
-![](COL_collideNoEvent.png)(w:940 convert:false)
-
-| Sphere Collision Setup| Wall Collision Setup |
-| --- | ---- |
-| ![](COL_collideNoEvent_Sphere.png)(w:300)  |![](COL_collideNoEvent_Box.png)(w:301)  |
-| In this case, the sphere is a **PhysicsBody** and it is set to `block` **WorldDynamic** (which is what the wall is). | The wall is a **WorldDynamic** and is set to `block` **PhysicsBody** Actors (which is what the sphere is). |
-
-In this case, the sphere and the wall will simply collide; no further notifications of the collision will take place.
-
-### Collision and Simulation Generates Hit Events
-
-Just collision is useful and in general, the bare minimum for physics interactions, but if you want something to "report" it has collided so a Blueprint or section of 
-code can be triggered:
-
-![](COL_collideEvent.png)(w:940 convert:false)
-
-| Sphere Collision Setup| Wall Collision Setup |
-| --- | ---- |
-| ![](COL_collideEvent_Sphere.png)(w:300) |![](COL_collideNoEvent_Box.png)(w:301)|
-| As in the example above, the sphere is a **PhysicsBody** and it is set to `block` **WorldDynamic** (which is what the wall is). However, the sphere has also enabled Simulation Generates Hit Events so it will trigger an event for itself whenever it collides with something. | The wall is a **WorldDynamic** and is set to `block` **PhysicsBody** Actors (which is what the sphere is). Since the wall is not set to Simulation Generates Hit Events, it will not generate an event for itself. |
-
-With the sphere set to Simulation Generates Hit Events, the sphere will tell itself that it has had a collision. It will fire off events such as **ReceiveHit** or 
-**OnComponentHit** in the sphere's blueprint. Now if the box had an event for collision, it would not fire because it will never notify itself it has happened.
-
-Further, an object that is reporting rigid collisions will report them all and spam reports when it is just sitting on something, so it is best to be careful to filter 
-what it is colliding within its blueprint or in code.
-
-### Overlap and Ignore
-
-For all intents and purposes, **Overlap** and **Ignore** work exactly the same assuming Generate Overlap Events is disabled. In this case, the sphere is set to 
-overlap or ignore the box:
-
-![](COL_ignore.png)(w:940 convert:false)
-
-
-
-
-| Sphere Collision Setup| Wall Collision Setup |
-| --- | ---- |
-| ![](COL_OverlapNoEvent_Sphere.png)(w:300)  |![](COL_collideNoEvent_Box.png)(w:303) |
-| Here the sphere is set to `overlap` **WorldDynamic** Actors (like our wall), but it does not have Generate Overlap Events enabled. As far as the sphere is concerned, it has not collided or overlapped anything, effectively it has ignored the wall.  | The wall is a **WorldDynamic** and is set to `block` **PhysicsBody** Actors (which is what the sphere is). As stated above, both Actors need to be set to block each other's respective object types. If they do not, they will not collide. |
-
-Or:
-
-| Sphere Collision Setup| Wall Collision Setup |
-| --- | ---- |
-| ![](COL_ignore_sphere.png)(w:300) |![](COL_collideNoEvent_Box.png)(w:303) |
-| Here the sphere is set to `ignore` **WorldDynamic** Actors (like our wall), and it will pass through the wall. | The wall is a **WorldDynamic** and is set to `block` **PhysicsBody** Actors (which is what the sphere is). As stated above, both Actors need to be set to block each other's respective object types. If they do not, they will not collide. |
-
-
-
-### Overlap and Generate Overlap Events
-
-Unlike collisions that can fire every frame, the overlap events are **BeginOverlap** and **EndOverlap**, which only fire in those specific cases.
-
-![](COL_overlapEvent.png)(w:940 convert:false)
-
-| Sphere Collision Setup| Wall Collision Setup |
-| --- | ---- |
-| ![](SphereCollisiionSetup.png) | ![](WallCollisionSetup.png) |
-| Here the sphere is set to `overlap` **WorldDynamic** Actors (like our wall), and it will generate an event for itself when it does overlap something. | The wall is a **WorldDynamic** and is set to `block` **PhysicsBody** Actors (which is what the sphere is). As stated above, both Actors need to be set to block each other's respective object types. If they do not, they will not collide. But, an **Overlap** does occur here, and an event is fired. |
-
- 
-## Custom Object Channels and Trace Channels
-
-There will be times where 6 Object Response Channels and 2 Trace Response Channels simply are not granular enough for what you want to create. This is what the Collision 
-Editor in your Project Settings is for. You can access them from the **Edit Menu** -> **Project Settings** -> **Collision**:
-
-![](COL_projectSettings_1.png)
-
-[REGION:lightbox]
-[![](COL_projectSettings_2.png)(w:940)](COL_projectSettings_2.png)
-[/REGION]
-
-From here you can add new Object Response Channels and Trace Response Channels. Click the **New Object Channel...** or **New Trace Channel...** button, provide a name, 
-select a **Default Response**, and click **Accept**.
-
-![](COL_customChannel.png)
-
-You can have up to 18 Custom Object Response Channels or Custom Trace Response Channels.
-
-### Presets
-
-Custom Presets can also be set here by expanding the **Preset** Category and clicking the **New...** button.
-
-![](COL_customPresets.png)
-
-From here you can name your preset, enable or disable collision, select your preset's Object Type, and finally define the behavior of each Response Channel for the 
-selected Object Type.
-
-
-
-
-
-
+[OBJECT:TopicButtonList]
+	[PARAM:icon]
+		![](%ROOT%/start_icon.png)(convert:false)
+	[/PARAM]
+	[PARAM:title]
+		Essentials
+	[/PARAM]
+	[PARAM:links]
+		* [](Engine/Physics/Collision/Overview "%Engine/Physics/Collision/Overview:description%")
+		* [](Engine/Physics/Collision/HowTo "%Engine/Physics/Collision/HowTo:description%")
+		* [](Engine/Physics/Collision/Reference "%Engine/Physics/Collision/Reference:description%")
+	[/PARAM]
+[/OBJECT]
\ No newline at end of file
diff --git a/Engine/Documentation/Source/Engine/Physics/Collision/Reference/Reference.INT.udn b/Engine/Documentation/Source/Engine/Physics/Collision/Reference/Reference.INT.udn
--- a/Engine/Documentation/Source/Engine/Physics/Collision/Reference/Reference.INT.udn
+++ b/Engine/Documentation/Source/Engine/Physics/Collision/Reference/Reference.INT.udn
@@ -1,9 +1,8 @@
 Availability: Public
-Title:Static Mesh Collision Reference
+Title:Collision Response Reference
 Crumbs: %ROOT%, Engine, Engine/Physics, Engine/Physics/Collision
-Description:Content guide to creating and setting up collision geometry.
+Description:The Collision Presets section in the Collision Property Category has a large number of properties broken out here for clarity.
 Version: 4.9
-SkillLevel: Beginner
 
 [VAR:Topic]
 [OBJECT:Topic]
@@ -25,128 +24,155 @@ SkillLevel: Beginner
 [/OBJECT]
 [/VAR]
 
-[TOC(start:2 end:2)]
+There are a large number of properties and settings in the Collision Presets/Collision Response section of the Collision Properties category, so they've been separated out for clarity.
 
+## Properties
 
-As graphical detail has increased, it is important to ensure that collision within the game remains simple. This is because of two reasons; it is much faster to calculate collision 
-against a few polygons and it is a much smoother experience for the player as then the player does not get snagged on a small piece geometry. 
+Below are the properties for the Collision Presets, a sub-category in the Collision Properties on Physics Bodies (BodyInstances).
 
-[REGION:tip]
-All collision shapes produced in the **Static Mesh Editor** can be moved, rotated, and scaled after they are created. You can create more than one collision shape per Static Mesh so 
-you can get the fidelity you need in your Static Mesh's collision.
-
-This is most useful with the simple collision shapes (box, sphere, and capsule), but can be used with the Convex Hull generator and K-DOP hulls.
-[/REGION]
-
-## Creating Skeletal Mesh Collision hulls
-
-Please see **[PhAT User Guide](Engine/Physics/PhAT)** on how to create collision hulls used by Skeletal Meshes.
-
-## Creating Static Mesh Collision Hulls in the Static Mesh Editor
-
-![](COLREF_collisionMenu.png)
-
-### K-DOP
+![](collProp.png)(convert:false)
 
+[OBJECT:modulesectionexpanded]
+	[PARAMLITERAL:id]
+		coll
+	[/PARAMLITERAL]
+	[PARAM:heading]
+		Collision Presets
+	[/PARAM]
+	[PARAM:content]
+		%p2%
+	[/PARAM]
+[/OBJECT]
 
+[VAR:p2]
 
 [REGION:raw]
-![](COLREF_collisionMenu_KDOP.png)
+| Property | %spacer% | Description |
+| --- | :---: | :--- |
+| **Collision Presets** || [INCLUDE:#colpre] |
+| **Collision Enabled** || [INCLUDE:#colen] |
+| **Object Type** || [INCLUDE:#objtype] |
+| **Collision Responses** || [INCLUDE:#colrep] |
+| **Trace Responses** || [INCLUDE:#deftrace] |
+| **Object Responses** || [INCLUDE:#objtrace] |
 [/REGION]
 
-There are a series of menu options in the Static Mesh Editor under the **Collision** menu entry, named _##DOP_, these are the **K-DOP** simple collision generators. 
-**K-DOP** is a type of bounding volume, which stands for _K discrete oriented polytope_ (where K is the number of axis aligned planes). Basically it takes K axis-aligned planes
-and pushes them as close to the mesh as it can. The resulting shape is used as a collision hull. In the Static Mesh Editor K can be:
+[/VAR]
+
 
+[VAR:spacer]![](%ROOT%/spacer.png)(w:150 convert:false)[/VAR]
 
-* **10** - Box with 4 edges beveled - you can choose X- Y- or Z-aligned edges.
-* **18** - Box with all edges beveled.
-* **26** - Box with all edges and corners beveled.
+[VAR:spacer2]![](%ROOT%/spacer.png)(w:40 convert:false)[/VAR]
 
-See below for an example. This tool is quite handy for packages full of pipes, pillars, and railings:
 
-![kdop_sizes.jpg](kdop_sizes.jpg)
+<!--
+[EXCERPT:colrep]
 
+These define how this Physics Body should interact with all of the other types of Trace and Object Types. Remember, it's the interaction between both Physics Bodies that define the resulting action, so the Object Type and Collision Responses of both Physics Bodies matter.
 
-### Creating a Collision Hull Using the Auto Convex Collision Tool
+![](iob.png)
 
 [REGION:raw]
-![](COLREF_collisionMenu_Auto.png)
+| Response | %spacer2% | Description |
+| --- | :---: | :--- |
+| **Ignore** || Regardless of the Collision Response of the other Physics Body, this Physics Body will ignore the interaction. |
+| **Overlap** || If the other Physics Body is set to Overlap or Block this Physics Body's Object Type, an Overlap Event can occur. |
+| **Block** || If the other Physics Body is set to Block this Physics Body's Object Type, a Hit Event can occur. |
 [/REGION]
 
-The Static Mesh Editor has the capability to generate collision using a number of convex hulls.
-
-![Collision_CreateCollisionHullWithAutoConvex_0.jpg](Collision_CreateCollisionHullWithAutoConvex_0.png)(convert:false)
+[/EXCERPT]
 
-This rounded cube would probably work just fine with **6DOP Simplified Collision** but let us assume it needs more than that.
 
-Selecting **Collision Menu -> Auto Convex Collision** will bring up this panel in the lower right hand side:
 
-![Collision_CreateCollisionHullWithAutoConvex_1.jpg](Collision_CreateCollisionHullWithAutoConvex_1.png)(convert:false)
+[EXCERPT:defTrace]
 
-* **Accuracy** - Determines the the fewest number of hulls to create to best match the shape of the object.
-* **Max Hull Verts** - Maximum number of collision hull vertices. Increasing this value increase how complex the collision hulls are allowed to be.
-* **Apply** - Adds collision to the object based on the settings above.
-* **Defaults** - Resets the values of Accuracy and Max Hull Verts to the values seen in the image above.
-
-![Collision_CreateCollisionHullWithAutoConvex_2.jpg](Collision_CreateCollisionHullWithAutoConvex_2.png)(convert:false)
-
-As you can see, the auto convex collision is reasonably accurate.
-
-### Creating a Collision Using Simple Shapes
+Trace responses are used in Traces (ray casts), such as the Blueprint Node **Line Trace by Channel**.
 
 [REGION:raw]
-![](COLREF_collisionMenu_BuildYourOwn.png)
+| Property | %spacer2% | Description |
+| --- | :---: | :--- |
+| **Visibility** || General visibility testing channel. |
+| **Camera** || Usually used when tracing from the camera to something. |
 [/REGION]
 
-These options create basic collision shapes:
+[/EXCERPT]
 
-[REGION:raw]
-![](COLREF_BuildYourOwn_01.png)(convert:false)
-[/REGION]
-
-Once a simple collision shape is created, it can be moved, rotated, and scaled into the desired position. You can create more than one collision shape for more complex Static Meshes.
-
-| **Control** | **Action** |
-| --- | --- |
-| **W, E, R** | Switch between Translation, Rotation, and Scale tools. |
-| **Delete** | Will delete the currently selected Physics Body or Bodies. |
-| **F** | Will focus on the currently selected Physics Body or Bodies. |
-| **Ctrl + Left Click** | Adds a Physics Body to the selection. |
-| **Alt + Left Drag (On a manipulation widget)** | Duplicates a Physics Body. |
-
-
-### Creating a Collision Hull Using Your 3D Content Creation Package
-
-Please see **[FBX Static Mesh Pipeline](Engine/Content/FBX/StaticMeshes#Collision)** on how to create Static Mesh collision hulls within a 3D content creation package.
-
-
-## Collision Responses
+[EXCERPT:objTrace]
 
-Unreal Engine 4 is pretty drastically different in terms of how you can have collision objects interact with each other. Now each object is an arbitrary object type 
-(a Static Mesh can be in the **WorldDynamic** collision response group for example), that can choose to **ignore**, **block**, or **overlap** an object from another response group.
 
-The short version is that instead of an object being "I will collide with all things set to collide" and object can be set to "I will collide with only objects that are 
-_Pawns_ and I will overlap _WorldDynamic_ while ignoring everything else." Further, using that object that will only collide with _Pawns_, and object in the _Pawn_ response 
-group can choose to **ignore** the group our pawn collider exists in.
 
-It can get really complicated but the tools at your disposal make it easy to get a handle on. See the **[Collision Responses](Engine/Physics\Collision)** documentation for more info.
+[REGION:raw]
+| Property | %spacer2% | Description |
+| --- | :---: | :--- |
+| **WorldStatic** || How this Physics Body should react when interaction with a WorldStatic Physics Body Object Type. |
+| **WorldDynamic** || How this Physics Body should react when interaction with a WorldDynamic Physics Body Object Type. |
+| **Pawn** || How this Physics Body should react when interaction with a Pawn Physics Body Object Type. |
+| **PhysicsBody** || How this Physics Body should react when interaction with a PhysicsBody Physics Body Object Type. |
+| **Vehicle** || How this Physics Body should react when interaction with a Vehicle Physics Body Object Type. |
+| **Destructible** || How this Physics Body should react when interaction with a Destructible Physics Body Object Type. |
+[/REGION]
+[/EXCERPT]
 
+[EXCERPT:colPre]
 
-## Reviewing Collision In-Game
+The default collision profiles and any you have created in Project **Settings** -> **Engine** -> **Collision** -> **Preset** (see: []() for information on creating custom collision profiles), will appear here.
 
-There are several console commands you can use in-game for reviewing the collision setup. Here are a few. Typing them toggles them on and off.
+[REGION:raw]
+| Property | %spacer2% | Description |
+| --- | :---: | :--- |
+| **Custom...** || Set all custom collision settings for this instance. |
+| **NoCollision** || No Collision. |
+| **BlockAll** || WorldStatic object that blocks all actors by default. All new custom channels will use its own default response.  |
+| **OverlapAll** || WorldStatic object that overlaps all actors by default. All new custom channels will use its own default response.  |
+| **BlockAllDynamic** || WorldDynamic object that blocks all actors by default. All new custom channels will use its own default response.  |
+| **OverlapAllDynamic** || WorldDynamic object that overlaps all actors by default. All new custom channels will use its own default response.  |
+| **IngoreOnlyPawn** || WorldDynamic object that ignores Pawn and Vehicle. All other channels will be set to default. |
+| **OverlapOnlyPawn** || WorldDynamic object that overlaps Pawn, Camera, and Vehicle. All other channels will be set to default.  |
+| **Pawn** || Pawn object. Can be used for capsule of any playable character or AI.  |
+| **Spectator** || Pawn object that ignores all other actors except WorldStatic. |
+| **CharacterMesh** || Pawn object that is used for Character Mesh. All other channels will be set to default. |
+| **PhysicsActor** || Simulating actors. |
+| **Destructible** || Destructible actors. |
+| **InvisibleWall** || WorldStatic object that is invisible. |
+| **InvisibleWallDynamic** || WorldDynamic object that is invisible. |
+| **Trigger** || WorldDynamic object that is used for trigger. All other channels will be set to default. |
+| **Ragdoll** || Simulating Skeletal Mesh Component. All other channels will be set to default. |
+| **Vehicle** || Vehicle object that blocks Vehicle, WorldStatic, and WorldDynamic. All other channels will be set to default. |
+| **UI** || WorldStatic object that overlaps all actors by default. All new custom channels will use its own default response.  |
+[/REGION]
+[/EXCERPT]
 
-* **show COLLISION** - This will draw all collision models and blocking volumes in use in the level. 
 
-![](COLREF_showCollision.png)(w:900)
 
-* **stat game** - This shows you various useful stats on how long different types of collision are taking.
+[EXCERPT:colEn]
 
-![](Collision_StatGame.png)(w:900)
+Collision Enabled can have 4 possible states as listed below.
 
+[REGION:raw]
+| Property | %spacer2% | Description |
+| --- | :---: | :--- |
+| **No Collision** || No collision is enabled for this body.  |
+| **Query Only** || This body is used only for collision queries (raycasts, sweeps, and overlaps). |
+| **Physics Only** || This body is used only for physics collision. |
+| **Collision Enabled** || This body interacts with all collision (Query and Physics). |
+[/REGION]
+[/EXCERPT]
 
 
 
+[EXCERPT:objType]
 
+%doxygen:FBodyInstance::ObjectType% These are arbitrary, but should be set in accordance to what the Physics Body will be doing.
 
+[REGION:raw]
+| Property | %spacer2% | Description |
+| --- | :---: | :--- |
+| **WorldStatic** || This should be used for any Actor that doesn't move. Static Mesh Actors are a good example of an Actor that will probably have a WorldStatic type. |
+| **WorldDynamic** || WorldDynamic is for Actor types that will be moving under the influence of animation or code; kinematic. Lifts and doors are good examples of WorldDynamic Actors. |
+| **Pawn** || Any player controlled entity should have the Pawn type. The player's character is a good example of an Actor that should receive the Pawn Object Type. |
+| **PhysicsBody** || Any Actor that will be moving due to the physics simulation. |
+| **Vehicle** || Vehicles receive this type by default. |
+| **Destructible** || Destructible Meshes receive this by default. |
+[/REGION]
+[/EXCERPT]
+-->
\ No newline at end of file
diff --git a/Engine/Documentation/Source/Engine/Physics/Constraints/ConstraintsReference/Constraints.INT.udn b/Engine/Documentation/Source/Engine/Physics/Constraints/ConstraintsReference/Constraints.INT.udn
--- a/Engine/Documentation/Source/Engine/Physics/Constraints/ConstraintsReference/Constraints.INT.udn
+++ b/Engine/Documentation/Source/Engine/Physics/Constraints/ConstraintsReference/Constraints.INT.udn
@@ -4,39 +4,316 @@ Crumbs: %ROOT%, Engine, Engine/Physics
 Description:A reference for all the properties available on Physics Constraints in Unreal Editor.
 Version: 4.9
 
-[TOC(start:2)]
+[VAR:Topic]
+[OBJECT:Topic]
+	[PARAM:image]
+		![%Engine/Physics/Constraints/ConstraintsReference:title%](Engine/Physics/physics_topic.png)
+	[/PARAM]
+	[PARAM:icon]
+		![](%ROOT%/physics_icon.png)(convert:false)
+	[/PARAM]
+	[PARAM:title]
+		%Engine/Physics/Constraints/ConstraintsReference:title%
+	[/PARAM]
+	[PARAM:description]
+		%Engine/Physics/Constraints/ConstraintsReference:description%
+	[/PARAM]
+	[PARAM:path]
+		[RELATIVE:Engine/Physics/Constraints/ConstraintsReference]
+	[/PARAM]
+[/OBJECT]
+[/VAR]
+[VAR:TopicCompact]
+[OBJECT:TopicCompact]
+	[PARAM:image]
+		![%Engine/Physics/Constraints/ConstraintsReference:title%](Engine/Physics/physics_topic.png)
+	[/PARAM]
+	[PARAM:icon]
+		![](%ROOT%/physics_icon.png)(convert:false)
+	[/PARAM]
+	[PARAM:title]
+		%Engine/Physics/Constraints/ConstraintsReference:title%
+	[/PARAM]
+	[PARAM:description]
+		%Engine/Physics/Constraints/ConstraintsReference:description%
+	[/PARAM]
+	[PARAM:path]
+		[RELATIVE:Engine/Physics/Constraints/ConstraintsReference]
+	[/PARAM]
+[/OBJECT]
+[/VAR]
+
+
+
+This page contains a reference listing of properties in the Physics Constraints. 
+
+## Properties
+
+Below are the properties for Physics Constraints seperated by major category.
+
+[VAR:spacer]![](%ROOT%/spacer.png)(w:150 convert:false)[/VAR]
+
+[VAR:spacer2]![](%ROOT%/spacer.png)(w:40 convert:false)[/VAR]
+
+[OBJECT:modulesectionexpanded]
+	[PARAMLITERAL:id]
+		cons
+	[/PARAMLITERAL]
+	[PARAM:heading]
+		Constraint
+	[/PARAM]
+	[PARAM:content]
+		%p1%
+	[/PARAM]
+[/OBJECT]
+
+[VAR:p1]
+
+![](propConstraint.png)
+
+[REGION:raw]
+| Property | %spacer% | Description |
+| --- | :---: | :--- |
+| **Constraint Actor 1** || %doxygen:UPhysicsConstraintComponent::ConstraintActor1% |
+| **Component Name 1** || %doxygen:UPhysicsConstraintComponent::ComponentName1% |
+| **Constraint Actor 2** || %doxygen:UPhysicsConstraintComponent::ConstraintActor2% |
+| **Component Name 2** || %doxygen:UPhysicsConstraintComponent::ComponentName2%  |
+| **Joint Name** || %doxygen:FConstraintInstance::JointName% |
+| **Constraint Bone 1** || %doxygen:FConstraintInstance::ConstraintBone1% |
+| **Constraint Bone 2** || %doxygen:FConstraintInstance::ConstraintBone2% |
+| **Disable Collision** || %doxygen:FConstraintInstance::bDisableCollision% |
+[/REGION]
 
-[EXCERPT:PhysXConstraints_001]
-A constraint is basically a type of joint. It allows you to connect two Actors together (presumably one physically simulating), and also apply limits or forces. Unreal Engine 4 
-has a very flexible and data-driven constraint system that allows designers to make many different types of joints simply by changing some options. The engine has a few default 
-joint types (ball-and-socket, hinge, prismatic) but they only differ in their settings. You should feel free to start with any joint type and experiment with making your own 
-configurations.
-[/EXCERPT:PhysXConstraints_001]
+[/VAR]
 
-## Property Reference
 
-There are 3 ways to create PhysX Constraints: **Blueprint**, **Actor**, and in the **Physics Asset Tool** (**PhAT**).
 
-### Actor Constraint
+[OBJECT:modulesectionexpanded]
+	[PARAMLITERAL:id]
+		proj
+	[/PARAMLITERAL]
+	[PARAM:heading]
+		Projection
+	[/PARAM]
+	[PARAM:content]
+		%p2%
+	[/PARAM]
+[/OBJECT]
+
+[VAR:p2]
+
+![](propProjection.png)
+
+[REGION:raw]
+| Property | %spacer% | Description |
+| --- | :---: | :--- |
+| **Enable Projection** || %doxygen:FConstraintInstance::bEnableProjection% |
+| **Projection Linear Tolerance** || %doxygen:FConstraintInstance::ProjectionLinearTolerance% |
+| **Projection Angular Tolerance** || %doxygen:FConstraintInstance::ProjectionAngularTolerance% |
+[/REGION]
+
+[/VAR]
+
+
+
+
+[OBJECT:modulesectionexpanded]
+	[PARAMLITERAL:id]
+		join
+	[/PARAMLITERAL]
+	[PARAM:heading]
+		Joint Presets
+	[/PARAM]
+	[PARAM:content]
+		%p3%
+	[/PARAM]
+[/OBJECT]
+
+[VAR:p3]
+
+![](propJointPresets.png)
+
+[REGION:raw]
+| Property | %spacer% | Description |
+| --- | :---: | :--- |
+| **Hinge** || Setup joint as hinge. A hinge joint allows motion only in one plane. |
+| **Prismatic** || Setup joint as prismatic. A prismatic joint allows only linear sliding movement. |
+| **Ball and Socket** || Setup joint as ball and socket. A Ball and Socket joint allows motion around an indefinite number of axes, which have one common center. |
+[/REGION]
+
+[/VAR]
+
+
+
+[OBJECT:modulesectionexpanded]
+	[PARAMLITERAL:id]
+		linlim
+	[/PARAMLITERAL]
+	[PARAM:heading]
+		Linear Limits
+	[/PARAM]
+	[PARAM:content]
+		%p4%
+	[/PARAM]
+[/OBJECT]
+
+[VAR:p4]
+
+![](propLinearLimits.png)
+
+[REGION:raw]
+| Property | %spacer% | Description |
+| --- | :---: | :--- |
+| **Linear XMotion** || %doxygen:FConstraintInstance::LinearXMotion% |
+| **Linear YMotion** || %doxygen:FConstraintInstance::LinearYMotion% |
+| **Linear ZMotion** || %doxygen:FConstraintInstance::LinearZMotion% |
+| **Linear Limit Size** || %doxygen:FConstraintInstance::LinearLimitSize% |
+[/REGION]
+
+### Advanced
+
+![](propLinearLimitsADV.png)
+
+[REGION:raw]
+| Property | %spacer% | Description |
+| --- | :---: | :--- |
+| **Linear Limit Soft** || %doxygen:FConstraintInstance::bLinearLimitSoft% |
+| **Linear Limit Stiffness** || %doxygen:FConstraintInstance::LinearLimitStiffness% |
+| **Linear Limit Damping** || %doxygen:FConstraintInstance::LinearLimitDamping% |
+| **Linear Breakable** || %doxygen:FConstraintInstance::bLinearBreakable% |
+| **Linear Break Threshold** || %doxygen:FConstraintInstance::LinearBreakThreshold% |
+[/REGION]
+
+[/VAR]
+
 
-| Property | Description |
-| -------- | ----------- |
-|[REGION:tablesection]Constraint[/REGION]||
-| **Constraint Actor 1** | %doxygen:UPhysicsConstraintComponent::ConstraintActor1% |
-| **Component Name 1** | %doxygen:UPhysicsConstraintComponent::ComponentName1% |
-| **Constraint Actor 2** | %doxygen:UPhysicsConstraintComponent::ConstraintActor2% |
-| **Component Name 2** | %doxygen:UPhysicsConstraintComponent::ComponentName2%  |
-| **Joint Name** | %doxygen:FConstraintInstance::JointName% |
-| **Constraint Bone 1** | %doxygen:FConstraintInstance::ConstraintBone1% |
-| **Constraint Bone 2** | %doxygen:FConstraintInstance::ConstraintBone2% |
-| **Disable Collision** | %doxygen:FConstraintInstance::bDisableCollision% |
 
+
+
+
+[OBJECT:modulesectionexpanded]
+	[PARAMLITERAL:id]
+		anglim
+	[/PARAMLITERAL]
+	[PARAM:heading]
+		Angular Limits
+	[/PARAM]
+	[PARAM:content]
+		%p5%
+	[/PARAM]
+[/OBJECT]
+
+[VAR:p5]
+
+![](propsAngularLimits.png)
+
+[REGION:raw]
+| Property | %spacer% | Description |
+| --- | :---: | :--- |
+| **Angular Swing 1Motion** || %doxygen:FConstraintInstance::AngularSwing1Motion% |
+| **Angular Twist Motion** || %doxygen:FConstraintInstance::AngularTwistMotion% |
+| **Angular Swing 2Motion** || %doxygen:FConstraintInstance::AngularSwing2Motion% |
+| **Swing 1Limit Angle** || %doxygen:FConstraintInstance::Swing1LimitAngle% |
+| **Twist Limit Angle** || %doxygen:FConstraintInstance::TwistLimitAngle% |
+| **Swing 2Limit Angle** || %doxygen:FConstraintInstance::Swing2LimitAngle% |
+| **Angular Rotation Offset** || %doxygen:FConstraintInstance::AngularRotationOffset% |
+[/REGION]
+
+### Advanced
+
+![](propsAngularLimitsADV.png)
+
+[REGION:raw]
+| Property | %spacer% | Description |
+| --- | :---: | :--- |
+| **Swing Limit Soft** || %doxygen:FConstraintInstance::bSwingLimitSoft% |
+| **Swing Limit Stiffness** || %doxygen:FConstraintInstance::SwingLimitStiffness% |
+| **Swing Limit Damping** || %doxygen:FConstraintInstance::SwingLimitDamping% |
+| **Twist Limit Soft** || %doxygen:FConstraintInstance::bTwistLimitSoft% |
+| **Twist Limit Stiffness** || %doxygen:FConstraintInstance::TwistLimitStiffness% |
+| **Twist Limit Damping** || %doxygen:FConstraintInstance::TwistLimitDamping% |
+| **Angular Breakable** || %doxygen:FConstraintInstance::bAngularBreakable% |
+| **Angular Break Threshold** || %doxygen:FConstraintInstance::AngularBreakThreshold% |
+[/REGION]
+
+[/VAR]
+
+
+
+
+[OBJECT:modulesectionexpanded]
+	[PARAMLITERAL:id]
+		linmot
+	[/PARAMLITERAL]
+	[PARAM:heading]
+		Linear Motor
+	[/PARAM]
+	[PARAM:content]
+		%p6%
+	[/PARAM]
+[/OBJECT]
+
+[VAR:p6]
+
+![](propLinearMotor.png)
+
+[REGION:raw]
+| Property | %spacer% | Description |
+| --- | :---: | :--- |
+| **Linear Position Drive** || %doxygen:FConstraintInstance::bLinearPositionDrive% |
+| **Linear Position Target** || %doxygen:FConstraintInstance::LinearPositionTarget% |
+| **Linear Position Strength** || %doxygen:FConstraintInstance::LinearDriveSpring% |
+| **Linear Velocity Drive** || %doxygen:FConstraintInstance::bLinearVelocityDrive% |
+| **Linear Velocity Target** || %doxygen:FConstraintInstance::LinearVelocityTarget% |
+| **Linear Velocity Strength** || %doxygen:FConstraintInstance::LinearDriveDamping% |
+| **Max Linear Force** || %doxygen:FConstraintInstance::LinearDriveForceLimit% |
+[/REGION]
+
+[/VAR]
+
+
+
+
+[OBJECT:modulesectionexpanded]
+	[PARAMLITERAL:id]
+		angmot
+	[/PARAMLITERAL]
+	[PARAM:heading]
+		Angular Motor
+	[/PARAM]
+	[PARAM:content]
+		%p7%
+	[/PARAM]
+[/OBJECT]
+
+[VAR:p7]
+
+![](propsAngularMotor.png)
+
+[REGION:raw]
+| Property | %spacer% | Description |
+| --- | :---: | :--- |
+| **Angular Orientation Drive** || %doxygen:FConstraintInstance::bAngularOrientationDrive% |
+| **Angular Orientation Target** || %doxygen:FConstraintInstance::AngularOrientationTarget% |
+| **Angular Orientation Strength** || %doxygen:FConstraintInstance::AngularDriveSpring% |
+| **Angular Velocity Drive** || %doxygen:FConstraintInstance::bAngularVelocityDrive% |
+| **Angular Velocity Target** || %doxygen:FConstraintInstance::AngularVelocityTarget% |
+| **Angular Velocity Strength** || %doxygen:FConstraintInstance::AngularDriveDamping% |
+| **Max Angular Force** || %doxygen:FConstraintInstance::AngularDriveForceLimit% |
+| **Angular Drive Mode** || %doxygen:FConstraintInstance::AngularDriveMode% |
+[/REGION]
+
+[/VAR]
+
+
+
+<!--
 ![](PhysCon_JointPresets.png)(convert:false) 
 
 [INCLUDE:#presetTable]
 
 * For information on using **Physics Constraint Actor**, see: [**%Engine/Physics\Constraints\ConstraintsUserGuide:title%**](Engine/Physics\Constraints\ConstraintsUserGuide)
-<!--
+
 [EXCERPT:presetTable]
 | Joint Preset | Description |
 | --- | ---- |
@@ -44,7 +321,7 @@ There are 3 ways to create PhysX Constraints: **Blueprint**, **Actor**, and in t
 | **Prismatic** | _Locks_ the **Swing Motions**, _Locks_ the **Angular Twist Motion**, _Locks_ **Linear Y** and **Linear Z motions**, and sets **Linear ZMotion** to _Free_. Like a hydraulic piston. |
 | **Ball and Socket** | _Locks_ all **Linear Motions**, sets the **Swing Motions** and **Twist Motions** to _Free_. Completely free to rotate, similar to a punching bag or wrecking ball. |
 [/EXCERPT:presetTable]
--->
+
 
 ### Blueprint Constraint
 
@@ -252,7 +529,7 @@ You can simulate joint friction using the drive settings. Turn on velocity drive
 You could simulate something like a punch-bag (a ball on the top of a vertical pole that can be knocked around) with a joint by first creating a ball-and-socket joint. Then turn on swing position drive and adjust the **AngularDriveSpring** and **AngularDriveDamping** until it springs nicely around the upright position when pushed. You could turn on swing limits as well to ensure that it cannot bend over too far.
 
 -->
- 
+-->
 
 
 
diff --git a/Engine/Documentation/Source/Engine/Physics/Constraints/PhysicsConstraints.INT.udn b/Engine/Documentation/Source/Engine/Physics/Constraints/PhysicsConstraints.INT.udn
--- a/Engine/Documentation/Source/Engine/Physics/Constraints/PhysicsConstraints.INT.udn
+++ b/Engine/Documentation/Source/Engine/Physics/Constraints/PhysicsConstraints.INT.udn
@@ -4,17 +4,19 @@ Crumbs: %ROOT%, Engine, Engine/Physics
 Description:Constraining physics objects together and to the world.
 Version: 4.9
 
-[OBJECT:TopicList]
+
+[OBJECT:TopicButtonList]
 	[PARAM:icon]
-		![](%ROOT%/physics_icon.png)(convert:false)
+		![](%ROOT%/start_icon.png)(convert:false)
 	[/PARAM]
 	[PARAM:title]
-		Constraint Topics
+		Essentials
 	[/PARAM]
 	[PARAM:links]
-		* [Physics Constraints User Guide](Engine/Physics/Constraints/ConstraintsUserGuide) - A guide to constraining the joints between bones.
-		* [](Engine/Physics/Constraints/ConstraintsBlueprints) - A user guide on using the Physics Constraint Component in Blueprints.
-		* [](Engine/Physics/Constraints/ConstraintsReference) - A reference for all the properties available on Physics Constraints in Unreal Editor.
-		* [Damping and friction](Engine/Physics/Constraints/DampingAndFriction) - Explanation of the damping and friction attributes in Unreal Engine 4.
+		* [Physics Constraints User Guide](Engine/Physics/Constraints/ConstraintsUserGuide "%Engine/Physics/Constraints/ConstraintsUserGuide:description%")
+		* [](Engine/Physics/Constraints/ConstraintsBlueprints "%Engine/Physics/Constraints/ConstraintsBlueprints:description%")
+		* [](Engine/Physics/Constraints/ConstraintsReference "%Engine/Physics/Constraints/ConstraintsReference:description%")
+		* [](Engine/Physics/FrictionRestitutionAndDamping "%Engine/Physics/FrictionRestitutionAndDamping:description%")
 	[/PARAM]
 [/OBJECT]
+
diff --git a/Engine/Documentation/Source/Engine/Physics/PhysicalMaterials/PhysMatUserGuide/PhysMatUserGuide.INT.udn b/Engine/Documentation/Source/Engine/Physics/PhysicalMaterials/PhysMatUserGuide/PhysMatUserGuide.INT.udn
--- a/Engine/Documentation/Source/Engine/Physics/PhysicalMaterials/PhysMatUserGuide/PhysMatUserGuide.INT.udn
+++ b/Engine/Documentation/Source/Engine/Physics/PhysicalMaterials/PhysMatUserGuide/PhysMatUserGuide.INT.udn
@@ -66,7 +66,7 @@ This document will cover the creation and usage of **Physical Materials** as wel
 
 1. **Click Save** ![](Shared\Icons\icon_SaveAsset_40x.png)
 
-See the [**Physical Material Reference**](Engine/Physics\PhysicalMaterials\PhysMatReference) for information on the properties in the Physical Material.
+See the [**Physical Material Reference**](Engine/Physics\PhysicalMaterials\Reference) for information on the properties in the Physical Material.
 
 ### Surface Type
 
diff --git a/Engine/Documentation/Source/Engine/Physics/PhysicalMaterials/PhysicalMaterials.INT.udn b/Engine/Documentation/Source/Engine/Physics/PhysicalMaterials/PhysicalMaterials.INT.udn
--- a/Engine/Documentation/Source/Engine/Physics/PhysicalMaterials/PhysicalMaterials.INT.udn
+++ b/Engine/Documentation/Source/Engine/Physics/PhysicalMaterials/PhysicalMaterials.INT.udn
@@ -45,20 +45,16 @@ Version: 4.9
 
 [INCLUDE:Shared/Glossary/P/#Glossary_Physical_Material]
 
-[OBJECT:TopicList]
+[OBJECT:TopicButtonList]
 	[PARAM:icon]
-		![](%ROOT%/physics_icon.png)(convert:false)
+		![](%ROOT%/start_icon.png)(convert:false)
 	[/PARAM]
 	[PARAM:title]
 		Physical Materials Topics
 	[/PARAM]
 	[PARAM:links]
-		* [](Engine/Physics/PhysicalMaterials/PhysMatUserGuide "%Engine/Physics/PhysicalMaterials/PhysMatUserGuide:description%")
-		* [](Engine/Physics/PhysicalMaterials/PhysMatReference "%Engine/Physics/PhysicalMaterials/PhysMatReference:description%")
+		* [](Engine/Physics/PhysicalMaterials/HowTo "%Engine/Physics/PhysicalMaterials/HowTo:description%")
+		* [](Engine/Physics/PhysicalMaterials/Reference "%Engine/Physics/PhysicalMaterials/Reference:description%")
 	[/PARAM]
 [/OBJECT]
 
-
-
-
-
diff --git a/Engine/Documentation/Source/Engine/Physics/PhysicsBodies/PhysicsBodies.INT.udn b/Engine/Documentation/Source/Engine/Physics/PhysicsBodies/PhysicsBodies.INT.udn
--- a/Engine/Documentation/Source/Engine/Physics/PhysicsBodies/PhysicsBodies.INT.udn
+++ b/Engine/Documentation/Source/Engine/Physics/PhysicsBodies/PhysicsBodies.INT.udn
@@ -1,8 +1,7 @@
 Availability: Public
-Title:Physics Bodies Reference
+Title:Physics Bodies
 Crumbs: %ROOT%, Engine, Engine/Physics
-Description:The basis for most physics interactions: the Physics Body.
-Version: 4.9
+Description:The basis for most physics interactions: the Physics Body [BodyInstance].
 
 [VAR:Topic]
 [OBJECT:Topic]
@@ -44,87 +43,36 @@ Version: 4.9
 [/VAR]
 
 
-[TOC(start:2)]
+Given the ever increasing complexity of 3D meshes, proxy shapes must be used to drive physics simulations in real time 3D rendering applications; Physics Bodies (BodyInstances) are the simplified 3D meshes that Unreal Engine 4 uses for physics simulations. They can be formed of boxes, spheres, capsules, or convex hulls.
 
+The following links have information on Physics Bodies.
 
-Given the ever increasing complexity of 3D meshes, proxy shapes must be used to drive physics simulations in real time 3D rendering applications; Physics Bodies are the simplified 3D meshes that Unreal Engine 4 uses for physics simulations. They can be formed of boxes, spheres, capsules, or convex hulls.
-
-The following links have information on the various editors and workflows for creating and working with Physics Bodies.
-
-
-### Static Mesh
-
-For creating Physics Bodies using your 3D art program, see: [**FBX: Static Meshes**](Engine/Content\FBX\StaticMeshes#Collision)
-
-For creating Physics Bodies using the Static Mesh Editor, see: [**Collision Reference: Static Mesh Editor**](Engine/Physics\Collision\Reference#CreatingStaticMeshCollisionHullsintheStaticMeshEditor)
-
-
-### Physics Assets/Skeletal Meshes
-
-For documentation over the creation of Physics Assets for use by Skeletal Meshes, see: **[](Engine/Physics\PhAT)** documentation.
-
-
-### Shape Components
-
-The Box, Sphere, and Capsule components will automatically gain Physics Bodies if they are set for any sort of collision, overlap, or other physics simulation events.
-
-
-## Reference
-
-All Physics Bodies have a shared pool of properties. Below are the properties from the major categories that appear with Physics Bodies. Not all properties will appear on all Actors or Components, depending on the context of the property.
-
-
-### Physics
-
-| Property | Description |
-| --- | --- |
-|[REGION:tablesection]Physics[/REGION]||
-| **Mass in KG** | Mass of the body in KG. |
-| **Linear Damping** | %doxygen:FBodyInstance::LinearDamping% |
-| **Angular Damping** | %doxygen:FBodyInstance::AngularDamping% |
-| **Enable Gravity** | %doxygen:FBodyInstance::bEnableGravity% |
-| **Physics Type** | %doxygen:UBodySetup::PhysicsType% |
-| **Double Sided Geometry** | If _True_, the physics triangle mesh will use double sided faces when doing scene queries. This is useful for planes and single sided meshes that need traces to work on both sides. |
-| **Simple Collision Physical Material** | Physical material to use for simple collision on this body. Encodes informatrion about density, friction, etc. |
-|[REGION:tablesection]Extras[/REGION]||
-| **Start Awake** | %doxygen:FBodyInstance::bStartAwake% |
-| **Center Of Mass Offset** | User Specified offset for the center of mass of this object, from the calculated location. |
-| **Mass Scale** | %doxygen:FBodyInstance::MassScale% |
-| **Max Angular Velocity** | The maximum angular velocity for this instance. |
-| **Use Async Scene** | %doxygen:FBodyInstance::bUseAsyncScene% |
-| **Generate Wake Events** | Should 'wake/sleep' events fire when this object is woken up or put to sleep by the physics simulation. |
-| **Max Depenetration Velocity** | The maximum velocity used to depenetrate this object. |
-| **Walkable Slope Override** | %doxygen:UBodySetup::WalkableSlopeOverride% |
-| **Walkable Slope Behavior** | Behavior of this surface (whether we affect the walkable slope). |
-| **Walkable Slope Angle** | %doxygen:FWalkableSlopeOverride::WalkableSlopeAngle% |
-| **Custom Sleep Threshold Multiplier** | If the SleepFamily is set to custom, multiply the natural sleep threshold by this amount. A higher number will cause the body to sleep sooner. |
-| **Position Solver Iteration Count** | %doxygen:FBodyInstance::PositionSolverIterationCount% |
-| **Velocity Solver Iteration Count** | %doxygen:FBodyInstance::VelocitySolverIterationCount% |
-| **Sleep Family** | %doxygen:FBodyInstance::SleepFamily% |
-
-
-### Collision
-
-| Property | Description |
-| --- | --- |
-|[REGION:tablesection]Collision[/REGION]||
-| **Simulation Generates Hit Events** | Should 'Hit' events fire when this objects collides during physics simulation. |
-| **Phys Material Override** | %doxygen:FBodyInstance::PhysMaterialOverride% |
-|[REGION:tablesection]Collision Responses[/REGION]||
-| | See the [](Engine/Physics\Collision) documentation for more information. |
-|[REGION:tablesection]Extras[/REGION]||
-| **Use CCD** | %doxygen:FBodyInstance::bUseCCD% |
-| **Always Create Physics State** | %doxygen:UPrimitiveComponent::bAlwaysCreatePhysicsState% |
-| **Generate Overlap Events** | %doxygen:UPrimitiveComponent::bGenerateOverlapEvents% |
-| **Multi Body Overlap** | %doxygen:UPrimitiveComponent::bMultiBodyOverlap% | 
-| **Trace Complex On Move** | %doxygen:UPrimitiveComponent::bTraceComplexOnMove% |
-| **Return Material On Move** | %doxygen:UPrimitiveComponent::bReturnMaterialOnMove% |
-| **Check Async Scene On Move** | %doxygen:UPrimitiveComponent::bCheckAsyncSceneOnMove% |
-| **Can Ever Affect Navigation** | %doxygen:UPrimitiveComponent::bCanEverAffectNavigation% |
-
+[OBJECT:TopicButtonList]
+	[PARAM:icon]
+		![](%ROOT%/start_icon.png)(convert:false)
+	[/PARAM]
+	[PARAM:title]
+		Essentials
+	[/PARAM]
+	[PARAM:links]
+		* [](Engine/Physics/PhysicsBodies/Overview "%Engine/Physics/PhysicsBodies/Overview:description%")
+		<!-- * [](Engine/Physics/PhysicsBodies/HowTo "%Engine/Physics/PhysicsBodies/HowTo:description%") -->
+		* [](Engine/Physics/PhysicsBodies/Reference "%Engine/Physics/PhysicsBodies/Reference:description%")
+	[/PARAM]
+[/OBJECT]
 
 
+[OBJECT:TopicButtonList]
+	[PARAM:icon]
+		![](%ROOT%/start_icon.png)(convert:false)
+	[/PARAM]
+	[PARAM:title]
+		Related
+	[/PARAM]
+	[PARAM:links]
 
+	[/PARAM]
+[/OBJECT]
 
 
 
diff --git a/Engine/Documentation/Source/Engine/Physics/Substepping/Substepping.INT.udn b/Engine/Documentation/Source/Engine/Physics/Substepping/Substepping.INT.udn
--- a/Engine/Documentation/Source/Engine/Physics/Substepping/Substepping.INT.udn
+++ b/Engine/Documentation/Source/Engine/Physics/Substepping/Substepping.INT.udn
@@ -45,8 +45,7 @@ Version: 4.9
 
 
 
-One of the lesser known features in UE4 is the ability to turn on physics **Sub-stepping**. By doing this, you can get physics simulations that are more accurate and stable. 
-The most noticeable improvement will be with ragdoll jitter and other complex physical assets.
+By using physics **Sub-stepping** you can get physics simulations that are more accurate and stable, however, this comes at the expense of performance. The most noticeable improvement will be with ragdoll jitter and other complex physical assets.
 
 ## Enabling Sub-Stepping
 
@@ -60,13 +59,18 @@ Sub-stepping can be turned on in the **Project Setting** -> **Physics** tab:
 This feature is still a work in progress, and is not fully supported by APEX Destruction. It is also currently not supported on mobile devices. However, we are working towards this as it can greatly improve the look and feel.
 [/REGION]
 
-| Property | Description |
-| --- | --- |
-| **Substepping** | This enables or disables sub-stepping in your current project. |
-| **Substepping Async** | Whether to substep the async phsyics simulation. |
-| **Max Substep Delta Time** | This is the maximum time, in seconds a sub-step is allowed to take. So, if a full step takes .05 seconds, and **Max Substep Delta Time** is set to .025, the full step will be split into 2 sub-steps. Now, if the full step were to take a substantial amount of time, say 2 full seconds, then the step will be evenly split based on the value in **Max Substeps**, and not **80** `(2.0/0.025=80)` sub-steps. It is worth noting that **Max Physics Delta Time** limits how long a physics step is permitted to take. |
-| **Max Substeps** | The maximum number of sub-steps a full step is permitted to be broken into. |
+[REGION:raw]
+| Property | %spacer% | Description |
+| --- | :---: | :--- |
+| **Substepping** || This enables or disables sub-stepping in your current project. |
+| **Substepping Async** || Whether to substep the async phsyics simulation. |
+| **Max Substep Delta Time** || This is the maximum time, in seconds a sub-step is allowed to take. So, if a full step takes .05 seconds, and **Max Substep Delta Time** is set to .025, the full step will be split into 2 sub-steps. Now, if the full step were to take a substantial amount of time, say 2 full seconds, then the step will be evenly split based on the value in **Max Substeps**, and not **80** `(2.0/0.025=80)` sub-steps. It is worth noting that **Max Physics Delta Time** limits how long a physics step is permitted to take. |
+| **Max Substeps** || The maximum number of sub-steps a full step is permitted to be broken into. |
+[/REGION]
+
+[VAR:spacer]![](%ROOT%/spacer.png)(w:150 convert:false)[/VAR]
 
+[VAR:spacer2]![](%ROOT%/spacer.png)(w:40 convert:false)[/VAR]
 
 ## Technical Details
 
