#UE4doc #docs:
Added engine version field
Replaced multiple images
Changed a few steps
Replaced mentions of a lot of out of date nodes

diff --git a/Engine/Documentation/Source/Gameplay/HowTo/UseTimers/Blueprints/UseTimers_BP.INT.udn b/Engine/Documentation/Source/Gameplay/HowTo/UseTimers/Blueprints/UseTimers_BP.INT.udn
--- a/Engine/Documentation/Source/Gameplay/HowTo/UseTimers/Blueprints/UseTimers_BP.INT.udn
+++ b/Engine/Documentation/Source/Gameplay/HowTo/UseTimers/Blueprints/UseTimers_BP.INT.udn
@@ -2,6 +2,7 @@ Availability: Public
 Title:Using Timers in Blueprints
 Crumbs: %ROOT%, Gameplay, Gameplay\HowTo\UseTimers
 Description:A How To Guide for using timers in Blueprints. 
+version: 4.9
 
 [TOC (start:2 end:2)]
 
@@ -11,7 +12,7 @@ This page covers how to use **Timers** for gameplay purposes in Blueprints and t
 
 The steps below will show you how you can create/set a Timer in Blueprints. 
 
-_For this example, we are using the **Blueprint Third Person Template** and will be causing fire damage to the player over time._
+For this example, we are using the **Blueprint Third Person Template** and will be causing fire damage to the player over time.
 
 1.	Inside your project, open the **ThirdPersonCharacter** Blueprint. 
 
@@ -25,19 +26,21 @@ _For this example, we are using the **Blueprint Third Person Template** and will
 
 	This Function will be called over time by the **Timer** and will apply the damage. 
 
-1.	Drag off the **PlayerInFire** Event and add the **Set Timer** node. 
+1.	Drag off the **PlayerInFire** Event and add the **Set Timer by Function Name** node. 
 
-1.	Drag off the **PlayerNotInFire** Event and add the **Clear Timer** node. 
+1.	Drag off the **PlayerNotInFire** Event and add the **Clear Timer by Handle** node. 
 
-1.	For both **Set Timer** and **Clear Timer**, set the **Function Name** to **FireDamage**. 
+1.	For **Set Timer by Function Name**, set the **Function Name** to **FireDamage**. 
 
 1.	On the **Set Timer**, set the **Time** to **0.75** and check the **Looping** box. 
 
+1. Connect the **Return Value** pin of the **Set Timer by Function Name** node to the **Handle** pin of the **Clear Timer by Handle** node.
+
 	![](GHT6B_3.png)
 
 	The **Time** entered is the interval in which the **Function Name** Function will be called (every 0.75 seconds). 
 
-1.	Create a new **Int** variable called **PlayerHealth**, click **Compile**, then set the **Default Value** to **100**. 
+1.	Create a new **Integer** variable called **PlayerHealth**, click **Compile**, then set the **Default Value** to **100**. 
 
 	![](GHT6B_4.png)
 
@@ -57,7 +60,7 @@ _For this example, we are using the **Blueprint Third Person Template** and will
 
 	We are not displaying Health Bars in this example; however this will show us that the Health Value is indeed decreasing.
 
-1.	Open the **Blueprint_Effect_Fire** Blueprint, then click the **Components** tab. 
+1.	Open the **Blueprint_Effect_Fire** Blueprint located in the Content/StarterContent/Blueprints folder. 
 
 1.	Click the **Add Component** button and add a **Box** shape and call it **Trigger**. 
 
@@ -65,7 +68,7 @@ _For this example, we are using the **Blueprint Third Person Template** and will
 
 	When the player enters/exits this Trigger, we will call the Events in the **ThirdPersonCharacter** Blueprint. 
 
-1.	Return to the **Graph** tab and **Right-click** on the **Trigger** and add the **OnComponentBegin** and **OnComponentEnd** Overlap Events. 
+1. Locate the **My Blueprint** tab and **Right-click** on the **Trigger** and add the **OnComponentBegin** and **OnComponentEnd** Overlap Events. 
 
 	![](GHT6B_8.png)
 
@@ -88,28 +91,28 @@ This is an example of how you could use timers to create a damage-over-time effe
 
 ## Additional Timer Related Nodes
 
-While the **Set Timer** and **Clear Timer** nodes are used to start/stop a Timer, the nodes below may also be of use based on your needs. 
+While the **Set Timer by Function Name** and **Clear Timer by Handle** nodes are used to start/stop a Timer, the nodes below may also be of use based on your needs. 
 
 ### Get Timer Elapsed and Get Timer Remaining
 
 ![](GHT6B_10.png)
 
-The **Get Timer Elapsed** and **Get Timer Remaining** nodes will get the amount of time that has elapsed since starting or get the amount of time remaining for the specified Timer (indicated by the FunctionName). It returns the value as a Float and can be used for several different purposes, for example, changing an effect over time or providing a message when an effect is about to expire. 
+The **Get Timer Elapsed by Handle** and **Get Timer Remaining by Handle** nodes will get the amount of time that has elapsed since starting or get the amount of time remaining for the specified Timer (indicated by the Handle). It returns the value as a Float and can be used for several different purposes, for example, changing an effect over time or providing a message when an effect is about to expire. 
 
 ### Pause Timer and Un Pause Timer
 
 ![](GHT6B_11.png)
 
-As the names suggest, the **Pause Timer** pauses the timer at its current time whereas the **Clear Timer** stops it entirely and resets it. The **Un Pause Timer** node can be used to resume a paused timer and pick up where it left off. One example use case would be if a player has an effect applied to them that runs on a timer and they enter a menu or inventory screen for example and you want to pause the effect and resume it when they return to gameplay. 
+As the names suggest, the **Pause Timer by Handle** pauses the timer at its current time whereas the **Clear Timer by Handle** stops it entirely and resets it. The **Unpause Timer by Handle** node can be used to resume a paused timer and pick up where it left off. One example use case would be if a player has an effect applied to them that runs on a timer and they enter a menu or inventory screen for example and you want to pause the effect and resume it when they return to gameplay. 
 
 ### Timer Exists, Is Timer Active and Is Timer Paused
 
 ![](GHT6B_12.png)
 
-The **Timer Exists**, **Is Timer Active**, and **Is Timer Paused** nodes all return a **True** or **False** value and can be used to determine if a Timer exists, if a Timer is currently active or if a Timer is currently paused respectively. 
+The **Does Timer Exist by Handle**, **Is Timer Active by Handle**, and **Is Timer Paused by Handle** nodes all return a **True** or **False** value and can be used to determine if a Timer exists, if a Timer is currently active or if a Timer is currently paused respectively. 
 
-### Set Timer Delegate
+### Set Timer by Event
 
 ![](GHT6B_13.png)
 
-The **Set Timer Delegate** node functions similar to the Set Timer node except instead of specifying a Function Name, you can connect a **Custom Event** which will be fired. You can set this to looping as well, just as you could with the Set Timer node. To clear it, you can use the Clear Timer node and in the Function Name, enter the name of your Custom Event.
+The **Set Timer by Event** node functions similar to the** Set Timer by Function Name** node except instead of specifying a Function Name, you can connect a **Custom Event** which will be fired. You can set this to looping as well, just as you could with the **Set Timer by Function Name** node.
