#UE4doc #docs:
Added engine version field

diff --git a/Engine/Documentation/Source/Gameplay/HowTo/UsingCameras/UsingCamerasOverview.INT.udn b/Engine/Documentation/Source/Gameplay/HowTo/UsingCameras/UsingCamerasOverview.INT.udn
--- a/Engine/Documentation/Source/Gameplay/HowTo/UsingCameras/UsingCamerasOverview.INT.udn
+++ b/Engine/Documentation/Source/Gameplay/HowTo/UsingCameras/UsingCamerasOverview.INT.udn
@@ -2,6 +2,7 @@ Availability: Public
 Title:Using Cameras
 Crumbs: %ROOT%, Gameplay
 Description:A How To Guide for Finding Actors in your Scenes in Unreal Engine 4. 
+version: 4.9
 
 This "How To" series covers working with and placing **Cameras** in Unreal Engine 4. A Camera can be used by itself and placed directly into a level or it can be a part of a Blueprint (for instance to provide a perspective for the player when flying a plane, driving a car, or controlling a character). 
 
