#UE4doc #docs:
Updated engine version field

diff --git a/Engine/Documentation/Source/Gameplay/HowTo/UsingCameras/Blueprints/UsingCameras_BP.INT.udn b/Engine/Documentation/Source/Gameplay/HowTo/UsingCameras/Blueprints/UsingCameras_BP.INT.udn
--- a/Engine/Documentation/Source/Gameplay/HowTo/UsingCameras/Blueprints/UsingCameras_BP.INT.udn
+++ b/Engine/Documentation/Source/Gameplay/HowTo/UsingCameras/Blueprints/UsingCameras_BP.INT.udn
@@ -7,7 +7,7 @@ Related: Gameplay\Framework\Camera
 Related: Engine/Matinee
 Related: Engine/Blueprints
 SkillLevel: Beginner
-Version: 4.8
+Version: 4.9
 
 In this How-to we take a look at creating a static (or fixed) camera angle that is used for a player's perspective during gameplay. Once completing this tutorial, you can take the process used here and apply it to your own game to set up a fixed perspective for a player to view your game from. 
 
