#UE4doc #docs:
Added engine version field

diff --git a/Engine/Documentation/Source/Gameplay/Localization/Localization.INT.udn b/Engine/Documentation/Source/Gameplay/Localization/Localization.INT.udn
--- a/Engine/Documentation/Source/Gameplay/Localization/Localization.INT.udn
+++ b/Engine/Documentation/Source/Gameplay/Localization/Localization.INT.udn
@@ -3,6 +3,7 @@ Title: Localization
 Crumbs:%ROOT%, Gameplay
 Description:Information about how to localize your project.
 Related: Programming/UnrealArchitecture/StringHandling/FText
+version: 4.9
 
 [OBJECT:TopicButtonList]
 	[PARAM:icon]
