#UE4doc #docs:
Added engine version field

diff --git a/Engine/Documentation/Source/Gameplay/Localization/Commandlets/Commandlets.INT.udn b/Engine/Documentation/Source/Gameplay/Localization/Commandlets/Commandlets.INT.udn
--- a/Engine/Documentation/Source/Gameplay/Localization/Commandlets/Commandlets.INT.udn
+++ b/Engine/Documentation/Source/Gameplay/Localization/Commandlets/Commandlets.INT.udn
@@ -2,6 +2,7 @@ Availability:Public
 Title: Commandlets
 Crumbs:%ROOT%, Gameplay, Gameplay/Localization
 Description:Descriptions and arguments for the commandlets used when localizing your project. 
+version: 4.9
 
 
 [TOC (start:2)]
