#UE4doc #docs:
Added engine version field

diff --git a/Engine/Documentation/Source/Gameplay/Localization/Overview/Overview.INT.udn b/Engine/Documentation/Source/Gameplay/Localization/Overview/Overview.INT.udn
--- a/Engine/Documentation/Source/Gameplay/Localization/Overview/Overview.INT.udn
+++ b/Engine/Documentation/Source/Gameplay/Localization/Overview/Overview.INT.udn
@@ -2,6 +2,7 @@ Availability:Public
 Title:Localization System Overview
 Crumbs:%ROOT%, Gameplay, Gameplay/Localization
 Description:An overview of the elements involved in localizing your project. 
+version: 4.9
 
 Localization of a project involves a few key concepts:
 
