#UE4doc #docs:
Added engine version field
Replaced a tip region with a note region
Replaced an outdated image
Replaced a term as the name had been changed in the editor

diff --git a/Engine/Documentation/Source/Gameplay/Localization/Setup/Setup.INT.udn b/Engine/Documentation/Source/Gameplay/Localization/Setup/Setup.INT.udn
--- a/Engine/Documentation/Source/Gameplay/Localization/Setup/Setup.INT.udn
+++ b/Engine/Documentation/Source/Gameplay/Localization/Setup/Setup.INT.udn
@@ -2,6 +2,7 @@ Availability:Public
 Title: Setup
 Crumbs:%ROOT%, Gameplay, Gameplay/Localization
 Description:How to set up a project for localization. 
+version: 4.9
 
 [TOC (start:2)]
 
@@ -16,9 +17,9 @@ Conventionally, localization configuration scripts reside under `Config/Localiza
 ![](SetupFlowchart.png)
 [/REGION:raw]
 
-[REGION:tip]
-Note: Setting up and iterating on localization is a manual process currently, but a dashboard UI for managing localization is in progress that will greatly simplify the workflow.
-[/REGION:tip]
+[REGION:note]
+Setting up and iterating on localization is a manual process currently, but a dashboard UI for managing localization is in progress that will greatly simplify the workflow.
+[/REGION]
 
 
 
@@ -145,7 +146,7 @@ An example directory structure using a project called "MyProject", with two targ
 
 ## Packaging Settings
 In order to package the project properly, the supported cultures' localization data need to be specified for packaging. Within the editor, under the **File** menu, open the **Package Project** submenu, and 
-select **Packaging Settings...**. From the **Project Settings** window, under the **Packaging** category, expand the advanced section in order to access the **Cultures to Package** setting. You may check or uncheck 
+select **Packaging Settings...**. From the **Project Settings** window, under the **Packaging** category, expand the advanced section in order to access the **Localizations to Package** setting. You may check or uncheck 
 which cultures should have their localization data packaged. English is checked by default.
 
 
