#UE4doc #docs:
Added engine version field

diff --git a/Engine/Documentation/Source/Gameplay/Networking/Networking.INT.udn b/Engine/Documentation/Source/Gameplay/Networking/Networking.INT.udn
--- a/Engine/Documentation/Source/Gameplay/Networking/Networking.INT.udn
+++ b/Engine/Documentation/Source/Gameplay/Networking/Networking.INT.udn
@@ -2,6 +2,7 @@ Availability:Public
 Title: Networking and Multiplayer
 Crumbs:%ROOT%, Gameplay, Gameplay/Framework
 Description:Setting up networked games for multiplayer.
+version: 4.9
 
 [OBJECT:TopicButtonList]
 	[PARAM:icon]
