#UE4doc #docs:
Added engine version field
Bolded a term
Fixed a capitalization error
Reworded some hyperlinks

diff --git a/Engine/Documentation/Source/Gameplay/Networking/Actors/Actors.INT.udn b/Engine/Documentation/Source/Gameplay/Networking/Actors/Actors.INT.udn
--- a/Engine/Documentation/Source/Gameplay/Networking/Actors/Actors.INT.udn
+++ b/Engine/Documentation/Source/Gameplay/Networking/Actors/Actors.INT.udn
@@ -4,8 +4,9 @@ Crumbs:%ROOT%, Gameplay, Gameplay/Networking
 Description:The various aspects of replicating Actor objects.
 Related: Gameplay/Networking/Actors/Properties
 Related: Gameplay/Networking/Actors/RPCs
+version: 4.9
 
-As mentioned in the networking overview, actors are the main workhorse for replication. The server will maintain a list of actors, and will update the client periodically so that the client will maintain a close approximation of each actor (that is marked to be replicated).
+As mentioned in the networking overview, **Actors** are the main workhorse for replication. The server will maintain a list of actors, and will update the client periodically so that the client will maintain a close approximation of each actor (that is marked to be replicated).
 
 Actors are updated in two main ways:
 * Property updates
@@ -15,8 +16,8 @@ The main difference between property replication and RPCs is that properties are
 
 A good example of a type of property that should be replicated could be an actor's health. This is something you will likely want the client to know about any time it changes. When the health value isn't changing, nothing will be sent. This is pretty efficient. But keep in mind, that even though this property isn't changing (so it isn't consuming any bandwidth), there is still CPU overhead in determining whether or not this value has changed. So this is well suited to properties that change often.
 
-For more information about properties, see the [detailed property replication documentation](Gameplay/Networking/Actors/Properties).
+For more information about properties, see the [](Gameplay/Networking/Actors/Properties) documentation.
 
 An example of an RPC might be an explosion that you want each client to see at a certain location. You could declare an RPC that takes a position and radius as parameters, and call that each time an explosion occurs. You could store this as a group of properties that replicate to communicate this to the client, but this is a little less efficient, since explosions might not occur often enough to justify them being properties.
 
-For more information about RPCs, see the [detailed RPC replication documentation](Gameplay/Networking/Actors/RPCs).
+For more information about RPCs, see the [](Gameplay/Networking/Actors/RPCs) documentation.
