#UE4doc #docs:
Added engine version field
Edited an example code snippet to follow modern syntax

diff --git a/Engine/Documentation/Source/Gameplay/Networking/Actors/Components/Components.INT.udn b/Engine/Documentation/Source/Gameplay/Networking/Actors/Components/Components.INT.udn
--- a/Engine/Documentation/Source/Gameplay/Networking/Actors/Components/Components.INT.udn
+++ b/Engine/Documentation/Source/Gameplay/Networking/Actors/Components/Components.INT.udn
@@ -5,6 +5,7 @@ Description:Setting up components for replication across the network
 Related: Gameplay/Networking/Actors
 Related: Gameplay/Networking/Actors/Properties
 Related: Gameplay/Networking/Actors/ReplicationPerformance
+version: 4.9
 
 [TOC (start:2 end:3)]
 
@@ -30,12 +31,11 @@ Setting up properties and RPCs on components works exactly the same as an Actor.
 To make a component replicate, simply call `AActorComponent::SetIsReplicated(true)`. If your component is a default subobject, this should be done in the class constructor after spawning the component. 
 
 Example:
-	ACharacter::ACharacter(const FObjectInitializer& ObjectInitializer)
-	: Super(ObjectInitializer)
+	ACharacter::ACharacter()
 	{
 		// Etc...
 
-		CharacterMovement = ObjectInitializer.CreateDefaultSubobject<UMovementComp_Character>(this, TEXT("CharMoveComp"));
+		CharacterMovement = CreateDefaultSubobject<UMovementComp_Character>(TEXT("CharMoveComp"));
 		if (CharacterMovement)
 		{
 			CharacterMovement->UpdatedComponent = CapsuleComponent;
