#UE4doc #docs:
Added engine version field

diff --git a/Engine/Documentation/Source/Gameplay/Networking/Actors/OwningConnections/OwningConnections.INT.udn b/Engine/Documentation/Source/Gameplay/Networking/Actors/OwningConnections/OwningConnections.INT.udn
--- a/Engine/Documentation/Source/Gameplay/Networking/Actors/OwningConnections/OwningConnections.INT.udn
+++ b/Engine/Documentation/Source/Gameplay/Networking/Actors/OwningConnections/OwningConnections.INT.udn
@@ -5,6 +5,7 @@ Description:An overview of the role of the server in multiplayer.
 Related: Gameplay/Networking/Actors
 Related: Gameplay/Networking/Actors/RPCs
 Related: Gameplay/Networking/Server
+version: 4.9
 
 You will hear a lot about the owning connection of an actor. Ultimately, each connection has a PlayerController, created specifically for that connection (read more about this process in the client connection flow details). Each PlayerController that is created for this reason, is owned by that connection. To determine if an actor in general is owned by a connection, you query for the actors most outer owner, and if the owner is a PlayerController, then that actor is also owned by the same connection that owns the PlayerController.
 
