#UE4doc #docs:
Added engine version field

diff --git a/Engine/Documentation/Source/Gameplay/Networking/Actors/Properties/Conditions/Conditions.INT.udn b/Engine/Documentation/Source/Gameplay/Networking/Actors/Properties/Conditions/Conditions.INT.udn
--- a/Engine/Documentation/Source/Gameplay/Networking/Actors/Properties/Conditions/Conditions.INT.udn
+++ b/Engine/Documentation/Source/Gameplay/Networking/Actors/Properties/Conditions/Conditions.INT.udn
@@ -2,6 +2,7 @@ Availability:Public
 Title:Conditional Property Replication
 Crumbs:%ROOT%, Gameplay, Gameplay/Networking, Gameplay/Networking/Actors, Gameplay/Networking/Actors/Properties
 Description:Detailed information about how to conditionally replicate Actor properties.
+version: 4.9
 
 Once a property is registered for replication, you can't unregister it (that's where the lifetime part comes from). The reason for this is because we bake in as much information as possible, so we can take advantage of sharing work across many connections for the same set of properties. This saves a lot of computation time.
 
