#UE4doc #docs:
Added engine version field
Fixed capitalization issues

diff --git a/Engine/Documentation/Source/Gameplay/Networking/Actors/Properties/ObjectReferences/ObjectReferences.INT.udn b/Engine/Documentation/Source/Gameplay/Networking/Actors/Properties/ObjectReferences/ObjectReferences.INT.udn
--- a/Engine/Documentation/Source/Gameplay/Networking/Actors/Properties/ObjectReferences/ObjectReferences.INT.udn
+++ b/Engine/Documentation/Source/Gameplay/Networking/Actors/Properties/ObjectReferences/ObjectReferences.INT.udn
@@ -4,8 +4,9 @@ Crumbs:%ROOT%, Gameplay, Gameplay/Networking, Gameplay/Networking/Actors, Gamepl
 Description:Detailed information about how object reference properties are replicated.
 Related: Gameplay/Networking/Actors/Properties/Conditions
 Related: Gameplay/Networking/Actors/RPCs
+version: 4.9
 
-Generally, object references are handled automatically in the UE4 multiplayer framework. What this means is, if you have a UObject property that is replicated, the reference to that object will be sent over the network as a special ID that is assigned by the server. This special id is a FNetworkGUID. The server is in charge of assigning this id, and then notifying all connected clients of this assignment.
+Generally, object references are handled automatically in the UE4 multiplayer framework. What this means is, if you have a UObject property that is replicated, the reference to that object will be sent over the network as a special ID that is assigned by the server. This special ID is a FNetworkGUID. The server is in charge of assigning this ID, and then notifying all connected clients of this assignment.
 
 To replicate object references, you simply mark a UObject property as replicated, like this:
 
