#UE4doc #docs:
Added engine version field

diff --git a/Engine/Documentation/Source/Gameplay/Networking/Actors/Relevancy/Relevancy.INT.udn b/Engine/Documentation/Source/Gameplay/Networking/Actors/Relevancy/Relevancy.INT.udn
--- a/Engine/Documentation/Source/Gameplay/Networking/Actors/Relevancy/Relevancy.INT.udn
+++ b/Engine/Documentation/Source/Gameplay/Networking/Actors/Relevancy/Relevancy.INT.udn
@@ -4,6 +4,7 @@ Crumbs:%ROOT%, Gameplay, Gameplay/Networking, Gameplay/Networking/Actors
 Description:Information about Actor relevancy and priority
 Related: Gameplay/Networking/Actors
 Related: Gameplay/Networking/Actors/ReplicationPerformance
+version: 4.9
 
 [TOC (start:2 end:3)]
 
