#UE4doc #docs:
Added engine version field

diff --git a/Engine/Documentation/Source/Gameplay/Networking/Actors/ReplicationPerformance/ReplicationPerformance.INT.udn b/Engine/Documentation/Source/Gameplay/Networking/Actors/ReplicationPerformance/ReplicationPerformance.INT.udn
--- a/Engine/Documentation/Source/Gameplay/Networking/Actors/ReplicationPerformance/ReplicationPerformance.INT.udn
+++ b/Engine/Documentation/Source/Gameplay/Networking/Actors/ReplicationPerformance/ReplicationPerformance.INT.udn
@@ -5,6 +5,7 @@ Description:Some tips for optimizing the performance and bandwidth usage of Acto
 Related: Gameplay/Networking/Actors
 Related: Gameplay/Networking/Actors/Relevancy
 Related: Gameplay/Networking/Actors/ReplicationFlow
+version: 4.9
 
 Replicating actors can take time. The engine does its best to make this as efficient as possible, but there are a few things you can do to make this job easier.
 
