#UE4doc #docs:
Resubmitting as not all changes were saved

diff --git a/Engine/Documentation/Source/Gameplay/Networking/Actors/RPCs/RPCs.INT.udn b/Engine/Documentation/Source/Gameplay/Networking/Actors/RPCs/RPCs.INT.udn
--- a/Engine/Documentation/Source/Gameplay/Networking/Actors/RPCs/RPCs.INT.udn
+++ b/Engine/Documentation/Source/Gameplay/Networking/Actors/RPCs/RPCs.INT.udn
@@ -91,7 +91,7 @@ Custom events can however be marked as replicated from within the Blueprint edit
 
 To use this feature, create a new custom event in your event graph. Click on the custom event and edit the Replication settings in the details view:
 
-![](RPC_BP.png)(w:400)
+![](RPC_BP.png)
 
 ## Validation
 
