#UE4doc #docs:
Added engine version field
Fixed some capitalization issues
Changed font settings for multiple terms

diff --git a/Engine/Documentation/Source/Gameplay/Networking/Blueprints/Blueprints.INT.udn b/Engine/Documentation/Source/Gameplay/Networking/Blueprints/Blueprints.INT.udn
--- a/Engine/Documentation/Source/Gameplay/Networking/Blueprints/Blueprints.INT.udn
+++ b/Engine/Documentation/Source/Gameplay/Networking/Blueprints/Blueprints.INT.udn
@@ -4,6 +4,7 @@ Crumbs:%ROOT%, Gameplay, Gameplay/Networking
 Description:The various aspects of multiplayer applied to Blueprints.
 Related: Resources/ContentExamples/Networking
 Related: Gameplay/Networking/Actors
+version: 4.9
 
 [TOC (start:2 end:3)]
 
@@ -13,20 +14,20 @@ Unreal Engine 4 provides a lot of multiplayer functionality out of the box, and
 
 To add multiplayer functionality to your game, it's important to understand the roles of the major gameplay classes that are provided by the engine and how they work together - and especially, how they work in a multiplayer context:
 
-* `GameInstance`
-* `GameMode`
-* `GameState`
-* `Pawn` (and `Character`, which inherits from Pawn)
-* `PlayerController`
-* `PlayerState`
+* GameInstance
+* GameMode
+* GameState
+* Pawn (and Character, which inherits from Pawn)
+* PlayerController
+* PlayerState
 
- See the [framework documentation](Gameplay/Framework) for more information - but at least keep in mind the following tips when designing your game for multiplayer are:
+ See the [](Gameplay/Framework) documentation for more information - but at least keep in mind the following tips when designing your game for multiplayer are:
 
-* The `GameMode` object only exists on the server.
-* The `GameState` exists on the server and the clients, so the server can use replicated variables on the `GameState` to keep all clients up-to-date on data about the game.
-* A `PlayerController` exists on the server for every player connected to the game. On clients, only the `PlayerController`s for players owned by that client exist. This means that `PlayerController`s are not the best place to store replicated properties that all connected clients are interested in. Instead, use the `PlayerState`.
-* A `PlayerState` will exist for every player connected to the game on both the server and the clients. This class can be used for replicated properties that all clients, not just the owning client, are interested in, such as the individual player's current score.
-* `Pawns` (including `Characters`) will also exist on the server and on all clients, and can also contain replicated variables and events. The decision of whether to use the `PlayerState` or the `Pawn` for a certain variable or event will depend on the situation, but the main thing to keep in mind is that the `PlayerState` will persist for the entire time the player is connected, but a `Pawn` may not. For example, if a player character dies during gameplay, the Pawn he was controlling may be destroyed and a new one created when the player respawns.
+* The GameMode object only exists on the server.
+* The GameState exists on the server and the clients, so the server can use replicated variables on the GameState to keep all clients up-to-date on data about the game.
+* A PlayerController exists on the server for every player connected to the game. On clients, only the PlayerControllers for players owned by that client exist. This means that PlayerControllers are not the best place to store replicated properties that all connected clients are interested in. Instead, use the PlayerState.
+* A PlayerState will exist for every player connected to the game on both the server and the clients. This class can be used for replicated properties that all clients, not just the owning client, are interested in, such as the individual player's current score.
+* Pawns (including Characters) will also exist on the server and on all clients, and can also contain replicated variables and events. The decision of whether to use the PlayerState or the Pawn for a certain variable or event will depend on the situation, but the main thing to keep in mind is that the PlayerState will persist for the entire time the player is connected, but a `Pawn` may not. For example, if a player character dies during gameplay, the Pawn he was controlling may be destroyed and a new one created when the player respawns.
 
 ## Actor replication
 
@@ -34,17 +35,17 @@ To add multiplayer functionality to your game, it's important to understand the
 
 The core of the networking technology in UE4 is actor replication. An actor with its "Replicates" flag set to true will automatically be synchronized from the server to clients who are connected to that server. An important point to understand is that actors are only replicated from the server to the clients - it's not possible to have an actor replicate from a client to the server. Of course, clients still need to be able to send data to the server, and they do this through replicated "Run on server" events.
 
-See this [how-to guide](Gameplay/HowTo/Networking/ReplicateActor) for a walkthrough of a concrete example, as well as the [full actor replication documentation](Gameplay/Networking/Actors).
+See this [](Gameplay/HowTo/Networking/ReplicateActor) guide for a walkthrough of a concrete example, as well as the [](Gameplay/Networking/Actors) documentation.
 
 ### Authority
 
-For every actor in the world, one of the connected players is considered to have authority over that actor. For every actor that exists on the server, the server has authority over that actor - including all replicated actors. As a result, when the "Has Authority" function is run on a client, and the target is an actor that was replicated to him, it will return false. You can also use the "Switch Has Authority" convenience macro as a quick way to branch for different server and client behavior in replicated actors.
+For every actor in the world, one of the connected players is considered to have authority over that actor. For every actor that exists on the server, the server has authority over that actor - including all replicated actors. As a result, when the **Has Authority** function is run on a client, and the target is an actor that was replicated to him, it will return false. You can also use the **Switch Has Authority** convenience macro as a quick way to branch for different server and client behavior in replicated actors.
 
 ![](switch_has_authority.png)
 
 ### Variables
 
-In the details panel of variables on your actors, there is a "Replication" drop-down that lets you control how your variables are replicated, if at all.
+In the details panel of variables on your actors, there is a **Replication** drop-down that lets you control how your variables are replicated, if at all.
 
 ![](variable_replication.png)
 
@@ -53,13 +54,13 @@ In the details panel of variables on your actors, there is a "Replication" drop-
 | ------ | ----------- |
 | **None** | This is the default for new variables, and means the value will not be sent over the network to clients. |
 | **Replicated** | When the server replicates this actor, it will send this variable to clients. The value of the variable on the receiving client will update automatically, so that the next time it's accessed, it will reflect what the value was on the server. Of course, when playing over a real-world network, the updates will be delayed by an amount of time dependent on the network's latency. Remember that replicated variables only go in one direction, from the server to the client! To send data from a client to the server, see the "Events" section. |
-| **RepNotify** | The variable will be replicated as in the "replicated" option, but in addition, an OnRep_<variable name> function will be created in your blueprint. This function will be called by the engine automatically on the client and the server whenever the value of this variable changes. You're free to implement this function however you'd like, as needed by your game. ![](onrep.png) |
+| **RepNotify** | The variable will be replicated as in the Replicated option, but in addition, an OnRep_<variable name> function will be created in your blueprint. This function will be called by the engine automatically on the client and the server whenever the value of this variable changes. You're free to implement this function however you'd like, as needed by your game. ![](onrep.png) |
 
 [region:note]
 Many of the variables in the engine's built-in classes already have replication enabled, so that many features work automatically in a multiplayer context.
 [/region]
 
-See this [how-to guide](Gameplay/HowTo/Networking/ReplicateVariable) for a walkthrough of a concrete example of variable replication, as well as the [full property replication documentation](Gameplay/Networking/Actors/Properties).
+See this [](Gameplay/HowTo/Networking/ReplicateVariable) guide for a walkthrough of a concrete example of variable replication, as well as the [](Gameplay/Networking/Actors/Properties) documentation.
 
 ### Spawning and destroying
 
@@ -71,18 +72,18 @@ The situation is similar for destroying replicated actors: if the server destroy
 
 In Blueprints, in addition to replicating actors and their variables, you can also run events across the clients and the server.
 
-See this [how-to guide](Gameplay/HowTo/Networking/ReplicateFunction) for a walkthrough of a concrete example, as well as the [full RPC documentation](Gameplay/Networking/Actors/RPCs).
+See this [](Gameplay/HowTo/Networking/ReplicateFunction) guide for a walkthrough of a concrete example, as well as the [](Gameplay/Networking/Actors/RPCs) documentation.
 
 [region:note]
-You may also see the term RPC (for "remote procedure call"). If so, just be aware that replicated events in Blueprints essentially compile down to RPCs inside the engine - and this is what they're usually called in C++.
+You may also see the term RPC (Remote Procedure Call). If so, just be aware that replicated events in Blueprints essentially compile down to RPCs inside the engine - and this is what they're usually called in C++.
 [/region]
 
 ### Ownership
 
 An important concept to understand when working on multiplayer, and especially with replicated events, is which connection is [considered to be the owner of a particular actor or component](Gameplay/Networking/Actors/OwningConnections). For our purposes, know that "Run on server" events can only be invoked from actors (or their components) which the client owns. Usually, this means you can only send "Run on server" events from the following actors, or from a component of one of the actors:
-* The client's `PlayerController` itself,
-* A Pawn that the client's `PlayerController` possesses, or
-* The client's `PlayerState`.
+* The client's PlayerController itself,
+* A Pawn that the client's PlayerController possesses, or
+* The client's PlayerState.
 
 Likewise, for a server sending "Run on owning client" events, those events should also be invoked on one of these actors. Otherwise, the server won't know which client to send the event to, and it will only run on the server!
 
@@ -119,11 +120,11 @@ If the event is invoked from a client, given the target in the left-hand column,
 | **Unowned target**						| Invoking client	| Invoking client	| Dropped		| Invoking client		|
 
 
-As you can see from the table above, any events that are invoked from a client and that are not set to "Run on Server" are treated as if they are not replicated.
+As you can see from the table above, any events that are invoked from a client and that are not set to Run on Server are treated as if they are not replicated.
 
 Sending a replicated event from the client to the server is the only way to communicate information from a client to the server, since general actor replication is designed to be server-to-client only.
 
-Also, note that multicast events can only be sent from the server. Because of Unreal's client-server model, a client isn't directly connected to any of the other clients, he's only connected to the server. Therefore, a client is unable to send a multicast event directly to the other clients, and must only communicate with the server. You can emulate this behavior, however, by using two replicated events: one "Run on server" event, and one "Multicast" event. The implementation of the "Run on server" event can perform validation, if desired, and then call the multicast event. The implementation of the multicast event would perform the logic that you'd like to run for all connected players. As an example that doesn't perform any validation at all, see the following image:
+Also, note that multicast events can only be sent from the server. Because of Unreal's client-server model, a client isn't directly connected to any of the other clients, he's only connected to the server. Therefore, a client is unable to send a multicast event directly to the other clients, and must only communicate with the server. You can emulate this behavior, however, by using two replicated events: one Run on server event, and one Multicast event. The implementation of the Run on server event can perform validation, if desired, and then call the multicast event. The implementation of the multicast event would perform the logic that you'd like to run for all connected players. As an example that doesn't perform any validation at all, see the following image:
 
 ![](forward_multicast.png)
 
@@ -133,8 +134,8 @@ One thing to keep in mind when using replicated events to communicate game state
 
 #### Reliability
 
-For any replicated event, you can choose whether it is **reliable** or **unreliable**.
+For any replicated event, you can choose whether it is **Reliable** or **Unreliable**.
 
-**Reliable** events are guaranteed to reach their destination (assuming the ownership rules above are followed), but they introduce more bandwidth, and potentially latency, to meet this guarantee. Try to avoid sending reliable events too often, such as on every tick, since the engine's internal buffer of reliable events may overflow - when this happens, the associated player will be disconnected!
+Reliable events are guaranteed to reach their destination (assuming the ownership rules above are followed), but they introduce more bandwidth, and potentially latency, to meet this guarantee. Try to avoid sending reliable events too often, such as on every tick, since the engine's internal buffer of reliable events may overflow - when this happens, the associated player will be disconnected!
 
-**Unreliable** events work as their name implies - they may not reach their destination, in case of packet loss on the network, or if the engine determines there is a lot of higher-priority traffic it needs to send, for example. As a result, unreliable events use less bandwidth than reliable events, and they can be called safely more often.
+Unreliable events work as their name implies - they may not reach their destination, in case of packet loss on the network, or if the engine determines there is a lot of higher-priority traffic it needs to send, for example. As a result, unreliable events use less bandwidth than reliable events, and they can be called safely more often.
