#UE4doc #docs:
Added engine version field
Changed font settings for multiple terms

diff --git a/Engine/Documentation/Source/Gameplay/Networking/CharacterMovementComponent/CharacterMovementComponent.INT.udn b/Engine/Documentation/Source/Gameplay/Networking/CharacterMovementComponent/CharacterMovementComponent.INT.udn
--- a/Engine/Documentation/Source/Gameplay/Networking/CharacterMovementComponent/CharacterMovementComponent.INT.udn
+++ b/Engine/Documentation/Source/Gameplay/Networking/CharacterMovementComponent/CharacterMovementComponent.INT.udn
@@ -3,6 +3,7 @@ Title: Character Movement Component
 Crumbs:%ROOT%, Gameplay, Gameplay/Networking, Gameplay/Networking/Actors
 Description:Detailed explanation of Character Movement Component
 Related: Gameplay/Networking/Actors
+version: 4.9
 
 [TOC (start:2 end:3)]
 
@@ -44,10 +45,10 @@ Another helpful way to diagnose issues is to turn on some logging of the data se
 
 ##Advanced Topic: Adding New Movement Abilities to Character Movement
 
-When adding new movement abilities to a character, you have many options on how to accomplish this. Let's step through a theoretical series of attempts to add a "Teleport" ability to a character that works in a networked game. The basic idea is that when a player presses the "T" key, they will teleport forward 10 meters, if the destination is free of obstruction.
+When adding new movement abilities to a character, you have many options on how to accomplish this. Let's step through a theoretical series of attempts to add a "Teleport" ability to a character that works in a networked game. The basic idea is that when a player presses the **T** key, they will teleport forward 10 meters, if the destination is free of obstruction.
 
 ###Approach 1: Execute only on Client
-In this approach, we didn't think about networking at all and just teleport forward when pressing T.
+In this approach, we didn't think about networking at all and just teleport forward when pressing **T**.
 
 **Result**: Fails in networked games. The local client will appear to teleport forward but will then quickly warp back to their starting location.
 
