#UE4doc #docs:
Added engine version field

diff --git a/Engine/Documentation/Source/Gameplay/Networking/Example/Example.INT.udn b/Engine/Documentation/Source/Gameplay/Networking/Example/Example.INT.udn
--- a/Engine/Documentation/Source/Gameplay/Networking/Example/Example.INT.udn
+++ b/Engine/Documentation/Source/Gameplay/Networking/Example/Example.INT.udn
@@ -4,6 +4,7 @@ Crumbs:%ROOT%, Gameplay, Gameplay/Framework
 Description:Setting up networked games for multiplayer.
 Related: Gameplay/Networking/Blueprints
 Related: Gameplay/Networking/Actors
+version: 4.9
 
 To give a concrete example, consider the case where you are a client in a network game. You see two opponents running towards you, shooting at you, and you hear their shots. Since all of the game 
 state is being maintained on the server rather than on your machine, why can you see and hear those things happening?
