#UE4doc #docs:
Removed Overview header
Added engine version field
Changed font settings for multiple terms

diff --git a/Engine/Documentation/Source/Gameplay/Networking/Server/Server.INT.udn b/Engine/Documentation/Source/Gameplay/Networking/Server/Server.INT.udn
--- a/Engine/Documentation/Source/Gameplay/Networking/Server/Server.INT.udn
+++ b/Engine/Documentation/Source/Gameplay/Networking/Server/Server.INT.udn
@@ -6,10 +6,10 @@ Related: Gameplay/Networking/Actors
 Related: Gameplay/Networking/Actors/Properties
 Related: Gameplay/Networking/Actors/RPCs
 Related: Gameplay/Framework
+version: 4.9
 
 [TOC (start:2 end:3)]
 
-## Overview
 
 UE4 multiplayer is based around the client-server model. This means that there will be a single server that will be authoritative over game state, while connected clients will maintain a close approximation.
 
@@ -36,7 +36,7 @@ The server is in charge is driving the flow of gameplay. It is the server's job
 
 The [major framework portions](Gameplay/Framework) are mostly out of the scope of this document, but we can talk about some of the additional multiplayer nuances introduced when dealing working with certain classes.
 
-Gameplay state and flow are generally driven through the `GameMode` actor. Only the server contains a valid copy of this actor (the client doesn't contain a copy at all). To communicate this state to the client, there is a `GameState actor`, that will shadow important state of the `GameMode` actor. This `GameState` actor is marked to be replicated to each client. Clients will contain an approximate copy of this `GameState` actor, and can use this actor as a reference to know the general state of the game.
+Gameplay state and flow are generally driven through the GameMode actor. Only the server contains a valid copy of this actor (the client doesn't contain a copy at all). To communicate this state to the client, there is a GameState actor, that will shadow important state of the GameMode actor. This GameState actor is marked to be replicated to each client. Clients will contain an approximate copy of this `GameState` actor, and can use this actor as a reference to know the general state of the game.
 
 For a more details explanation of GameMode, please see [the GameMode documentation](Gameplay/Framework/GameMode).
 
@@ -50,10 +50,10 @@ The major steps are:
 1. Client sends a connect request.
 2. If the server accepts, it will send the current map.
 3. The server will wait for the client to load this map.
-4. Once loaded, the server will then locally call `AGameMode::PreLogin`.
-	* This will give the `GameMode` a chance to reject the connection
-5. If accepted, the server will then call `AGameMode::Login`
-	* The role of this function is to create a `PlayerController`, which will then be replicated to the newly connected client. Once received, this `PlayerController` will replace the clients temporary `PlayerController` that was used as a placeholder during the connection process.
-	* Note that `APlayerController::BeginPlay` will be called here. It should be noted that it is NOT yet safe to call RPC functions on this actor. You should wait until AGameMode::PostLogin is called.
-6. Assuming everything went well, `AGameMode::PostLogin` is called.
-	* At this point, it is safe for the server to start calling RPC functions on this `PlayerController`.
+4. Once loaded, the server will then locally call AGameMode::PreLogin.
+	* This will give the GameMode a chance to reject the connection
+5. If accepted, the server will then call AGameMode::Login
+	* The role of this function is to create a PlayerController, which will then be replicated to the newly connected client. Once received, this PlayerController will replace the clients temporary PlayerController that was used as a placeholder during the connection process.
+	* Note that APlayerController::BeginPlay will be called here. It should be noted that it is NOT yet safe to call RPC functions on this actor. You should wait until AGameMode::PostLogin is called.
+6. Assuming everything went well, AGameMode::PostLogin is called.
+	* At this point, it is safe for the server to start calling RPC functions on this PlayerController.
