#UE4doc #docs:
Added engine version field
Added bullets to some lists as they were in that type of format

diff --git a/Engine/Documentation/Source/Gameplay/Networking/Travelling/Travelling.INT.udn b/Engine/Documentation/Source/Gameplay/Networking/Travelling/Travelling.INT.udn
--- a/Engine/Documentation/Source/Gameplay/Networking/Travelling/Travelling.INT.udn
+++ b/Engine/Documentation/Source/Gameplay/Networking/Travelling/Travelling.INT.udn
@@ -4,6 +4,7 @@ Crumbs:%ROOT%, Gameplay, Gameplay/Networking
 Description:An overview of how travelling works in multiplayer.
 Related: Gameplay/Networking/Server
 Related: Gameplay/Framework
+version: 4.9
 
 [TOC (start:2 end:3)]
 
@@ -23,31 +24,30 @@ There are three ways in which a non-seamless travel must occur:
 There are three main function that drive travelling: `UEngine::Browse`, `UWorld::ServerTravel`, and `APlayerController::ClientTravel`. These can be a bit confusing when trying to figure out which one to use, so here are some guidelines that should help:
 
 ### `UEngine::Browse`
-Is like a hard reset when loading a new map.
+* Is like a hard reset when loading a new map.
+* Will always result in a non-seamless travel.
 
-Will always result in a non-seamless travel.
+* Will result in the server disconnecting current clients before travelling to the destination map.
 
-Will result in the server disconnecting current clients before travelling to the destination map.
+* Clients will disconnect from current server.
 
-Clients will disconnect from current server.
-
-Dedicated server cannot travel to other servers, so the map must be local (cannot be URL).
+* Dedicated server cannot travel to other servers, so the map must be local (cannot be URL).
 
 ### `UWorld::ServerTravel`
-For the server only.
+* For the server only.
 
-Will jump the server to a new world/level.
+* Will jump the server to a new world/level.
 
-All connected clients will follow.
+* All connected clients will follow.
 
-This is the way multiplayer games travel from map to map, and the server is the one in charge to call this function.
+* This is the way multiplayer games travel from map to map, and the server is the one in charge to call this function.
 
-The server will call `APlayerController::ClientTravel` for all client players that are connected.
+* The server will call `APlayerController::ClientTravel` for all client players that are connected.
 
 ### `APlayerController::ClientTravel`
-If called from a client, will travel to a new server
+* If called from a client, will travel to a new server
 
-If called from a server, will instruct the particular client to travel to the new map (but stay connected to the current server)
+* If called from a server, will instruct the particular client to travel to the new map (but stay connected to the current server)
 
 ## Enabling Seamless Travel
 
