#UE4doc #docs - Added engine version field.

diff --git a/Engine/Documentation/Source/Resources/ContentExamples/Decals/1_1/1_1.INT.udn b/Engine/Documentation/Source/Resources/ContentExamples/Decals/1_1/1_1.INT.udn
--- a/Engine/Documentation/Source/Resources/ContentExamples/Decals/1_1/1_1.INT.udn
+++ b/Engine/Documentation/Source/Resources/ContentExamples/Decals/1_1/1_1.INT.udn
@@ -10,6 +10,7 @@ Related: Engine/Rendering/Materials/ExpressionReference
 Related: Engine/Rendering/Materials/MaterialInputs
 Related: Engine/Rendering/Materials/LayeredMaterials
 Related: Engine/Rendering/Materials/MaterialInstances
+Version: 4.9
 
 [VAR:Steps]
 [OBJECT:Navigation]
@@ -73,7 +74,7 @@ Related: Engine/Rendering/Materials/MaterialInstances
 
 
 
-Decals are Materials that are projected onto meshes in your level, including Static Meshes and Skeletal Meshes. These meshes can have a Mobility setting of Static or Movable and the Decal will still project on them. Many Decals can be rendered at once without a large performance decrease. Performance decreases with larger screen space size and higher shader instruction count.
+**Decals** are Materials that are projected onto meshes in your level, including Static Meshes and Skeletal Meshes. These meshes can have a Mobility setting of Static or Movable and the Decal will still project on them. Many Decals can be rendered at once without a large performance decrease. Performance decreases with larger screen space size and higher shader instruction count.
 
 
 You can add the Deferred Decal Actor directly from the **All Classes** tab of the **Modes** panel. You then need to make a new Material and change the **Material Domain** property to **Deferred Decal** in the Material's **Details** panel. You can then apply the Material to the Decal Actor.
