#UE4doc #docs - Added engine version field. Updated image call out color. Added additional Decal Blend Modes.

diff --git a/Engine/Documentation/Source/Resources/ContentExamples/Decals/1_2/1_2.INT.udn b/Engine/Documentation/Source/Resources/ContentExamples/Decals/1_2/1_2.INT.udn
--- a/Engine/Documentation/Source/Resources/ContentExamples/Decals/1_2/1_2.INT.udn
+++ b/Engine/Documentation/Source/Resources/ContentExamples/Decals/1_2/1_2.INT.udn
@@ -10,6 +10,7 @@ Related: Engine/Rendering/Materials/ExpressionReference
 Related: Engine/Rendering/Materials/MaterialInputs
 Related: Engine/Rendering/Materials/LayeredMaterials
 Related: Engine/Rendering/Materials/MaterialInstances
+Version: 4.9
 
 [VAR:Steps]
 [OBJECT:Navigation]
@@ -74,16 +75,26 @@ Related: Engine/Rendering/Materials/MaterialInstances
 
 
 
-This example demonstrates the available Blend Modes that Decals can use. Different Blend Modes change the way the Decal affects the surface. This is set using the **Decal Blend Mode** property found within a Material.
+This example demonstrates the available **Blend Modes** that Decals can use. Different Blend Modes change the way the Decal affects the surface. This is set using the **Decal Blend Mode** property found within a Material.
 
 ![](DecalBlendMode.png)
 
-There are 4 blend modes that decals can use:
+There are 11 blend modes that decals can use:
 
 * **Translucent** - Can use Diffuse, Metallic, Specular, Roughness, Emissive, Opacity, and Normal. 
 * **Stain** - Is a modulate type blend with Diffuse and Opacity.
 * **Normal** - Uses the Opacity and Normal channels and only affects the Normal map layer it is projecting on.
 * **Emissive** - Uses Emissive and Opacity only.
+* **Volumetric Distance Function** - Use the output of signed distance in Opacity depending on Light Vector.
+
+The DBuffer decals can be used with lighting. These are not on by default and must be enabled in the Project Settings > Rendering section.
+* **DBuffer Translucent Color, Normal, Roughness** - This is non-metallic and will use the Color, Opacity, Roughness, and Normal to work with baked lighting.
+* **DBuffer Translucent Color** - This is non-metallic and will use only the Color and Opacity to work with baked lighting.
+* **DBuffer Translucent Color, Normal** - This is non-metallic and will use the Color and Normal to work with baked lighting.
+* **DBuffer Translucent Color, Roughness** - This is non-metallic and will use the Color and Roughness to work with baked lighting.
+* **DBuffer Translucent Normal** - This will only use the Opacity and Normal channels to work with baked lighting.
+* **DBuffer Translucent Normal, Roughness** - This will only use the Roughness, Opacity, and Normal to work with baked lighting.
+* **DBuffer Translucent Roughness** - This will only use the Roughness and Opacity to work with baked lighting. 
 
 
 [/REGION]
