#UE4doc #docs - Added Engine version field.

diff --git a/Engine/Documentation/Source/Resources/ContentExamples/Decals/1_3/1_3.INT.udn b/Engine/Documentation/Source/Resources/ContentExamples/Decals/1_3/1_3.INT.udn
--- a/Engine/Documentation/Source/Resources/ContentExamples/Decals/1_3/1_3.INT.udn
+++ b/Engine/Documentation/Source/Resources/ContentExamples/Decals/1_3/1_3.INT.udn
@@ -10,6 +10,7 @@ Related: Engine/Rendering/Materials/ExpressionReference
 Related: Engine/Rendering/Materials/MaterialInputs
 Related: Engine/Rendering/Materials/LayeredMaterials
 Related: Engine/Rendering/Materials/MaterialInstances
+Version: 4.9
 
 [VAR:Steps]
 [OBJECT:Navigation]
diff --git a/Engine/Documentation/Source/Resources/ContentExamples/Decals/1_4/1_4.INT.udn b/Engine/Documentation/Source/Resources/ContentExamples/Decals/1_4/1_4.INT.udn
--- a/Engine/Documentation/Source/Resources/ContentExamples/Decals/1_4/1_4.INT.udn
+++ b/Engine/Documentation/Source/Resources/ContentExamples/Decals/1_4/1_4.INT.udn
@@ -10,6 +10,7 @@ Related: Engine/Rendering/Materials/ExpressionReference
 Related: Engine/Rendering/Materials/MaterialInputs
 Related: Engine/Rendering/Materials/LayeredMaterials
 Related: Engine/Rendering/Materials/MaterialInstances
+Version: 4.9
 
 [VAR:Steps]
 [OBJECT:Navigation]
