#UE4doc #docs - Added engine version field. Updated image call out color.

diff --git a/Engine/Documentation/Source/Resources/ContentExamples/Decals/1_5/1_5.INT.udn b/Engine/Documentation/Source/Resources/ContentExamples/Decals/1_5/1_5.INT.udn
--- a/Engine/Documentation/Source/Resources/ContentExamples/Decals/1_5/1_5.INT.udn
+++ b/Engine/Documentation/Source/Resources/ContentExamples/Decals/1_5/1_5.INT.udn
@@ -10,6 +10,7 @@ Related: Engine/Rendering/Materials/ExpressionReference
 Related: Engine/Rendering/Materials/MaterialInputs
 Related: Engine/Rendering/Materials/LayeredMaterials
 Related: Engine/Rendering/Materials/MaterialInstances
+Version: 4.9
 
 [VAR:Steps]
 [OBJECT:Navigation]
diff --git a/Engine/Documentation/Source/Resources/ContentExamples/Decals/1_6/1_6.INT.udn b/Engine/Documentation/Source/Resources/ContentExamples/Decals/1_6/1_6.INT.udn
--- a/Engine/Documentation/Source/Resources/ContentExamples/Decals/1_6/1_6.INT.udn
+++ b/Engine/Documentation/Source/Resources/ContentExamples/Decals/1_6/1_6.INT.udn
@@ -10,6 +10,7 @@ Related: Engine/Rendering/Materials/ExpressionReference
 Related: Engine/Rendering/Materials/MaterialInputs
 Related: Engine/Rendering/Materials/LayeredMaterials
 Related: Engine/Rendering/Materials/MaterialInstances
+Version: 4.9
 
 [VAR:Steps]
 [OBJECT:Navigation]
