#UE4doc #docs - Added engine version field. Commented out ToC for single item.

diff --git a/Engine/Documentation/Source/Resources/ContentExamples/Decals/1_6/1_6.INT.udn b/Engine/Documentation/Source/Resources/ContentExamples/Decals/1_6/1_6.INT.udn
--- a/Engine/Documentation/Source/Resources/ContentExamples/Decals/1_6/1_6.INT.udn
+++ b/Engine/Documentation/Source/Resources/ContentExamples/Decals/1_6/1_6.INT.udn
@@ -30,7 +30,7 @@ Version: 4.9
 
 %Steps%
 
-[TOC(start:2 end:2)]
+<!---[TOC(start:2 end:2)]--->
 
 [VAR:Topic]
 [OBJECT:Topic]
diff --git a/Engine/Documentation/Source/Resources/ContentExamples/Destructables/Destructables.INT.udn b/Engine/Documentation/Source/Resources/ContentExamples/Destructables/Destructables.INT.udn
--- a/Engine/Documentation/Source/Resources/ContentExamples/Destructables/Destructables.INT.udn
+++ b/Engine/Documentation/Source/Resources/ContentExamples/Destructables/Destructables.INT.udn
@@ -43,12 +43,12 @@ Related:
 [/OBJECT]
 [/VAR]
 
-[TOC(start:2 end:2)]
+<!---[TOC(start:2 end:2)]--->
 
 
 
 
-The Content Examples provided inside the **Destructibles** map showcase how destructible meshes can be created and used in **Unreal Engine 4**. Automatically generated and custom imported destructibles, fracturing controls with **Support** and **Destruction Depth** levels as well as how to fire events upon destruction like spawning particles or playing sounds are demonstrated in this map. 
+The Content Examples provided inside the **Destructibles** map showcase how destructible meshes can be created and used in Unreal Engine 4. Automatically generated and custom imported destructibles, fracturing controls with Support and Destruction Depth levels as well as how to fire events upon destruction like spawning particles or playing sounds are demonstrated in this map. 
 
 
 ## Destructibles Map
