#UE4doc #docs - Added engine version field. Commented out Toc for single item on page.

diff --git a/Engine/Documentation/Source/Resources/ContentExamples/DynamicSceneShadows/DynamicSceneShadows.INT.udn b/Engine/Documentation/Source/Resources/ContentExamples/DynamicSceneShadows/DynamicSceneShadows.INT.udn
--- a/Engine/Documentation/Source/Resources/ContentExamples/DynamicSceneShadows/DynamicSceneShadows.INT.udn
+++ b/Engine/Documentation/Source/Resources/ContentExamples/DynamicSceneShadows/DynamicSceneShadows.INT.udn
@@ -6,6 +6,7 @@ Related: Engine/Rendering/LightingAndShadows/Basics
 Related: Engine/Rendering/LightingAndShadows/LightTypes/Directional
 Related: Engine/Rendering/LightingAndShadows/LightMobility/DynamicLights
 Related: Engine/Rendering/LightingAndShadows/Shadows
+Version: 4.9
 
 
 [VAR:Topic]
diff --git a/Engine/Documentation/Source/Resources/ContentExamples/EffectsGallery/EffectsGallery.INT.udn b/Engine/Documentation/Source/Resources/ContentExamples/EffectsGallery/EffectsGallery.INT.udn
--- a/Engine/Documentation/Source/Resources/ContentExamples/EffectsGallery/EffectsGallery.INT.udn
+++ b/Engine/Documentation/Source/Resources/ContentExamples/EffectsGallery/EffectsGallery.INT.udn
@@ -8,6 +8,7 @@ Related: Engine/Rendering/ParticleSystems/Optimization
 Related: Engine/Rendering/ParticleSystems/Reference/TypeData/GPUSprites
 Related: Engine/Rendering/ParticleSystems/ParticleLights
 Related: Engine/Rendering/ParticleSystems/Reference/Modules/VectorField
+Version: 4.9
 
 
 [REGION:banner]
@@ -45,7 +46,7 @@ Related: Engine/Rendering/ParticleSystems/Reference/Modules/VectorField
 [/OBJECT]
 [/VAR]
 
-[TOC(start:2 end:2)]
+<!---[TOC(start:2 end:2)]--->
 
 
 
@@ -70,21 +71,21 @@ Listed below are the examples provided inside the **Effects** map:
 | Example | What is Demonstrated |
 | ------- | ------------ |
 | **[1.1 CPU and GPU particles](Resources\ContentExamples\EffectsGallery\1_A)** | The difference between CPU and GPU particles along with pros/cons are demonstrated with two simple particle fountains. |
-| **[1.2 GPU particles with Velocity Cone module](Resources\ContentExamples\EffectsGallery\1_B)** | The use of the **Velocity Cone** module which outputs particles in a cone shape. This example also shows how a **Static Mesh** can be added to a particle through the **Mesh TypeData** module (e.g. two emitters output sparks while another outputs a Static Mesh). |
-| **[1.3 GPU particles with Emitter Initial Location](Resources\ContentExamples\EffectsGallery\1_C)** | Three spheres travel around a central point while leaving trails behind them. This example also shows how to spawn a particle from the location of other particles with the **Emitter Initial Location** module.  |
-| **[1.4 GPU particles with Point Gravity](Resources\ContentExamples\EffectsGallery\1_D)** | Shows how to specify a point in space where all particles inside an emitter will be drawn to through the **Point Gravity** module. |
+| **[1.2 GPU particles with Velocity Cone module](Resources\ContentExamples\EffectsGallery\1_B)** | The use of the Velocity Cone module which outputs particles in a cone shape. This example also shows how a Static Mesh can be added to a particle through the Mesh TypeData module (e.g. two emitters output sparks while another outputs a Static Mesh). |
+| **[1.3 GPU particles with Emitter Initial Location](Resources\ContentExamples\EffectsGallery\1_C)** | Three spheres travel around a central point while leaving trails behind them. This example also shows how to spawn a particle from the location of other particles with the Emitter Initial Location module.  |
+| **[1.4 GPU particles with Point Gravity](Resources\ContentExamples\EffectsGallery\1_D)** | Shows how to specify a point in space where all particles inside an emitter will be drawn to through the Point Gravity module. |
 | **[1.5 GPU particles with collision](Resources\ContentExamples\EffectsGallery\1_E)** | Demonstrates particles that collide with the environment and have collision (e.g. a downpour of rain bouncing off objects and falling to the ground). |
-| **[1.6 Blueprint with random burst of particles, light and sound](Resources\ContentExamples\EffectsGallery\1_F)** | A **Blueprint** is used to combine light, sound and particle effect into a burst effect similar to a short circuit spark. |
-| **[1.7 Particle emission from bone](Resources\ContentExamples\EffectsGallery\1_G)** | Illustrates how to spawn a particle effect attached to the bone of a **Skeletal Mesh** (e.g. specifying a limb of a character that is on fire). |
+| **[1.6 Blueprint with random burst of particles, light and sound](Resources\ContentExamples\EffectsGallery\1_F)** | A Blueprint is used to combine light, sound and particle effect into a burst effect similar to a short circuit spark. |
+| **[1.7 Particle emission from bone](Resources\ContentExamples\EffectsGallery\1_G)** | Illustrates how to spawn a particle effect attached to the bone of a Skeletal Mesh (e.g. specifying a limb of a character that is on fire). |
 | **[1.8 Particle Shadow Casting](Resources\ContentExamples\EffectsGallery\1_H)** | How to create particles that cast shadows on the environment (e.g. puffy clouds that cast shadows). |
-| **[2.1 Local Vector Field](Resources\ContentExamples\EffectsGallery\2_A)** | This example shows how to control the movement, acceleration and velocity of GPU sprites through **Vector Fields** by using a module contained inside the emitter itself. |
-| **[2.2 Global Vector Field](Resources\ContentExamples\EffectsGallery\2_B)** | Similar to a **Local Vector Field** except, instead of a module inside the emitter, this is an **Actor** that can be placed in a level that you can edit as desired. The GPU particles will only be affected by the forces of the vector field when they are inside the field's volume. |
-| **[2.3 Black body node (dynamic parameter control)](Resources\ContentExamples\EffectsGallery\2_C)** | How to generate realistic color for effects by using **Dynamic Parameter** control (e.g. a fire effect which blends into and billow a black cloud of smoke). |
+| **[2.1 Local Vector Field](Resources\ContentExamples\EffectsGallery\2_A)** | This example shows how to control the movement, acceleration and velocity of GPU sprites through Vector Fields by using a module contained inside the emitter itself. |
+| **[2.2 Global Vector Field](Resources\ContentExamples\EffectsGallery\2_B)** | Similar to a Local Vector Field except, instead of a module inside the emitter, this is an Actor that can be placed in a level that you can edit as desired. The GPU particles will only be affected by the forces of the vector field when they are inside the field's volume. |
+| **[2.3 Black body node (dynamic parameter control)](Resources\ContentExamples\EffectsGallery\2_C)** | How to generate realistic color for effects by using Dynamic Parameter control (e.g. a fire effect which blends into and billow a black cloud of smoke). |
 | **[2.4 Event-driven spawn](Resources\ContentExamples\EffectsGallery\2_D)** | Using events called during the lifepsan of a particle to spawn a different particle (e.g. fireworks where the death of a particle spawns an explosion of additional particles). |
 | **[2.5 Ribbon trails](Resources\ContentExamples\EffectsGallery\2_E)** | How to create a particle that leaves a trail effect behind it as it moves from one location to another (e.g. a flare being shot into the air). |
 | **[2.6 Lit particle reacting to lights](Resources\ContentExamples\EffectsGallery\2_F)** | Showcases how to generate a realistic lighting effect on particles such as smoke (e.g. a police car's lights effecting smoke). |
-| **[2.7 CPU particles emitting light](Resources\ContentExamples\EffectsGallery\2_G)** | Adding a **Light Module** to an emitter and spawning dynamic lights from a particle effect. |
-| **[2.8 Beam emitter with a noise module](Resources\ContentExamples\EffectsGallery\2_H)** | Shows how to generate a particle effect that extends from one point to another using the **Beam** emitter (e.g. could be used for lightning or as a bullet trail effect). |
+| **[2.7 CPU particles emitting light](Resources\ContentExamples\EffectsGallery\2_G)** | Adding a Light Module to an emitter and spawning dynamic lights from a particle effect. |
+| **[2.8 Beam emitter with a noise module](Resources\ContentExamples\EffectsGallery\2_H)** | Shows how to generate a particle effect that extends from one point to another using the Beam emitter (e.g. could be used for lightning or as a bullet trail effect). |
 
 
 
