#UE4doc #docs - Added engine version field.

diff --git a/Engine/Documentation/Source/Resources/ContentExamples/EffectsGallery/1_D/1_D.INT.udn b/Engine/Documentation/Source/Resources/ContentExamples/EffectsGallery/1_D/1_D.INT.udn
--- a/Engine/Documentation/Source/Resources/ContentExamples/EffectsGallery/1_D/1_D.INT.udn
+++ b/Engine/Documentation/Source/Resources/ContentExamples/EffectsGallery/1_D/1_D.INT.udn
@@ -8,6 +8,7 @@ Related: Engine/Rendering/ParticleSystems/Optimization
 Related: Engine/Rendering/ParticleSystems/Reference/TypeData/GPUSprites
 Related: Engine/Rendering/ParticleSystems/ParticleLights
 Related: Engine/Rendering/ParticleSystems/Reference/Modules/VectorField
+Version: 4.9 
 
 [VAR:Steps]
 [OBJECT:Navigation]
diff --git a/Engine/Documentation/Source/Resources/ContentExamples/EffectsGallery/1_E/1_E.INT.udn b/Engine/Documentation/Source/Resources/ContentExamples/EffectsGallery/1_E/1_E.INT.udn
--- a/Engine/Documentation/Source/Resources/ContentExamples/EffectsGallery/1_E/1_E.INT.udn
+++ b/Engine/Documentation/Source/Resources/ContentExamples/EffectsGallery/1_E/1_E.INT.udn
@@ -8,6 +8,7 @@ Related: Engine/Rendering/ParticleSystems/Optimization
 Related: Engine/Rendering/ParticleSystems/Reference/TypeData/GPUSprites
 Related: Engine/Rendering/ParticleSystems/ParticleLights
 Related: Engine/Rendering/ParticleSystems/Reference/Modules/VectorField
+Version: 4.9
 
 [VAR:Steps]
 [OBJECT:Navigation]
