#UE4doc #doc - Added engine version field.

diff --git a/Engine/Documentation/Source/Resources/ContentExamples/EffectsGallery/1_F/1_F.INT.udn b/Engine/Documentation/Source/Resources/ContentExamples/EffectsGallery/1_F/1_F.INT.udn
--- a/Engine/Documentation/Source/Resources/ContentExamples/EffectsGallery/1_F/1_F.INT.udn
+++ b/Engine/Documentation/Source/Resources/ContentExamples/EffectsGallery/1_F/1_F.INT.udn
@@ -8,6 +8,7 @@ Related: Engine/Rendering/ParticleSystems/Optimization
 Related: Engine/Rendering/ParticleSystems/Reference/TypeData/GPUSprites
 Related: Engine/Rendering/ParticleSystems/ParticleLights
 Related: Engine/Rendering/ParticleSystems/Reference/Modules/VectorField
+Version: 4.9
 
 [VAR:Steps]
 [OBJECT:Navigation]
diff --git a/Engine/Documentation/Source/Resources/ContentExamples/EffectsGallery/1_G/1_G.INT.udn b/Engine/Documentation/Source/Resources/ContentExamples/EffectsGallery/1_G/1_G.INT.udn
--- a/Engine/Documentation/Source/Resources/ContentExamples/EffectsGallery/1_G/1_G.INT.udn
+++ b/Engine/Documentation/Source/Resources/ContentExamples/EffectsGallery/1_G/1_G.INT.udn
@@ -8,6 +8,7 @@ Related: Engine/Rendering/ParticleSystems/Optimization
 Related: Engine/Rendering/ParticleSystems/Reference/TypeData/GPUSprites
 Related: Engine/Rendering/ParticleSystems/ParticleLights
 Related: Engine/Rendering/ParticleSystems/Reference/Modules/VectorField
+Version: 4.9
 
 [VAR:Steps]
 [OBJECT:Navigation]
