#UE4doc #docs - Added engine version field. Updated image call out color.

diff --git a/Engine/Documentation/Source/Resources/ContentExamples/EffectsGallery/1_H/1_H.INT.udn b/Engine/Documentation/Source/Resources/ContentExamples/EffectsGallery/1_H/1_H.INT.udn
--- a/Engine/Documentation/Source/Resources/ContentExamples/EffectsGallery/1_H/1_H.INT.udn
+++ b/Engine/Documentation/Source/Resources/ContentExamples/EffectsGallery/1_H/1_H.INT.udn
@@ -9,6 +9,7 @@ Related: Engine/Rendering/ParticleSystems/Optimization
 Related: Engine/Rendering/ParticleSystems/Reference/TypeData/GPUSprites
 Related: Engine/Rendering/ParticleSystems/ParticleLights
 Related: Engine/Rendering/ParticleSystems/Reference/Modules/VectorField
+Version: 4.9
 
 [VAR:Steps]
 [OBJECT:Navigation]
