#UE4doc #docs - Added engine version field.

diff --git a/Engine/Documentation/Source/Resources/ContentExamples/EffectsGallery/2_A/2_A.INT.udn b/Engine/Documentation/Source/Resources/ContentExamples/EffectsGallery/2_A/2_A.INT.udn
--- a/Engine/Documentation/Source/Resources/ContentExamples/EffectsGallery/2_A/2_A.INT.udn
+++ b/Engine/Documentation/Source/Resources/ContentExamples/EffectsGallery/2_A/2_A.INT.udn
@@ -9,6 +9,7 @@ Related: Engine/Rendering/ParticleSystems/Optimization
 Related: Engine/Rendering/ParticleSystems/Reference/TypeData/GPUSprites
 Related: Engine/Rendering/ParticleSystems/ParticleLights
 Related: Engine/Rendering/ParticleSystems/Reference/Modules/VectorField
+Version: 4.9
 
 [VAR:Steps]
 [OBJECT:Navigation]
