#UE4doc #docs - Added engine version field.

diff --git a/Engine/Documentation/Source/Resources/ContentExamples/EffectsGallery/2_B/2_B.INT.udn b/Engine/Documentation/Source/Resources/ContentExamples/EffectsGallery/2_B/2_B.INT.udn
--- a/Engine/Documentation/Source/Resources/ContentExamples/EffectsGallery/2_B/2_B.INT.udn
+++ b/Engine/Documentation/Source/Resources/ContentExamples/EffectsGallery/2_B/2_B.INT.udn
@@ -9,6 +9,7 @@ Related: Engine/Rendering/ParticleSystems/Optimization
 Related: Engine/Rendering/ParticleSystems/Reference/TypeData/GPUSprites
 Related: Engine/Rendering/ParticleSystems/ParticleLights
 Related: Engine/Rendering/ParticleSystems/Reference/Modules/VectorField
+Version: 4.9
 
 [VAR:Steps]
 [OBJECT:Navigation]
