#UE4doc #docs - Added engine version field. Removed bolding from term.

diff --git a/Engine/Documentation/Source/Resources/ContentExamples/ExampleLevelOverview/ExampleLevelOverview.INT.udn b/Engine/Documentation/Source/Resources/ContentExamples/ExampleLevelOverview/ExampleLevelOverview.INT.udn
--- a/Engine/Documentation/Source/Resources/ContentExamples/ExampleLevelOverview/ExampleLevelOverview.INT.udn
+++ b/Engine/Documentation/Source/Resources/ContentExamples/ExampleLevelOverview/ExampleLevelOverview.INT.udn
@@ -2,6 +2,7 @@ Availability: Public
 Title:Using Content Examples
 Crumbs:%ROOT%, Resources, Resources/ContentExamples
 Description:An overview of how the Example levels included with Unreal Engine 4 can be used to accelerate your training.
+Version: 4.9
 
 
 [VAR:Topic]
