#UE4doc #docs - Added engine version field. Updated image for FBX import options. Commented out ToC since there is a single section on the page.

diff --git a/Engine/Documentation/Source/Resources/ContentExamples/ImportOptions/ImportOptions.INT.udn b/Engine/Documentation/Source/Resources/ContentExamples/ImportOptions/ImportOptions.INT.udn
--- a/Engine/Documentation/Source/Resources/ContentExamples/ImportOptions/ImportOptions.INT.udn
+++ b/Engine/Documentation/Source/Resources/ContentExamples/ImportOptions/ImportOptions.INT.udn
@@ -11,6 +11,7 @@ Related: Engine/Content/FBX
 Related: Engine/Content/FBX/MorphTargets
 Related: Shared/Editor/FbxErrors
 Related: Engine/Content/FBX/Materials
+Version: 4.9
 
 
 
@@ -49,7 +50,7 @@ Related: Engine/Content/FBX/Materials
 [/OBJECT]
 [/VAR]
 
-[TOC(start:2 end:2)]
+<!---[TOC(start:2 end:2)]--->
 
 
 
