#UE4doc #docs - Added engine version field. Corrected term.

diff --git a/Engine/Documentation/Source/Resources/ContentExamples/ImportOptions/1_3/1_3.INT.udn b/Engine/Documentation/Source/Resources/ContentExamples/ImportOptions/1_3/1_3.INT.udn
--- a/Engine/Documentation/Source/Resources/ContentExamples/ImportOptions/1_3/1_3.INT.udn
+++ b/Engine/Documentation/Source/Resources/ContentExamples/ImportOptions/1_3/1_3.INT.udn
@@ -11,6 +11,7 @@ Related: Engine/Content/FBX
 Related: Engine/Content/FBX/MorphTargets
 Related: Shared/Editor/FbxErrors
 Related: Engine/Content/FBX/Materials
+Version: 4.9
 
 [VAR:Steps]
 [OBJECT:Navigation]
@@ -75,7 +76,7 @@ Related: Engine/Content/FBX/Materials
 ![](CE_1_3_001.png)
 [/REGION]
 
-Unreal Engine 4 can either take in normals only from an FBX file, normals and tangents from an FBX file, or it can simply re-calculate them on its own. The image above shows the difference in **calculated normals** on the **left**, and **imported normals and tangents** on the **right**.
+Unreal Engine 4 can either take in normals only from an FBX file, normals and tangents from an FBX file, or it can simply re-calculate them on its own. The image above shows the difference in **computed normals** on the **left**, and **imported normals and tangents** on the **right**.
 
 When importing an FBX file, you will find this feature near the top of the import dialog under **Mesh**. The pull down will have these 3 options:
 ![](ImportPulldown.png)
