#UE4doc #docs - Added engine version field.

diff --git a/Engine/Documentation/Source/Resources/ContentExamples/ImportOptions/1_4/1_4.INT.udn b/Engine/Documentation/Source/Resources/ContentExamples/ImportOptions/1_4/1_4.INT.udn
--- a/Engine/Documentation/Source/Resources/ContentExamples/ImportOptions/1_4/1_4.INT.udn
+++ b/Engine/Documentation/Source/Resources/ContentExamples/ImportOptions/1_4/1_4.INT.udn
@@ -11,6 +11,7 @@ Related: Engine/Content/FBX
 Related: Engine/Content/FBX/MorphTargets
 Related: Shared/Editor/FbxErrors
 Related: Engine/Content/FBX/Materials
+Version: 4.9
 
 [VAR:Steps]
 [OBJECT:Navigation]
diff --git a/Engine/Documentation/Source/Resources/ContentExamples/ImportOptions/1_5/1_5.INT.udn b/Engine/Documentation/Source/Resources/ContentExamples/ImportOptions/1_5/1_5.INT.udn
--- a/Engine/Documentation/Source/Resources/ContentExamples/ImportOptions/1_5/1_5.INT.udn
+++ b/Engine/Documentation/Source/Resources/ContentExamples/ImportOptions/1_5/1_5.INT.udn
@@ -11,6 +11,7 @@ Related: Engine/Content/FBX
 Related: Engine/Content/FBX/MorphTargets
 Related: Shared/Editor/FbxErrors
 Related: Engine/Content/FBX/Materials
+Version: 4.9
 
 [VAR:Steps]
 [OBJECT:Navigation]
diff --git a/Engine/Documentation/Source/Resources/ContentExamples/ImportOptions/1_6/1_6.INT.udn b/Engine/Documentation/Source/Resources/ContentExamples/ImportOptions/1_6/1_6.INT.udn
--- a/Engine/Documentation/Source/Resources/ContentExamples/ImportOptions/1_6/1_6.INT.udn
+++ b/Engine/Documentation/Source/Resources/ContentExamples/ImportOptions/1_6/1_6.INT.udn
@@ -11,6 +11,7 @@ Related: Engine/Content/FBX
 Related: Engine/Content/FBX/MorphTargets
 Related: Shared/Editor/FbxErrors
 Related: Engine/Content/FBX/Materials
+Version: 4.9
 
 [VAR:Steps]
 [OBJECT:Navigation]
