#UE4doc #docs - Added engine version field. Replaced image table with Comparison slider.

diff --git a/Engine/Documentation/Source/Resources/ContentExamples/ImportOptions/1_7/1_7.INT.udn b/Engine/Documentation/Source/Resources/ContentExamples/ImportOptions/1_7/1_7.INT.udn
--- a/Engine/Documentation/Source/Resources/ContentExamples/ImportOptions/1_7/1_7.INT.udn
+++ b/Engine/Documentation/Source/Resources/ContentExamples/ImportOptions/1_7/1_7.INT.udn
@@ -11,6 +11,7 @@ Related: Engine/Content/FBX
 Related: Engine/Content/FBX/MorphTargets
 Related: Shared/Editor/FbxErrors
 Related: Engine/Content/FBX/Materials
+Version: 4.9
 
 [VAR:Steps]
 [OBJECT:Navigation]
@@ -68,9 +69,20 @@ Related: Engine/Content/FBX/Materials
 [/OBJECT]
 [/VAR]
 
+[OBJECT:ComparisonSlider]
+	[PARAM:before]
+	![Cube](CE_1_7_001.png)(h:500)
+	[/PARAM]
+	[PARAM:after]
+	![Sphere](CE_1_7_002.png)(h:500)
+	[/PARAM]
+[/OBJECT]
+
+<!---
 [REGION:lightbox]
 [![](CE_1_7_001.png)(h:330)](CE_1_7_001.png) [![](CE_1_7_002.png)(h:330)](CE_1_7_002.png)
 [/REGION]
+--->
 
 [INCLUDE:Engine/Content\FBX\MorphTargets#FBX_morph]
 
