#UE4doc #docs - Added engine version field. Removed Overview header. Commented out single section ToC.

diff --git a/Engine/Documentation/Source/Resources/ContentExamples/Landscapes/Landscapes.INT.udn b/Engine/Documentation/Source/Resources/ContentExamples/Landscapes/Landscapes.INT.udn
--- a/Engine/Documentation/Source/Resources/ContentExamples/Landscapes/Landscapes.INT.udn
+++ b/Engine/Documentation/Source/Resources/ContentExamples/Landscapes/Landscapes.INT.udn
@@ -9,6 +9,7 @@ Related: Engine/Landscape/Materials
 Related: Engine/Landscape/Editing/Splines
 Related: Engine/Landscape/Editing/Gizmos
 Related: Engine/Landscape/Editing/Brushes
+Version: 4.9
 
 [REGION:banner]
 ![](Landscapes_Header.png)(convert:false)
@@ -45,10 +46,9 @@ Related: Engine/Landscape/Editing/Brushes
 [/OBJECT]
 [/VAR]
 
-[TOC(start:2 end:2)]
+<!---[TOC(start:2 end:2)]--->
 
 
-## Overview
 
 There are three examples provided inside the **Landscape** Content Example map which showcase how mountains, hills, and valleys can be created using the **Landscape** tools. The second and third examples demonstrate the use of two additional tools, the **Spline** tool (used to create pathways) and the **Foliage** tool (used to quickly scatter and place Static Meshes like trees). 
 
