#UE4doc #docs:
Added engine version field
Removed Table of Contents section

diff --git a/Engine/Documentation/Source/Gameplay/SaveGame/SaveGame.INT.udn b/Engine/Documentation/Source/Gameplay/SaveGame/SaveGame.INT.udn
--- a/Engine/Documentation/Source/Gameplay/SaveGame/SaveGame.INT.udn
+++ b/Engine/Documentation/Source/Gameplay/SaveGame/SaveGame.INT.udn
@@ -2,8 +2,8 @@ Availability:Public
 Title: Saving Your Game
 Crumbs:%ROOT%, Gameplay
 Description:Overview of how to save and load your game
+version: 4.9
 
-[TOC (start:2 end:2)]
 
 A simple way to set up saving and loading in your game is to use a SaveGame object. 
 
