#UE4 #DOCS physics new files

diff --git a/Engine/Documentation/Source/Engine/Physics/Collision/HowTo/AddCollisionPreset/AddCollisionPreset.INT.udn b/Engine/Documentation/Source/Engine/Physics/Collision/HowTo/AddCollisionPreset/AddCollisionPreset.INT.udn
new file mode 100644
--- /dev/null
+++ b/Engine/Documentation/Source/Engine/Physics/Collision/HowTo/AddCollisionPreset/AddCollisionPreset.INT.udn
@@ -0,0 +1,31 @@
+Availability: Public
+Title:Static Mesh Collision Reference
+Crumbs: %ROOT%, Engine, Engine/Physics, Engine/Physics/Collision
+Description:Content guide to creating and setting up collision geometry.
+
+[VAR:Topic]
+[OBJECT:Topic]
+	[PARAM:image]
+		![%Engine/Physics/Collision/Reference:title%](Engine/Physics/Collision/Reference/physics_topic.png)
+	[/PARAM]
+	[PARAM:icon]
+		![](%ROOT%\rendering_icon.png)(convert:false)
+	[/PARAM]
+	[PARAM:title]
+		%Engine/Physics/Collision:/Reference:title%
+	[/PARAM]
+	[PARAM:description]
+		%Engine/Physics/Collision/Reference:description%
+	[/PARAM]
+	[PARAM:path]
+		[RELATIVE:Engine/Physics/Collision/Reference]
+	[/PARAM]
+[/OBJECT]
+[/VAR]
+
+
+
+
+
+
+
diff --git a/Engine/Documentation/Source/Engine/Physics/Collision/HowTo/AddConvexHulls/AddConvexHulls.INT.udn b/Engine/Documentation/Source/Engine/Physics/Collision/HowTo/AddConvexHulls/AddConvexHulls.INT.udn
new file mode 100644
--- /dev/null
+++ b/Engine/Documentation/Source/Engine/Physics/Collision/HowTo/AddConvexHulls/AddConvexHulls.INT.udn
@@ -0,0 +1,52 @@
+Availability: Public
+Title:Add a Collision Hull to a Static Mesh Using the Auto Convex Collision Tool
+Crumbs: %ROOT%, Engine, Engine/Physics, Engine/Physics/Collision
+Description:
+
+[VAR:Topic]
+[OBJECT:Topic]
+	[PARAM:image]
+		![%Engine/Physics/Collision/Reference:title%](Engine/Physics/Collision/Reference/physics_topic.png)
+	[/PARAM]
+	[PARAM:icon]
+		![](%ROOT%/rendering_icon.png)(convert:false)
+	[/PARAM]
+	[PARAM:title]
+		%Engine/Physics/Collision:/Reference:title%
+	[/PARAM]
+	[PARAM:description]
+		%Engine/Physics/Collision/Reference:description%
+	[/PARAM]
+	[PARAM:path]
+		[RELATIVE:Engine/Physics/Collision/Reference]
+	[/PARAM]
+[/OBJECT]
+[/VAR]
+
+
+
+
+## Steps
+
+[REGION:raw]
+![](Engine/Physics/Collision/HowTo/COLREF_collisionMenu_Auto.png)
+[/REGION]
+
+The **Static Mesh Editor** has the capability to generate collision using a number of convex hulls.
+
+![](Engine/Physics/Collision/HowTo/Collision_CreateCollisionHullWithAutoConvex_0.png)
+
+This rounded cube would probably work just fine with **6DOP Simplified Collision** but let us assume it needs more than that.
+
+Selecting **Collision Menu -> Auto Convex Collision** will bring up this panel in the lower right hand side:
+
+![](Engine/Physics/Collision/HowTo/Collision_CreateCollisionHullWithAutoConvex_1.png)
+
+* **Max Hulls** - Determines the number of hulls to create to best match the shape of the object.
+* **Max Hull Verts** - Maximum number of collision hull vertices. Increasing this value increase how complex the collision hulls are allowed to be.
+* **Apply** - Adds collision to the object based on the settings above.
+* **Defaults** - Resets the values of **Max Hulls** and **Max Hull Verts** to the values seen in the image above.
+
+![](Engine/Physics/Collision/HowTo/Collision_CreateCollisionHullWithAutoConvex_2.png)
+
+As you can see, the auto convex collision is reasonably accurate.
diff --git a/Engine/Documentation/Source/Engine/Physics/Collision/HowTo/AddCustomCollisionType/AddCustomCollisionType.INT.udn b/Engine/Documentation/Source/Engine/Physics/Collision/HowTo/AddCustomCollisionType/AddCustomCollisionType.INT.udn
new file mode 100644
--- /dev/null
+++ b/Engine/Documentation/Source/Engine/Physics/Collision/HowTo/AddCustomCollisionType/AddCustomCollisionType.INT.udn
@@ -0,0 +1,55 @@
+Availability: Public
+Title:Add a Custom Object Type to Your Project
+Crumbs: %ROOT%
+Description:How To add a custom physics object channel to Unreal Engine 4.
+
+[VAR:Topic]
+[OBJECT:Topic]
+	[PARAM:image]
+		![%Engine/Physics/Collision/Reference:title%](Engine/Physics/Collision/Reference/physics_topic.png)
+	[/PARAM]
+	[PARAM:icon]
+		![](%ROOT%\rendering_icon.png)(convert:false)
+	[/PARAM]
+	[PARAM:title]
+		%Engine/Physics/Collision:/Reference:title%
+	[/PARAM]
+	[PARAM:description]
+		%Engine/Physics/Collision/Reference:description%
+	[/PARAM]
+	[PARAM:path]
+		[RELATIVE:Engine/Physics/Collision/Reference]
+	[/PARAM]
+[/OBJECT]
+[/VAR]
+
+
+
+There will be times where 6 Object Response Channels and 2 Trace Response Channels simply are not granular enough for what you want to create. This is what the Collision Editor in your Project Settings is for. You can access them from the **Edit Menu** -> **Project Settings** -> **Collision**:
+
+![](Engine/Physics/Collision/COL_projectSettings_1.png)
+
+[REGION:lightbox]
+[![](Engine/Physics/Collision/COL_projectSettings_2.png)(w:940)](Engine/Physics/Collision/COL_projectSettings_2.png)
+[/REGION]
+
+From here you can add new Object Response Channels and Trace Response Channels. Click the **New Object Channel...** or **New Trace Channel...** button, provide a name, select a **Default Response**, and click **Accept**.
+
+![](Engine/Physics/Collision/COL_customChannel.png)
+
+You can have up to 18 Custom Object Response Channels or Custom Trace Response Channels.
+
+### Presets
+
+Custom Presets can also be set here by expanding the **Preset** Category and clicking the **New...** button.
+
+![](Engine/Physics/Collision/COL_customPresets.png)
+
+From here you can name your preset, enable or disable collision, select your preset's Object Type, and finally define the behavior of each Response Channel for the selected Object Type.
+
+
+
+
+
+
+
diff --git a/Engine/Documentation/Source/Engine/Physics/Collision/HowTo/AddCustomTraceType/AddCustomTraceType.INT.udn b/Engine/Documentation/Source/Engine/Physics/Collision/HowTo/AddCustomTraceType/AddCustomTraceType.INT.udn
new file mode 100644
--- /dev/null
+++ b/Engine/Documentation/Source/Engine/Physics/Collision/HowTo/AddCustomTraceType/AddCustomTraceType.INT.udn
@@ -0,0 +1,32 @@
+Availability: Public
+Title:Add a Custom Trace Type to your Project
+Crumbs: %ROOT%, Engine, Engine/Physics, Engine/Physics/Collision
+Description:Content guide to creating and setting up collision geometry.
+
+[VAR:Topic]
+[OBJECT:Topic]
+	[PARAM:image]
+		![%Engine/Physics/Collision/Reference:title%](Engine/Physics/Collision/Reference/physics_topic.png)
+	[/PARAM]
+	[PARAM:icon]
+		![](%ROOT%\rendering_icon.png)(convert:false)
+	[/PARAM]
+	[PARAM:title]
+		%Engine/Physics/Collision:/Reference:title%
+	[/PARAM]
+	[PARAM:description]
+		%Engine/Physics/Collision/Reference:description%
+	[/PARAM]
+	[PARAM:path]
+		[RELATIVE:Engine/Physics/Collision/Reference]
+	[/PARAM]
+[/OBJECT]
+[/VAR]
+
+
+
+
+
+
+
+
diff --git a/Engine/Documentation/Source/Engine/Physics/Collision/HowTo/AddDOP/AddDOP.INT.udn b/Engine/Documentation/Source/Engine/Physics/Collision/HowTo/AddDOP/AddDOP.INT.udn
new file mode 100644
--- /dev/null
+++ b/Engine/Documentation/Source/Engine/Physics/Collision/HowTo/AddDOP/AddDOP.INT.udn
@@ -0,0 +1,43 @@
+Availability: Public
+Title:Add a K-DOP collision hull to a Static Mesh
+Crumbs: %ROOT%, Engine, Engine/Physics, Engine/Physics/Collision
+Description:Content guide to creating and setting up collision geometry.
+
+[VAR:Topic]
+[OBJECT:Topic]
+	[PARAM:image]
+		![%Engine/Physics/Collision/Reference:title%](Engine/Physics/Collision/Reference/physics_topic.png)
+	[/PARAM]
+	[PARAM:icon]
+		![](%ROOT%/rendering_icon.png)(convert:false)
+	[/PARAM]
+	[PARAM:title]
+		%Engine/Physics/Collision:/Reference:title%
+	[/PARAM]
+	[PARAM:description]
+		%Engine/Physics/Collision/Reference:description%
+	[/PARAM]
+	[PARAM:path]
+		[RELATIVE:Engine/Physics/Collision/Reference]
+	[/PARAM]
+[/OBJECT]
+[/VAR]
+
+### Steps
+
+[REGION:raw]
+![](Engine/Physics/Collision/HowTo/COLREF_collisionMenu_KDOP.png)
+[/REGION]
+
+There are a series of menu options in the **Static Mesh Editor** under the **Collision** menu entry, named _##DOP_, these are the **K-DOP** simple collision generators. **K-DOP** is a type of bounding volume, which stands for _K discrete oriented polytope_ (where K is the number of axis aligned planes). Basically it takes K axis-aligned planes and pushes them as close to the mesh as it can. The resulting shape is used as a collision hull. In the **Static Mesh Editor** K can be:
+
+
+* **10** - Box with 4 edges beveled - you can choose X- Y- or Z-aligned edges.
+* **18** - Box with all edges beveled.
+* **26** - Box with all edges and corners beveled.
+
+See below for an example. This tool is quite handy for packages full of pipes, pillars, and railings:
+
+![kdop_sizes.jpg](Engine/Physics/Collision/HowTo/kdop_sizes.jpg)
+
+
diff --git a/Engine/Documentation/Source/Engine/Physics/Collision/HowTo/AddSimpleCollision/AddSimpleCollision.INT.udn b/Engine/Documentation/Source/Engine/Physics/Collision/HowTo/AddSimpleCollision/AddSimpleCollision.INT.udn
new file mode 100644
--- /dev/null
+++ b/Engine/Documentation/Source/Engine/Physics/Collision/HowTo/AddSimpleCollision/AddSimpleCollision.INT.udn
@@ -0,0 +1,48 @@
+Availability: Public
+Title:Add Simple Collision to a Static Mesh
+Crumbs: %ROOT%, Engine, Engine/Physics, Engine/Physics/Collision
+Description:Content guide to creating and setting up collision geometry.
+
+[VAR:Topic]
+[OBJECT:Topic]
+	[PARAM:image]
+		![%Engine/Physics/Collision/Reference:title%](Engine/Physics/Collision/Reference/physics_topic.png)
+	[/PARAM]
+	[PARAM:icon]
+		![](%ROOT%\rendering_icon.png)(convert:false)
+	[/PARAM]
+	[PARAM:title]
+		%Engine/Physics/Collision:/Reference:title%
+	[/PARAM]
+	[PARAM:description]
+		%Engine/Physics/Collision/Reference:description%
+	[/PARAM]
+	[PARAM:path]
+		[RELATIVE:Engine/Physics/Collision/Reference]
+	[/PARAM]
+[/OBJECT]
+[/VAR]
+
+## Steps
+
+[REGION:raw]
+![](Engine\Physics\Collision\HowTo\COLREF_collisionMenu_BuildYourOwn.png)
+[/REGION]
+
+These options create basic collision shapes:
+
+[REGION:raw]
+![](Engine\Physics\Collision\HowTo\COLREF_BuildYourOwn_01.png)(convert:false)
+[/REGION]
+
+Once a simple collision shape is created, it can be moved, rotated, and scaled into the desired position. You can create more than one collision shape for more complex Static Meshes.
+
+| **Control** | **Action** |
+| --- | --- |
+| **W, E, R** | Switch between Translation, Rotation, and Scale tools. |
+| **Delete** | Will delete the currently selected Physics Body or Bodies. |
+| **F** | Will focus on the currently selected Physics Body or Bodies. |
+| **Ctrl + Left Click** | Adds a Physics Body to the selection. |
+| **Alt + Left Drag (On a manipulation widget)** | Duplicates a Physics Body. |
+
+
diff --git a/Engine/Documentation/Source/Engine/Physics/Collision/HowTo/CollisionHowTo.INT.udn b/Engine/Documentation/Source/Engine/Physics/Collision/HowTo/CollisionHowTo.INT.udn
new file mode 100644
--- /dev/null
+++ b/Engine/Documentation/Source/Engine/Physics/Collision/HowTo/CollisionHowTo.INT.udn
@@ -0,0 +1,47 @@
+Availability: Public
+Title:Collision How To's
+Crumbs: %ROOT%, Engine, Engine/Physics, Engine/Physics/Collision
+Description:Content guide to creating and setting up collision geometry.
+
+[VAR:Topic]
+[OBJECT:Topic]
+	[PARAM:image]
+		![%Engine/Physics/Collision/Reference:title%](Engine/Physics/Collision/Reference/physics_topic.png)
+	[/PARAM]
+	[PARAM:icon]
+		![](%ROOT%\rendering_icon.png)(convert:false)
+	[/PARAM]
+	[PARAM:title]
+		%Engine/Physics/Collision:/Reference:title%
+	[/PARAM]
+	[PARAM:description]
+		%Engine/Physics/Collision/Reference:description%
+	[/PARAM]
+	[PARAM:path]
+		[RELATIVE:Engine/Physics/Collision/Reference]
+	[/PARAM]
+[/OBJECT]
+[/VAR]
+
+
+[OBJECT:TopicButtonList]
+	[PARAM:icon]
+		![](%ROOT%/start_icon.png)(convert:false)
+	[/PARAM]
+	[PARAM:title]
+		How To...
+	[/PARAM]
+	[PARAM:links]
+		* [](Engine/Physics/Collision/HowTo/AddConvexHulls "%Engine/Physics/Collision/HowTo/AddConvexHulls:description%")
+		* [](Engine/Physics/Collision/HowTo/AddSimpleCollision "%Engine/Physics/Collision/HowTo/AddSimpleCollision:description%")
+		* [](Engine/Physics/Collision/HowTo/AddDOP "%Engine/Physics/Collision/HowTo/AddDOP:description%")
+		* [](Engine/Physics/Collision/HowTo/AddCustomCollisionType "%Engine/Physics/Collision/HowTo/AddCustomCollisionType:description%")
+		* [](Engine/Physics/Collision/HowTo/AddCustomTraceType "%Engine/Physics/Collision/HowTo/AddCustomTraceType:description%")
+		* [](Engine/Physics/Collision/HowTo/PreviewCollision "%Engine/Physics/Collision/HowTo/PreviewCollision:description%")
+	[/PARAM]
+[/OBJECT]
+
+
+
+
+
diff --git a/Engine/Documentation/Source/Engine/Physics/Collision/HowTo/PreviewCollision/PreviewCollision.INT.udn b/Engine/Documentation/Source/Engine/Physics/Collision/HowTo/PreviewCollision/PreviewCollision.INT.udn
new file mode 100644
--- /dev/null
+++ b/Engine/Documentation/Source/Engine/Physics/Collision/HowTo/PreviewCollision/PreviewCollision.INT.udn
@@ -0,0 +1,46 @@
+Availability: Public
+Title:Review Collision in Your Game
+Crumbs: %ROOT%, Engine, Engine/Physics, Engine/Physics/Collision
+Description:Content guide to creating and setting up collision geometry.
+
+[VAR:Topic]
+[OBJECT:Topic]
+	[PARAM:image]
+		![%Engine/Physics/Collision/Reference:title%](Engine/Physics/Collision/Reference/physics_topic.png)
+	[/PARAM]
+	[PARAM:icon]
+		![](%ROOT%/rendering_icon.png)(convert:false)
+	[/PARAM]
+	[PARAM:title]
+		%Engine/Physics/Collision:/Reference:title%
+	[/PARAM]
+	[PARAM:description]
+		%Engine/Physics/Collision/Reference:description%
+	[/PARAM]
+	[PARAM:path]
+		[RELATIVE:Engine/Physics/Collision/Reference]
+	[/PARAM]
+[/OBJECT]
+[/VAR]
+
+
+
+Here are a couple of console commands you can use while in game to give you information on the collision in your game.
+
+1. Run your game.
+1. Hit the **`** (back tick/tilda key) to bring up the console input box.
+1. Enter one of the following commands:
+
+	* **show COLLISION** - This will draw all collision models and blocking volumes in use in the level. 
+
+		![](Engine/Physics/Collision/HowTo/COLREF_showCollision.png)(w:900)
+
+	* **stat game** - This shows you various useful stats on how long different types of collision are taking.
+
+		![](Engine/Physics/Collision/HowTo/Collision_StatGame.png)(w:900)
+
+
+
+
+
+
diff --git a/Engine/Documentation/Source/Engine/Physics/Collision/Overview/CollisionOverview.INT.udn b/Engine/Documentation/Source/Engine/Physics/Collision/Overview/CollisionOverview.INT.udn
new file mode 100644
--- /dev/null
+++ b/Engine/Documentation/Source/Engine/Physics/Collision/Overview/CollisionOverview.INT.udn
@@ -0,0 +1,133 @@
+Availability: Public
+Title:Collision Overview
+Crumbs: %ROOT%, Engine, Engine/Physics
+Description:An overview of how Collision and Collision Responses operate in Unreal Engine 4.
+
+[VAR:Topic]
+[OBJECT:Topic]
+	[PARAM:image]
+		![%Engine/Physics/Collision:title%](physics_topic.png)
+	[/PARAM]
+	[PARAM:icon]
+		![](%ROOT%\rendering_icon.png)(convert:false)
+	[/PARAM]
+	[PARAM:title]
+		%Engine/Physics/Collision:title%
+	[/PARAM]
+	[PARAM:description]
+		%Engine/Physics/Collision:description%
+	[/PARAM]
+	[PARAM:path]
+		[RELATIVE:Engine/Physics/Collision]
+	[/PARAM]
+[/OBJECT]
+[/VAR]
+
+[TOC(start:2 end:3)]
+
+**Collision Responses** and **Trace Responses** form the basis for how Unreal Engine 4 handles collision and ray casting during run time. Every object that can collide gets an **Object Type** and a series of responses that define how it interacts will all other object types. When a collision or overlap event occurs, both (or all) objects involved can be set to affect or be affected by blocking, overlapping, or ignoring each other.
+
+**Trace Responses** work basically the same way, except the trace (ray cast) itself can be defined as one of the Trace Response types, thusly allowing Actors to either block it or ignore it based on _their_ Trace Responses.
+
+## Interactions
+
+There are a few rules to keep in mind with how collisions are handled:
+
+* **Blocking** will naturally occur between two (or more) Actors set to Block. However, **Simulation Generates Hit Events** needs to be enabled to execute **Event Hit**, which is used in Blueprints, Destructible Actors, Triggers, etc...
+* Setting Actors to **Overlap** will often look like they **Ignore** each other, and without **Generate Overlap Events**, they are essentially the same.
+* For two or more simulating objects to block each other, they both need to be set to block their respective object types.
+* For two or more simulating objects: if one is set to overlap an object, and the second object is set to block the other, the overlap will occur but not the block.
+* Overlap events **can** be generated even if an object **Blocks** another, especially if traveling at high speeds.
+	* It is not recommended for an object to have both collision and overlap events. Though possible, there is much that needs manual handling.
+* If one object is set to ignore and the other is set to overlap, no overlap events will be fired.
+<!--* The **Object Responses** listed are just the defaults and have no actual bearing on what an object _is_. See below for more information on setting up custom [Object Types and Trace Channels](#customobjectchannelsandtracechannels).-->
+
+For purposes of testing levels and looking around your worlds:
+
+* The default **Play In Editor** camera is a pawn. Thusly can be blocked by anything set to block pawns.
+* The **Simulate in Editor** camera, before possessing anything, is **not** a pawn. It can freely clip through everything and will not create any collision or overlap events.
+
+
+## Common Collision Interaction Examples
+
+[REGION:note]
+These interactions assume that all the objects have **Collision Enabled** set to **Collision Enabled** so they are set to fully collide with everything. If collision is disabled, it is as if **ignore** has been set for all **Collision Responses**.
+[/REGION]
+
+For the following section, this will be the setup used to explain what is happening:
+
+![](COL_setup.png)(w:940 convert:false)
+
+The sphere is a **PhysicsBody** and the box is **WorldDynamic**, and by changing their collision settings we can get a number of behaviors.
+
+### Collision
+
+[VAR:spacer]![](%ROOT%/spacer.png)[/VAR]
+
+By setting both of their collision settings to block each other, you get a collision, great for having objects interact with each other:
+
+![](COL_collideNoEvent.png)(w:940 convert:false)
+
+| Sphere Collision Setup| Wall Collision Setup |
+| --- | ---- |
+| ![](COL_collideNoEvent_Sphere.png)(convert:false) |![](COL_collideNoEvent_Box.png)(convert:false) |
+| In this case, the sphere is a **PhysicsBody** and it is set to `block` **WorldDynamic** (which is what the wall is). | The wall is a **WorldDynamic** and is set to `block` **PhysicsBody** Actors (which is what the sphere is). |
+
+In this case, the sphere and the wall will simply collide; no further notifications of the collision will take place.
+
+### Collision and Simulation Generates Hit Events
+
+Just collision is useful and in general, the bare minimum for physics interactions, but if you want something to **report** it has collided so a Blueprint or section of code can be triggered:
+
+![](COL_collideEvent.png)(w:940 convert:false)
+
+| Sphere Collision Setup| Wall Collision Setup |
+| --- | ---- |
+| ![](COL_collideEvent_Sphere.png)(convert:false)|![](COL_collideNoEvent_Box.png)(convert:false)|
+| As in the example above, the sphere is a **PhysicsBody** and it is set to `block` **WorldDynamic** (which is what the wall is). However, the sphere has also enabled **Simulation Generates Hit Events** so it will trigger an event for itself whenever it collides with something. | The wall is a **WorldDynamic** and is set to `block` **PhysicsBody** Actors (which is what the sphere is). Since the wall is not set to **Simulation Generates Hit Events**, it will not generate an event for itself. |
+
+With the sphere set to **Simulation Generates Hit Events**, the sphere will tell itself that it has had a collision. It will fire off events such as **ReceiveHit** or **OnComponentHit** in the sphere's blueprint. Now if the box had an event for collision, it would not fire because it will never notify itself it has happened.
+
+Further, an object that is reporting rigid collisions will report them all and spam reports when it is just sitting on something, so it is best to be careful to filter what it is colliding within its blueprint or in code.
+
+### Overlap and Ignore
+
+For all intents and purposes, **Overlap** and **Ignore** work exactly the same **assuming Generate Overlap Events** is disabled. In this case, the sphere is set to overlap or ignore the box:
+
+![](COL_ignore.png)(w:940 convert:false)
+
+
+
+
+| Sphere Collision Setup| Wall Collision Setup |
+| --- | ---- |
+| ![](COL_OverlapNoEvent_Sphere.png)(convert:false) |![](COL_collideNoEvent_Box.png)(convert:false) |
+| Here the sphere is set to `overlap` **WorldDynamic** Actors (like our wall), but it does not have **Generate Overlap Events** enabled. As far as the sphere is concerned, it has not collided or overlapped anything, effectively it has ignored the wall.  | The wall is a **WorldDynamic** and is set to `block` **PhysicsBody** Actors (which is what the sphere is). As stated above, both Actors need to be set to block each other's respective object types. If they do not, they will not collide. |
+
+Or:
+
+| Sphere Collision Setup| Wall Collision Setup |
+| --- | ---- |
+| ![](COL_ignore_sphere.png)(convert:false) |![](COL_collideNoEvent_Box.png)(convert:false) |
+| Here the sphere is set to `ignore` **WorldDynamic** Actors (like our wall), and it will pass through the wall. | The wall is a **WorldDynamic** and is set to `block` **PhysicsBody** Actors (which is what the sphere is). As stated above, both Actors need to be set to block each other's respective object types. If they do not, they will not collide. |
+
+
+
+### Overlap and Generate Overlap Events
+
+Unlike collisions that can fire every frame, the overlap events are **ReceiveBeginOverlap** and **ReceiveEndOverlap**, which only fire in those specific cases.
+
+![](COL_overlapEvent.png)(w:940 convert:false)
+
+| Sphere Collision Setup| Wall Collision Setup |
+| --- | ---- |
+| ![](COL_OverlapEvent_Sphere.png)(convert:false) | ![](COL_collideNoEvent_Box.png)(convert:false) |
+| Here the sphere is set to `overlap` **WorldDynamic** Actors (like our wall), and it will generate an event for itself when it does overlap something. | The wall is a **WorldDynamic** and is set to `block` **PhysicsBody** Actors (which is what the sphere is). As stated above, both Actors need to be set to block each other's respective object types. If they do not, they will not collide. But, an **Overlap** does occur here, and an event for the sphere is fired. |
+
+ 
+
+
+
+
+
+
diff --git a/Engine/Documentation/Source/Engine/Physics/FrictionRestitutionAndDamping/FrictionRestitutionAndDamping.INT.udn b/Engine/Documentation/Source/Engine/Physics/FrictionRestitutionAndDamping/FrictionRestitutionAndDamping.INT.udn
new file mode 100644
--- /dev/null
+++ b/Engine/Documentation/Source/Engine/Physics/FrictionRestitutionAndDamping/FrictionRestitutionAndDamping.INT.udn
@@ -0,0 +1,94 @@
+Availability: Public
+Title:Physics Damping
+Crumbs: %ROOT%, Engine, Engine/Physics
+Description:Properties and methods used to cause physical objects to resist movement.
+
+
+[TOC(start:2)]
+
+
+
+
+To slow down Physics Bodies, simulate atmospheric drag, or add resistance to a hinge, there are two properties available to Physics Bodies and Physics Constraints: Linear Damping 
+and Angular Damping. Linear Damping controls how much the Physics Body or Constraint resists translation, while Angular Damping controls how much they resist rotating. 
+
+## Physics Bodies
+
+Every Physics Body has properties for **Linear Damping** and **Angular Damping**. These are primarily used to simulate drag, and even a value of 1.0 has a substantial effect on the Physics 
+Body.
+
+[REGION:note]
+For reference, in normal gravity (9.8m/s^2), a Linear Damping value of 30 is enough to stop any single body Actor from falling under the initial pull of gravity, at the start of the 
+simulation. Once more than one Physics Body is interacting (through the use of Physics Constraints or if a collision is occurring), the amount of damping required to stop the Physics 
+Body goes up.
+
+The _minimum_ Linear Damping value to stop a Physics Body that has had force applied to it is around 100.
+
+Finally, with no Angular Damping, a Physics Body will continue to rotate until it is acted upon by an outside force. Even small values can reduce rotation rapidly, while a value of 
+100 will almost immediately stop any angular motion caused by an outside source.
+[/REGION]
+
+## Physics Constraints
+
+The Damping for Physics Constraints work a bit differently than with Physics Bodies: they only work on Limited Motions and when a motion attempts to go beyond its limit. Setting a 
+Physics Constraint's Motion(s) to Limited will display the **Soft Limits**:
+
+[REGION:raw]
+![](Phys_ConstraintDampProps.png)(convert:false)
+[/REGION]
+
+_The two Angular Swing Motions share Damping values, the Angular Twist receives its own Damping values, and all the Linear axes use the same Damping values. While they may have 
+separate properties, the properties all have the same effect on the Motions they relate to._
+
+| Property | Description |
+| -------- | ----------- |
+| **Limit Stiffness** | When a limit is exceeded, this value affects how harshly the constraint attempts to stop the offending motion. The default of 50 will barely affect the motion, a value of 5000 will appear to bounce back into limit, and 50000 will completely deflect the motion. |
+| **Limit Damping** | Like the Damping on Physics Bodies, this defines how much the motion should be slowed. Unlike Damping on Physics Bodies, the damping affect occurs only when the Motion Limit is exceeded, draining momentum until the Physics Body stops. |
+
+The two values work together to create different effects. A high Stiffness and low Damping will look like a hard limit, while having them both high will appear to be move through 
+tar, coming to a stop gently but quickly. A good deal of experimentation is necessary to find the right values for your specific needs.
+
+For more information on Constraints, see: [](Engine/Physics/Constraints/ConstraintsUserGuide)
+
+Or for a reference on the properties available to Physics Constraints,
+ see: [](Engine/Physics\Constraints\ConstraintsReference)
+
+<!--
+## PhysicalMaterial Damping
+This can be thought of like the air resistance that is applied to a body. A force is applied to a body to try and slow it down, and the strength of the force is proportional to the 
+speed that the object is moving. There is a separate setting for linear and angular damping, so you can control how much it will try and resist translation and rotation movement. 
+One place where you might want to use this would be a swinging chandelier that you did not want to swing forever. You may not want to use this for a door though, as a lot of the 
+friction on a real door is due to it's hinge. If you set a high damping on the door's PhysicalMaterial and then broke it off the hinge, it would still have that damping applied to 
+it which may look odd. 
+
+
+
+## Joint Drive
+To simulate 'friction' within a joint, you can use the joint drive properties in the RB_ConstraintInstance (you can get to these properties in PhAT by selecting the joint and 
+pressing the 'I' key or the 'i' icon on the toolbar). First you want to turn on the appropriate bTwistVelocityDrive or bSwingVelocityDrive drive, set the desired angular velocity 
+to zero and modify the AngularDriveForceLimit to control how much friction is in the joint. You can find more information on this on the 
+[PhysicsConstraintReference](Engine/Physics/PhysX/Constraints) page. 
+
+The key thing here is that the damping is applied between two jointed object. Imagine two objects jointed together, but then thrown through space. If you used PhysicalMaterial 
+damping here, it would slow down how quickly the pair of objects would move through the air, when what you really want is to modify how they move relative to each other. 
+That is why Joint Drive is the right thing to use in this case.
+
+
+
+## PhysicsVolume RigidBodyDamping
+Sometimes you want to apply damping or a force to objects only in a certain area of your level. To do this you would use a PhysicsVolume. The ZoneVelocity property of a 
+PhysicsVolume is the desired velocity that objects in that volume should achieve. You can think of this as the speed of a river or wind-tunnel in that area. You can then modify 
+the RigidBodyDamping property to control how strongly a physics body will try and achieve that velocity when in the volume. Setting a zero ZoneVelocity and a high RigidBodyDamping 
+can make the volume act like jelly. Setting a horizontal ZoneVelocity and a moderate RigidBodyDamping can make it act like a flowing river of water.
+-->
+ 
+
+
+
+
+
+
+
+
+
+
diff --git a/Engine/Documentation/Source/Engine/Physics/Overview/PhysicsSystemOverview.INT.udn b/Engine/Documentation/Source/Engine/Physics/Overview/PhysicsSystemOverview.INT.udn
new file mode 100644
--- /dev/null
+++ b/Engine/Documentation/Source/Engine/Physics/Overview/PhysicsSystemOverview.INT.udn
@@ -0,0 +1,66 @@
+Availability: Public
+Title:Physics
+Crumbs: %ROOT%, Engine
+Description:Subsystem used to calculate collision and simulate physical actors.
+
+[VAR:Topic]
+[OBJECT:Topic]
+	[PARAM:image]
+		![%Engine/Physics:title%](Engine/Physics/physics_topic.png)
+	[/PARAM]
+	[PARAM:icon]
+		![](%ROOT%/physics_icon.png)(convert:false)
+	[/PARAM]
+	[PARAM:title]
+		%Engine/Physics:title%
+	[/PARAM]
+	[PARAM:description]
+		%Engine/Physics:description%
+	[/PARAM]
+	[PARAM:path]
+		[RELATIVE:Engine/Physics]
+	[/PARAM]
+[/OBJECT]
+[/VAR]
+[VAR:TopicCompact]
+[OBJECT:TopicCompact]
+	[PARAM:image]
+		![%Engine/Physics:title%](Engine/Physics/physics_topic.png)
+	[/PARAM]
+	[PARAM:icon]
+		![](%ROOT%/physics_icon.png)(convert:false)
+	[/PARAM]
+	[PARAM:title]
+		%Engine/Physics:title%
+	[/PARAM]
+	[PARAM:description]
+		%Engine/Physics:description%
+	[/PARAM]
+	[PARAM:path]
+		[RELATIVE:Engine/Physics]
+	[/PARAM]
+[/OBJECT]
+[/VAR]
+
+
+[REGION:banner]
+![Physics](Engine/Physics/physics_landing_banner.png)
+[/REGION]
+
+[REGION:quote]
+Physics is the natural science that involves the study of matter and its motion through space and time."
+
+-- [Wikipedia](https://en.wikipedia.org/wiki/Physics)
+[/REGION]
+
+[EXCERPT:Intro]
+In Unreal Engine 4, physics are simulated using the PhysX 3.3 physics engine. All physics driven motion (physics bodies that are falling or have a force applied to them) and collisions (physics bodies interacting with each other) are managed by PhysX.Having physics in your game will help improve the immersion value of every scene, as it helps players believe that they are interacting with the scene and that the scene is responding back in some way or another.
+[/EXCERPT:Intro]
+
+## Simulation
+
+
+
+## Collision
+
+
diff --git a/Engine/Documentation/Source/Engine/Physics/PhAT/HowTo/UserGuide.INT.udn b/Engine/Documentation/Source/Engine/Physics/PhAT/HowTo/UserGuide.INT.udn
new file mode 100644
--- /dev/null
+++ b/Engine/Documentation/Source/Engine/Physics/PhAT/HowTo/UserGuide.INT.udn
@@ -0,0 +1,334 @@
+﻿Availability: Public
+Title:Physics Asset Tool How To Guide
+Crumbs: %ROOT%, Engine, Engine/Physics, Engine/Physics/PhAT
+Description:How to use the Physics Asset Tool to create Physics Asset for Skeletal Meshes.
+Version: 4.9
+
+[VAR:Topic]
+[OBJECT:Topic]
+	[PARAM:image]
+		![%Engine/Physics/PhAT/UserGuide:title%](Engine/Physics/physics_topic.png)
+	[/PARAM]
+	[PARAM:icon]
+		![](Engine/Physics/PhAT/HowTo/%ROOT%/physics_icon.png)(convert:false)
+	[/PARAM]
+	[PARAM:title]
+		%Engine/Physics/PhAT/UserGuide:title%
+	[/PARAM]
+	[PARAM:description]
+		%Engine/Physics/PhAT/UserGuide:description%
+	[/PARAM]
+	[PARAM:path]
+		[RELATIVE:Engine/Physics/Phat/UserGuide]
+	[/PARAM]
+[/OBJECT]
+[/VAR]
+[VAR:TopicCompact]
+[OBJECT:TopicCompact]
+	[PARAM:image]
+		![%Engine/Physics/PhAT/UserGuide:title%](Engine/Physics/physics_topic.png)
+	[/PARAM]
+	[PARAM:icon]
+		![](Engine/Physics/PhAT/HowTo/%ROOT%/physics_icon.png)(convert:false)
+	[/PARAM]
+	[PARAM:title]
+		%Engine/Physics/PhAT/UserGuide:title%
+	[/PARAM]
+	[PARAM:description]
+		%Engine/Physics/PhAT/UserGuide:description%
+	[/PARAM]
+	[PARAM:path]
+		[RELATIVE:Engine/Physics/Phat/UserGuide]
+	[/PARAM]
+[/OBJECT]
+[/VAR]
+
+
+
+
+The **Physics Asset Tool** is used for creating **Physics Assets** for use with **Skeletal Meshes**. One can start from nothing and build to a full ragdoll setup or use the automation 
+tools to create a basic set of Physics Bodies and Physics Constraints.
+
+This document will cover step by step procedures for the most common activities in the Physics Asset Tool.
+
+## New Physics Asset
+
+When a Skeletal Mesh is imported, there is an option to generate a Physics Asset for it as it is imported. Once the imported file is processed, you will be presented with the 
+"New Asset" dialog:
+
+[REGION:raw]
+![](Engine/Physics/PhAT/HowTo/PhATUG_ImportPA.png)
+[/REGION]
+
+However, you can follow these steps if you need to create a Physics Asset for a Skeletal Mesh at a later time:
+
+1. Find the Skeletal Mesh asset you wish to add a Physics Asset to in the 
+**Content Browser**.
+1. **Right-click** on the Skeletal Mesh in the Content Browser to bring up the **Context Menu** -> **Create** -> **Create Physics Asset**.
+1. Adjust setting to your liking.
+
+	![](Shared\PhAT\PhAT_restetDefault.png)
+
+1. **Click** 'Ok'.
+
+## Editing a Physics Asset
+
+1. Find the Physics Asset in the Content Browser.
+
+	![](Engine/Physics\PhysicalMaterials\PhysMatUserGuide\PhysMat_PhAT_02.png)
+
+1. **Double-click** to open it in the Physics Asset Tool (PhAT).
+
+	![](Engine/Physics/PhAT/HowTo/PhATUG_FullWindow.png)(w:900)
+
+
+## Physics Bodies
+
+### Editing Physics Bodies
+
+1. Enable Body Editing Mode ![](Shared\Icons\icon_PHatMode_Body_40x.png).
+1. Select a Physics Body in the Viewport or in the **Hierarchy Pane**.
+	* If it is difficult to see the Physics Body you desire, use the **Viewport** -> **Modes** -> **MeshRender Mode** -> **Wireframe/None** to adjust the visibility of the 
+	Skeletal Mesh.
+1. **Move, rotate, and scale** the Physics Body using the Translation, Rotation, and Scale tools so it fits the Skeletal Mesh and Skeletal Joint.
+1. Edit the Physics Body's properties in the Details Pane.
+1. **Enable** or **Disable** Collision between the current Physics Bodies and other in the Physics Asset using the _Collision On_ and _Collision Off_ tools. Any Physics Body that can 
+collide with the currently selected Physics Body will be blue, if they cannot collide, they will be gray.
+1. If a series of bodies should act as one, such as with the twist joints that form a wrist, use the _Weld_ tool to prevent undesired physics simulation behavior. Welded Physics 
+Bodies will be yellow in color.
+1. Save ![](Shared\Icons\icon_SaveAsset_40x.png) often.
+
+See the **[Physics Asset Tool Reference](Engine/Physics/PhAT/Reference)** for more information on the properties in the Physics Asset Tool.
+
+[PUBLISH:licensee]
+[REGION:note]
+As of change list #1977557, multi select is enabled for Physics Bodies and Physics Constraints, enabling you to edit multiple bodies or constraints at the same time.
+[/REGION]
+[/PUBLISH]
+
+### Deleting Physics Bodies
+
+1. **Enable** Body Editing Mode ![](Shared\Icons\icon_PHatMode_Body_40x.png).
+1. **Select** a Physics Body you wish to delete.
+	* If it is difficult to see the Physics Body you desire, use the **Viewport -> Modes -> MeshRender Mode -> Wireframe/None** to adjust the visibility of the Skeletal Mesh.
+1. Press the **Delete** key.
+
+	![](Shared\PhAT\PhAT_DelBone.png)
+
+### Creating New Physics Bodies
+
+![](Engine/Physics/PhAT/HowTo/PhATUG_addBody.png)(w:900)
+
+1. **Right-click** an already existing Physics Body or a Bone in the Hierarchy Panel.
+1. **Select** the type of Physics Body you wish to add from the Context Menu using the left mouse button.
+
+	If the Bone in the Hierarchy Panel does not already have a Physics Body on it, the only option in the **right-click** context menu is "Add Body."
+
+	Finally, you can just **double-click** a bone that does not have a Physics Body in the **Hierarchy** Panel to "Add Body."
+
+
+### Enabling and Disabling Collision Between Physics Bodies
+
+[REGION:simpletable]
+| --- | --- |
+|[INCLUDE:#ImgColOn]| [INCLUDE:#ColOn] |
+|[INCLUDE:#ImgColOff]| [INCLUDE:#ColOff] |
+[/REGION]
+
+<!--
+
+[EXCERPT:ImgColOn]
+[REGION:lightbox]
+[![](Shared\PhAT\PhAT_collisionOn.png)(w:350)](Shared\PhAT\PhAT_collisionOn.png)
+[/REGION]
+[/EXCERPT:ImgColOn]
+
+[EXCERPT:ImgColOff]
+[REGION:lightbox]
+[![](Shared\PhAT\PhAT_collisionOff.png)(w:350)](Shared\PhAT\PhAT_collisionOff.png)
+[/REGION]
+[/EXCERPT:ImgColOff]
+
+[EXCERPT:ColOn]
+[REGION:raw]
+**Enabling Collision**
+
+1. **Select** 2 or more Physics Bodies using **Ctrl + Left Mouse click**.
+1. **Click** the "Collision On" ![](Shared\Icons\icon_PhAT_EnableCollision_40x.png) button.
+
+
+[/REGION]
+[/EXCERPT:ColOn]
+
+[EXCERPT:ColOff]
+**Disabling Collision** 
+
+1. **Select** 2 or more Physics Bodies using **Ctrl + Left Mouse click**.
+1. **Click** the "Collision Off" ![](Shared\Icons\icon_PhAT_DisableCollision_40x.png) button.
+
+[/EXCERPT:ColOff]
+
+-->
+
+
+### Welding
+
+Welding Physics Bodies will make multiple bodies interact as one, as well as lock together the Skeletal Mesh Joints they are associated with. To Weld one Physics Body to another:
+
+1. **Select** 2 or more Physics Bodies using **Ctrl + Left Mouse click**.
+1. **Click** the "Weld" ![](Shared\Icons\icon_PhAT_Weld_40x.png) button.
+
+
+Any Physics Body welded to the currently selected Physics Body will show in Yellow.
+
+
+### Physical Materials
+
+Each Physics Body in a Physics Asset can have a Physical Material assigned to it. To apply a Physics Material to a single Physics Body:
+
+1. **Enable** Body Editing Mode ![](Shared\Icons\icon_PHatMode_Body_40x.png).
+1. **Select** a Physics Body in the Viewport or in the **Hierarchy** Pane.
+	* If it is difficult to see the Physics Body you desire, use the **Viewport** -> **Modes** -> **MeshRender Mode** -> **Wireframe/None** to adjust the visibility of the Skeletal 
+	Mesh.
+1. In the Details pane, locate the Simple Collision Physical Material property and assign a Physics Material.
+
+[REGION:raw]
+![](Engine/Physics/PhysicalMaterials/HowTo/PhysMat_PhAT_04.png)(w:920)
+[/REGION]
+
+
+## Constraints
+
+[REGION:note]
+The use of Physics Constraints is covered in the [Physics Constraints User Guide](Engine/Physics\Constraints\ConstraintsUserGuide) and their properties are covered in the [Physics Constraint Reference](Engine/Physics\Constraints\ConstraintsReference). This section will only cover workflows that are specific to the Physics Asset Tool or ones that have deviated substantially from the norm.
+[/REGION]
+
+### Editing Physics Constraints
+
+1. **Enable** Constraint Editing Mode ![](Shared\Icons\icon_PHatMode_Joint_40x.png).
+1. **Select** a Physics Constraint in the Viewport or in the Hierarchy Pane.
+1. **Move and rotate** the Physics Constraint using the Translation and Rotation to create the rotational point of the "joint" the Physics Constraint will form.
+1. Edit the Physics Constraint's properties in the Details Pane.
+
+	[REGION:tip]
+	You can switch Swing1, Swing2, and Twist from Limited to Locked quickly by using the "**1**", "**2**", and "**3**" keys respectively. The "**4**" key can be used cycle through 
+	them setting one to limited and the other two to locked.
+	[/REGION]
+
+1. Save ![](Shared\Icons\icon_SaveAsset_40x.png) often.
+
+
+
+See the **[Physics Asset Tool Reference](Engine/Physics/PhAT/Reference)** for more information on the properties in the Physics Asset Tool.
+
+
+
+
+### Aligning Physics Constraints
+
+If you are using the Physics Constraint's Linear or Angular limits, you will be able to see their alignment:
+
+![](Engine/Physics/PhAT/HowTo/PhATUG_Limits.png)(w:300 convert:false)
+
+By then translating and rotating the Physics Constraint, you can align the limits to give the desired results. In the most basic of terms, when a Physics Constraint is limited, 
+you can see a yellow line suspended in a green arc or cone structure. The line will be "constrained" within this arc or cone.
+
+For more information on Physics Constraints and their properties, see the **[Physics Constraints User Guide](Engine/Physics\Constraints\ConstraintsUserGuide)**.
+
+
+### Deleting Physics Constraints
+
+[REGION:note]
+There is no easy way to recreate a Physics Constraint, be mindful of this when deleting them.
+[/REGION]
+
+1. **Enable** Constraint Editing Mode ![](Shared\Icons\icon_PHatMode_Joint_40x.png).
+1. **Select** the Physics Constraint you wish to delete.
+1. **Press** the **Delete** key.
+
+
+### Recreating a Physics Constraint
+
+[REGION:note]
+There is no easy way to recreate a Physics Constraint, be mindful of this when deleting them.
+[/REGION]
+
+Physics Constraints are only created on generation of a Physics Body and are only generated upstream. So if you remove a Physics Constraint from the shoulder, you would have to 
+remove the upper arm Physics Body (which will remove the elbow Physics Constraint), then re-create the upper arm Physics Body. This would create the shoulder Physics Constraint, 
+but not the elbow Physics Constraint, so you would have to continue this process down the arm.
+
+
+### Constraint Snapping
+
+Snapping a Physics Constraint will take the blue locator for the Physics Constraint and snap its location and rotation to the red one's location and rotation.
+
+You can snap the currently selected Physics Constraint(s) by using the Snap Constraint ![](Shared\Icons\icon_PhAT_SnapConstraint_40x.png) button.
+
+
+## Both Physics Bodies and Physics Constraints
+
+### Copying Bodies and Constraints
+
+While in either mode, properties from one Physics Body or Constraint can be copied to other Physics Bodies or Constraints:
+
+1. **Select** the Physics Body or Constraint you wish to copy data from.
+1. **Press Ctrl + C**.
+1. **Select** the Physics Bodies or Constraints you wish to apply data to.
+1. **Press Ctrl + V**.
+
+
+This will not overwrite the shape of a Physics Body.
+
+
+## Testing
+
+Clicking the Simulation ![](Shared\Icons\icon_PhAT_PlaySim_40x.png) button will allow you to test your Physics Asset.
+
+![](Engine/Physics/PhAT/HowTo/SimPhAT.png)(w:600)
+
+The downward triangle next to the Simulation button will give you access to **Real Simulation** and **No Gravity**:
+
+* **Real Simulation** - Simulates the entire Physics Asset with gravity enabled.
+* **No Gravity** - Simulates the entire Physics Asset but does not turn on gravity, leaving you to **ctrl+click** to poke the Physics Bodies in a zero gravity environment.
+
+	This is useful for finding any interpenetrating Physics Bodies or Limited Physics Constraints already outside of their limits.
+
+You can also simulate a chain of joints by toggling **Selected Simulation** ![](Engine/Physics/PhAT/HowTo/selectedToggle_button.png) on, then clicking the **Simulation** ![](Shared\Icons\icon_PhAT_PlaySim_40x.png) button.
+
+* **Selected Simulation** - Does mostly what it sounds like, it only simulates the Physics Bodies you have selected (using **ctrl+left click** you can select more than one 
+Physics Body) and those down the hierarchy from the selected Physics Bodies. So if you select the shoulder, the entire arm will simulate.
+
+
+
+## Exploding
+
+This is caused by two Physics Bodies interpenetrating and the physics system attempting to correct it by applying a massive amount of force to expel the Physics Bodies from each other. 
+If a Physics Constraint is keeping the offending Physics Bodies together, the physics system will continually apply force to separate them, causing very erratic and extreme movement.
+
+![](Engine/Physics/PhAT/HowTo/brokenPhAT.png)
+
+This is normally fixed by disabling collision between the two offending Physics Bodies, or changing their positions and/or scale to make sure they do not interpenetrate. Also, welded 
+Physics Bodies act as one regardless of interpenetration.
+
+## Jitter
+
+If the Physics Asset is mostly "collapsed" but still shaking and twitching on the ground, but not violently bouncing around, there are a few values you can tweak to get it to stop 
+and sleep. 
+
+Before doing anything, try simulating with No Gravity, this will show you if any of your Physics Constraints are miss aligned and trying to correct before the Physics Asset has 
+even hit the ground.
+
+Often times, a small amount of Linear and Angular Damping on all the Physics Asset's Physics Bodies is more than enough to stop the asset from shaking. However, using large 
+values in Linear Damping will cause a Physics Body slow its movement through the world, even due to gravity. Damping, as it is implemented here, is not "drag" per se. It is 
+more like moving through a viscous substance.
+
+If jitter is still an issue, perhaps you have a number of really small Physics Bodies, check the **[Substepping](Engine/Physics/Substepping)** documentation for information on having the 
+physics simulation calculate in-between solutions.
+
+See the **[](Engine/Physics/PhAT/Reference)** and **[](Engine/Physics/Constraints/ConstraintsReference)** for more details on the properties for Physics Constraints.
+
+
+
+
+
+
diff --git a/Engine/Documentation/Source/Engine/Physics/PhysicsBodies/Overview/PhysicsBodiesOverview.INT.udn b/Engine/Documentation/Source/Engine/Physics/PhysicsBodies/Overview/PhysicsBodiesOverview.INT.udn
new file mode 100644
--- /dev/null
+++ b/Engine/Documentation/Source/Engine/Physics/PhysicsBodies/Overview/PhysicsBodiesOverview.INT.udn
@@ -0,0 +1,54 @@
+Availability: Public
+Title:Physics Bodies Overview
+Crumbs: %ROOT%, Engine, Engine/Physics
+Description:Overview of what a Physics Body, or BodyInstance, is and what it does.
+
+[VAR:Topic]
+[OBJECT:Topic]
+	[PARAM:image]
+		![%Engine/Physics/PhysicsBodies:title%](Engine/Physics/physics_topic.png)
+	[/PARAM]
+	[PARAM:icon]
+		![](%ROOT%/physics_icon.png)(convert:false)
+	[/PARAM]
+	[PARAM:title]
+		%Engine/Physics/PhysicsBodies:title%
+	[/PARAM]
+	[PARAM:description]
+		%Engine/Physics/PhysicsBodies:description%
+	[/PARAM]
+	[PARAM:path]
+		[RELATIVE:Engine/Physics/PhysicsBodies]
+	[/PARAM]
+[/OBJECT]
+[/VAR]
+[VAR:TopicCompact]
+[OBJECT:TopicCompact]
+	[PARAM:image]
+		![%Engine/Physics/PhysicsBodies:title%](Engine/Physics/physics_topic.png)
+	[/PARAM]
+	[PARAM:icon]
+		![](%ROOT%/physics_icon.png)(convert:false)
+	[/PARAM]
+	[PARAM:title]
+		%Engine/Physics/PhysicsBodies:title%
+	[/PARAM]
+	[PARAM:description]
+		%Engine/Physics/PhysicsBodies:description%
+	[/PARAM]
+	[PARAM:path]
+		[RELATIVE:Engine/Physics/PhysicsBodies]
+	[/PARAM]
+[/OBJECT]
+[/VAR]
+
+
+[TOC(start:2)]
+
+
+Overview Text
+
+## Additional Info
+
+
+
diff --git a/Engine/Documentation/Source/Engine/Physics/PhysicsBodies/Reference/PhysicsBodiesRef.INT.udn b/Engine/Documentation/Source/Engine/Physics/PhysicsBodies/Reference/PhysicsBodiesRef.INT.udn
new file mode 100644
--- /dev/null
+++ b/Engine/Documentation/Source/Engine/Physics/PhysicsBodies/Reference/PhysicsBodiesRef.INT.udn
@@ -0,0 +1,203 @@
+Availability: Public
+Title:Physics Bodies Reference
+Crumbs: %ROOT%
+Description:Physics Body, or Body Instance, property reference.
+
+[VAR:Topic]
+[OBJECT:Topic]
+	[PARAM:image]
+		![%Engine/Physics/PhysicsBodies:title%](Engine/Physics/physics_topic.png)
+	[/PARAM]
+	[PARAM:icon]
+		![](%ROOT%/physics_icon.png)(convert:false)
+	[/PARAM]
+	[PARAM:title]
+		%Engine/Physics/PhysicsBodies:title%
+	[/PARAM]
+	[PARAM:description]
+		%Engine/Physics/PhysicsBodies:description%
+	[/PARAM]
+	[PARAM:path]
+		[RELATIVE:Engine/Physics/PhysicsBodies]
+	[/PARAM]
+[/OBJECT]
+[/VAR]
+[VAR:TopicCompact]
+[OBJECT:TopicCompact]
+	[PARAM:image]
+		![%Engine/Physics/PhysicsBodies:title%](Engine/Physics/physics_topic.png)
+	[/PARAM]
+	[PARAM:icon]
+		![](%ROOT%/physics_icon.png)(convert:false)
+	[/PARAM]
+	[PARAM:title]
+		%Engine/Physics/PhysicsBodies:title%
+	[/PARAM]
+	[PARAM:description]
+		%Engine/Physics/PhysicsBodies:description%
+	[/PARAM]
+	[PARAM:path]
+		[RELATIVE:Engine/Physics/PhysicsBodies]
+	[/PARAM]
+[/OBJECT]
+[/VAR]
+
+This page contains a reference listing of properties in the Physics and Collision categories of properties. If you are looking for reference on Collision Responses or Collision Presets please see: [](Engine/Physics/Collision/Reference).
+
+
+## Properties
+
+Below are the properties for Physics Bodies (Body Instances) separated by major category.
+
+
+[OBJECT:modulesectionexpanded]
+	[PARAMLITERAL:id]
+		phys
+	[/PARAMLITERAL]
+	[PARAM:heading]
+		Physics
+	[/PARAM]
+	[PARAM:content]
+		%p1%
+	[/PARAM]
+[/OBJECT]
+
+[VAR:p1]
+
+![](physicsProps.png)
+
+[REGION:raw]
+| Property | %spacer% | Description |
+| --- | :---: | :--- |
+| **Simulate Physics** || %doxygen:FBodyInstance::bSimulatePhysics% |
+| **Mass in KG** || Mass of the body in KG. |
+| **Angular Damping** || %doxygen:FBodyInstance::AngularDamping% |
+| **Linear Damping** || %doxygen:FBodyInstance::LinearDamping% |
+| **Enable Gravity** || %doxygen:FBodyInstance::bEnableGravity% |
+|[REGION:tablesection]**Constraints**[/REGION]|||
+| **Lock Position** || [INCLUDE:#lockpos] |
+| **Lock Rotation** || [INCLUDE:#lockrot] |
+| **Mode** || [INCLUDE:#mode] |
+[/REGION]
+
+### Advanced
+
+![](physicsPropsAdv.png)
+
+[REGION:raw]
+| Property | %spacer% | Description |
+| --- | :---: | :--- |
+| **Auto Weld** || %doxygen:FBodyInstance::bAutoWeld% |
+| **Start Awake** || %doxygen:FBodyInstance::bStartAwake% |
+| **Center Of Mass Offset** || %doxygen:FBodyInstance::COMNudge% |
+| **Mass Scale** || %doxygen:FBodyInstance::MassScale% |
+| **Max Angular Velocity** || %doxygen:FBodyInstance::MaxAngularVelocity% |
+| **Use Async Scene** || %doxygen:FBodyInstance::bUseAsyncScene% |
+| **Override Max Depenetration Velocity** || %doxygen:FBodyInstance::bOverrideMaxDepenetrationVelocity% |
+| **Max Depenetration Velocity** || %doxygen:FBodyInstance::MaxDepenetrationVelocity% |
+| **Override Walkable Slope on Instance** || %doxygen:FBodyInstance::bOverrideWalkableSlopeOnInstance% |
+| **Walkable Slope Override** || %doxygen:UBodySetup::WalkableSlopeOverride% See the **[](Engine/Physics/WalkableSlope)** documentation for usage information. |
+| **Walkable Slope Behavior** || %doxygen:FWalkableSlopeOverride::WalkableSlopeBehavior% |
+| **Walkable Slope Angle** || %doxygen:FWalkableSlopeOverride::WalkableSlopeAngle% 
+| **Sleep Family** || %doxygen:FBodyInstance::SleepFamily% |
+| **Position Solver Iteration Count** || %doxygen:FBodyInstance::PositionSolverIterationCount% |
+| **Velocity Solver Iteration Count** || %doxygen:FBodyInstance::VelocitySolverIterationCount% |
+| **Should Update Physics Volume** || %doxygen:USceneComponent::bShouldUpdatePhysicsVolume% |
+
+[/REGION]
+
+[/VAR]
+
+
+[OBJECT:modulesectionexpanded]
+	[PARAMLITERAL:id]
+		coll
+	[/PARAMLITERAL]
+	[PARAM:heading]
+		Collision
+	[/PARAM]
+	[PARAM:content]
+		%p2%
+	[/PARAM]
+[/OBJECT]
+
+[VAR:p2]
+
+![](collisionProps.png)
+
+[REGION:raw]
+| Property | %spacer% | Description |
+| --- | :---: | :--- |
+| **Simulation Generates Hit Events** || %doxygen:FBodyInstance::bNotifyRigidBodyCollision% |
+| **Phys Material Override** || %doxygen:FBodyInstance::PhysMaterialOverride% |
+| **Generate Overlap Events** || %doxygen:UPrimitiveComponent::bGenerateOverlapEvents% |
+| **Collision Responses** || See the **[](Engine/Physics/Collision/Reference)** documentation for more information. |
+[/REGION]
+
+### Advanced
+
+![](collisionPropsAdv.png)
+
+[REGION:raw]
+| Property | %spacer% | Description |
+| --- | :---: | :--- |
+| **Use CCD** || %doxygen:FBodyInstance::bUseCCD% |
+| **Always Create Physics State** || %doxygen:UPrimitiveComponent::bAlwaysCreatePhysicsState% |
+| **Multi Body Overlap** || %doxygen:UPrimitiveComponent::bMultiBodyOverlap% | 
+| **Check Async Scene On Move** || %doxygen:UPrimitiveComponent::bCheckAsyncSceneOnMove% |
+| **Trace Complex On Move** || %doxygen:UPrimitiveComponent::bTraceComplexOnMove% |
+| **Return Material On Move** || %doxygen:UPrimitiveComponent::bReturnMaterialOnMove% |
+| **Can Ever Affect Navigation** || %doxygen:UPrimitiveComponent::bCanEverAffectNavigation% |
+
+[/REGION]
+
+[/VAR]
+
+[VAR:spacer]![](%ROOT%/spacer.png)(w:150 convert:false)[/VAR]
+
+[VAR:spacer2]![](%ROOT%/spacer.png)(w:40 convert:false)[/VAR]
+
+<!--
+[EXCERPT:lockpos]
+
+%doxygen:FBodyInstance::bLockTranslation%
+
+[REGION:raw]
+| Property | %spacer2% | Description |
+| --- | :---: | :--- |
+| **Lock X Translation** || %doxygen:FBodyInstance::bLockXTranslation% |
+| **Lock Y Translation** || %doxygen:FBodyInstance::bLockYTranslation% |
+| **Lock Z Translation** || %doxygen:FBodyInstance::bLockZTranslation% |
+[/REGION]
+[/EXCERPT]
+
+[EXCERPT:lockrot]
+
+%doxygen:FBodyInstance::bLockRotation%
+
+[REGION:raw]
+| Property | %spacer2% | Description |
+| --- | :---: | :--- |
+| **Lock X Rotation** || %doxygen:FBodyInstance::bLockXRotation% |
+| **Lock Y Rotation** || %doxygen:FBodyInstance::bLockYRotation% |
+| **Lock Z Rotation** || %doxygen:FBodyInstance::bLockZRotation% |
+[/REGION]
+[/EXCERPT]
+
+[EXCERPT:mode]
+
+%doxygen:FBodyInstance::DOFMode%
+
+[REGION:raw]
+| Property | %spacer2% | Description |
+| --- | :---: | :--- |
+| **Default** || Inherits the degrees of freedom from the project settings. |
+| **Six DOF** || Specifies which axis to freeze rotation and movement along. |
+| **YZPlane** || Allows 2D movement along the Y-Z plane. |
+| **XZPlane** || Allows 2D movement along the X-Z plane. |
+| **XYPlane** || Allows 2D movement along the X-Y plane. |
+| **Custom Plane** || Allows 2D movement along the plane of a given normal. |
+| **None** || No constraints. |
+[/REGION]
+[/EXCERPT]
+-->
\ No newline at end of file
