#UE4 #DOCS physics doc fixes

diff --git a/Engine/Documentation/Source/Engine/Actors/SkeletalMeshActors/SkeletalMeshActors.INT.udn b/Engine/Documentation/Source/Engine/Actors/SkeletalMeshActors/SkeletalMeshActors.INT.udn
--- a/Engine/Documentation/Source/Engine/Actors/SkeletalMeshActors/SkeletalMeshActors.INT.udn
+++ b/Engine/Documentation/Source/Engine/Actors/SkeletalMeshActors/SkeletalMeshActors.INT.udn
@@ -91,7 +91,7 @@ Normal collision creation and detection does not work with Skeletal Mesh Actors.
 
 ### Creating Collision for a Skeletal Mesh Actor 
 
-To create a Physics Asset for your mesh, all you have to do is simply **Right-click** on the mesh in the **Content Browser** and then choose **Create** -> **Create Physics Asset**. You can read more about how to create Physics Assets by checking out the [Physics Asset Tool User Guide]( Engine/Physics/PhAT/UserGuide)
+To create a Physics Asset for your mesh, all you have to do is simply **Right-click** on the mesh in the **Content Browser** and then choose **Create** -> **Create Physics Asset**. You can read more about how to create Physics Assets by checking out the [Physics Asset Tool](Engine/Physics/PhAT) documentation.
 
 ![](SkeletalMeshActor_Create_Physics_Asset.png)
 
diff --git a/Engine/Documentation/Source/Engine/Physics/Constraints/ConstraintsReference/Constraints.INT.udn b/Engine/Documentation/Source/Engine/Physics/Constraints/ConstraintsReference/Constraints.INT.udn
--- a/Engine/Documentation/Source/Engine/Physics/Constraints/ConstraintsReference/Constraints.INT.udn
+++ b/Engine/Documentation/Source/Engine/Physics/Constraints/ConstraintsReference/Constraints.INT.udn
@@ -49,7 +49,7 @@ This page contains a reference listing of properties in the Physics Constraints.
 
 ## Properties
 
-Below are the properties for Physics Constraints seperated by major category.
+Below are the properties for Physics Constraints separated by major category.
 
 [VAR:spacer]![](%ROOT%/spacer.png)(w:150 convert:false)[/VAR]
 
diff --git a/Engine/Documentation/Source/Engine/Physics/PhAT/HowTo/HowTo.INT.udn b/Engine/Documentation/Source/Engine/Physics/PhAT/HowTo/HowTo.INT.udn
new file mode 100644
--- /dev/null
+++ b/Engine/Documentation/Source/Engine/Physics/PhAT/HowTo/HowTo.INT.udn
@@ -0,0 +1,334 @@
+﻿Availability: Public
+Title:Physics Asset Tool How To Guide
+Crumbs: %ROOT%, Engine, Engine/Physics, Engine/Physics/PhAT
+Description:How to use the Physics Asset Tool to create Physics Asset for Skeletal Meshes.
+Version: 4.9
+
+[VAR:Topic]
+[OBJECT:Topic]
+	[PARAM:image]
+		![%Engine/Physics/PhAT/UserGuide:title%](Engine/Physics/physics_topic.png)
+	[/PARAM]
+	[PARAM:icon]
+		![](Engine/Physics/PhAT/HowTo/%ROOT%/physics_icon.png)(convert:false)
+	[/PARAM]
+	[PARAM:title]
+		%Engine/Physics/PhAT/UserGuide:title%
+	[/PARAM]
+	[PARAM:description]
+		%Engine/Physics/PhAT/UserGuide:description%
+	[/PARAM]
+	[PARAM:path]
+		[RELATIVE:Engine/Physics/Phat/UserGuide]
+	[/PARAM]
+[/OBJECT]
+[/VAR]
+[VAR:TopicCompact]
+[OBJECT:TopicCompact]
+	[PARAM:image]
+		![%Engine/Physics/PhAT/UserGuide:title%](Engine/Physics/physics_topic.png)
+	[/PARAM]
+	[PARAM:icon]
+		![](Engine/Physics/PhAT/HowTo/%ROOT%/physics_icon.png)(convert:false)
+	[/PARAM]
+	[PARAM:title]
+		%Engine/Physics/PhAT/UserGuide:title%
+	[/PARAM]
+	[PARAM:description]
+		%Engine/Physics/PhAT/UserGuide:description%
+	[/PARAM]
+	[PARAM:path]
+		[RELATIVE:Engine/Physics/Phat/UserGuide]
+	[/PARAM]
+[/OBJECT]
+[/VAR]
+
+
+
+
+The **Physics Asset Tool** is used for creating **Physics Assets** for use with **Skeletal Meshes**. One can start from nothing and build to a full ragdoll setup or use the automation 
+tools to create a basic set of Physics Bodies and Physics Constraints.
+
+This document will cover step by step procedures for the most common activities in the Physics Asset Tool.
+
+## New Physics Asset
+
+When a Skeletal Mesh is imported, there is an option to generate a Physics Asset for it as it is imported. Once the imported file is processed, you will be presented with the 
+"New Asset" dialog:
+
+[REGION:raw]
+![](Engine/Physics/PhAT/HowTo/PhATUG_ImportPA.png)
+[/REGION]
+
+However, you can follow these steps if you need to create a Physics Asset for a Skeletal Mesh at a later time:
+
+1. Find the Skeletal Mesh asset you wish to add a Physics Asset to in the 
+**Content Browser**.
+1. **Right-click** on the Skeletal Mesh in the Content Browser to bring up the **Context Menu** -> **Create** -> **Create Physics Asset**.
+1. Adjust setting to your liking.
+
+	![](Shared\PhAT\PhAT_restetDefault.png)
+
+1. **Click** 'Ok'.
+
+## Editing a Physics Asset
+
+1. Find the Physics Asset in the Content Browser.
+
+	![](Engine/Physics\PhysicalMaterials\PhysMatUserGuide\PhysMat_PhAT_02.png)
+
+1. **Double-click** to open it in the Physics Asset Tool (PhAT).
+
+	![](Engine/Physics/PhAT/HowTo/PhATUG_FullWindow.png)(w:900)
+
+
+## Physics Bodies
+
+### Editing Physics Bodies
+
+1. Enable Body Editing Mode ![](Shared\Icons\icon_PHatMode_Body_40x.png).
+1. Select a Physics Body in the Viewport or in the **Hierarchy Pane**.
+	* If it is difficult to see the Physics Body you desire, use the **Viewport** -> **Modes** -> **MeshRender Mode** -> **Wireframe/None** to adjust the visibility of the 
+	Skeletal Mesh.
+1. **Move, rotate, and scale** the Physics Body using the Translation, Rotation, and Scale tools so it fits the Skeletal Mesh and Skeletal Joint.
+1. Edit the Physics Body's properties in the Details Pane.
+1. **Enable** or **Disable** Collision between the current Physics Bodies and other in the Physics Asset using the _Collision On_ and _Collision Off_ tools. Any Physics Body that can 
+collide with the currently selected Physics Body will be blue, if they cannot collide, they will be gray.
+1. If a series of bodies should act as one, such as with the twist joints that form a wrist, use the _Weld_ tool to prevent undesired physics simulation behavior. Welded Physics 
+Bodies will be yellow in color.
+1. Save ![](Shared\Icons\icon_SaveAsset_40x.png) often.
+
+See the **[Physics Asset Tool Reference](Engine/Physics/PhAT/Reference)** for more information on the properties in the Physics Asset Tool.
+
+[PUBLISH:licensee]
+[REGION:note]
+As of change list #1977557, multi select is enabled for Physics Bodies and Physics Constraints, enabling you to edit multiple bodies or constraints at the same time.
+[/REGION]
+[/PUBLISH]
+
+### Deleting Physics Bodies
+
+1. **Enable** Body Editing Mode ![](Shared\Icons\icon_PHatMode_Body_40x.png).
+1. **Select** a Physics Body you wish to delete.
+	* If it is difficult to see the Physics Body you desire, use the **Viewport -> Modes -> MeshRender Mode -> Wireframe/None** to adjust the visibility of the Skeletal Mesh.
+1. Press the **Delete** key.
+
+	![](Shared\PhAT\PhAT_DelBone.png)
+
+### Creating New Physics Bodies
+
+![](Engine/Physics/PhAT/HowTo/PhATUG_addBody.png)(w:900)
+
+1. **Right-click** an already existing Physics Body or a Bone in the Hierarchy Panel.
+1. **Select** the type of Physics Body you wish to add from the Context Menu using the left mouse button.
+
+	If the Bone in the Hierarchy Panel does not already have a Physics Body on it, the only option in the **right-click** context menu is "Add Body."
+
+	Finally, you can just **double-click** a bone that does not have a Physics Body in the **Hierarchy** Panel to "Add Body."
+
+
+### Enabling and Disabling Collision Between Physics Bodies
+
+[REGION:simpletable]
+| --- | --- |
+|[INCLUDE:#ImgColOn]| [INCLUDE:#ColOn] |
+|[INCLUDE:#ImgColOff]| [INCLUDE:#ColOff] |
+[/REGION]
+
+<!--
+
+[EXCERPT:ImgColOn]
+[REGION:lightbox]
+[![](Shared\PhAT\PhAT_collisionOn.png)(w:350)](Shared\PhAT\PhAT_collisionOn.png)
+[/REGION]
+[/EXCERPT:ImgColOn]
+
+[EXCERPT:ImgColOff]
+[REGION:lightbox]
+[![](Shared\PhAT\PhAT_collisionOff.png)(w:350)](Shared\PhAT\PhAT_collisionOff.png)
+[/REGION]
+[/EXCERPT:ImgColOff]
+
+[EXCERPT:ColOn]
+[REGION:raw]
+**Enabling Collision**
+
+1. **Select** 2 or more Physics Bodies using **Ctrl + Left Mouse click**.
+1. **Click** the "Collision On" ![](Shared\Icons\icon_PhAT_EnableCollision_40x.png) button.
+
+
+[/REGION]
+[/EXCERPT:ColOn]
+
+[EXCERPT:ColOff]
+**Disabling Collision** 
+
+1. **Select** 2 or more Physics Bodies using **Ctrl + Left Mouse click**.
+1. **Click** the "Collision Off" ![](Shared\Icons\icon_PhAT_DisableCollision_40x.png) button.
+
+[/EXCERPT:ColOff]
+
+-->
+
+
+### Welding
+
+Welding Physics Bodies will make multiple bodies interact as one, as well as lock together the Skeletal Mesh Joints they are associated with. To Weld one Physics Body to another:
+
+1. **Select** 2 or more Physics Bodies using **Ctrl + Left Mouse click**.
+1. **Click** the "Weld" ![](Shared\Icons\icon_PhAT_Weld_40x.png) button.
+
+
+Any Physics Body welded to the currently selected Physics Body will show in Yellow.
+
+
+### Physical Materials
+
+Each Physics Body in a Physics Asset can have a Physical Material assigned to it. To apply a Physics Material to a single Physics Body:
+
+1. **Enable** Body Editing Mode ![](Shared\Icons\icon_PHatMode_Body_40x.png).
+1. **Select** a Physics Body in the Viewport or in the **Hierarchy** Pane.
+	* If it is difficult to see the Physics Body you desire, use the **Viewport** -> **Modes** -> **MeshRender Mode** -> **Wireframe/None** to adjust the visibility of the Skeletal 
+	Mesh.
+1. In the Details pane, locate the Simple Collision Physical Material property and assign a Physics Material.
+
+[REGION:raw]
+![](Engine/Physics/PhysicalMaterials/HowTo/PhysMat_PhAT_04.png)(w:920)
+[/REGION]
+
+
+## Constraints
+
+[REGION:note]
+The use of Physics Constraints is covered in the [Physics Constraints User Guide](Engine/Physics\Constraints\ConstraintsUserGuide) and their properties are covered in the [Physics Constraint Reference](Engine/Physics\Constraints\ConstraintsReference). This section will only cover workflows that are specific to the Physics Asset Tool or ones that have deviated substantially from the norm.
+[/REGION]
+
+### Editing Physics Constraints
+
+1. **Enable** Constraint Editing Mode ![](Shared\Icons\icon_PHatMode_Joint_40x.png).
+1. **Select** a Physics Constraint in the Viewport or in the Hierarchy Pane.
+1. **Move and rotate** the Physics Constraint using the Translation and Rotation to create the rotational point of the "joint" the Physics Constraint will form.
+1. Edit the Physics Constraint's properties in the Details Pane.
+
+	[REGION:tip]
+	You can switch Swing1, Swing2, and Twist from Limited to Locked quickly by using the "**1**", "**2**", and "**3**" keys respectively. The "**4**" key can be used cycle through 
+	them setting one to limited and the other two to locked.
+	[/REGION]
+
+1. Save ![](Shared\Icons\icon_SaveAsset_40x.png) often.
+
+
+
+See the **[Physics Asset Tool Reference](Engine/Physics/PhAT/Reference)** for more information on the properties in the Physics Asset Tool.
+
+
+
+
+### Aligning Physics Constraints
+
+If you are using the Physics Constraint's Linear or Angular limits, you will be able to see their alignment:
+
+![](Engine/Physics/PhAT/HowTo/PhATUG_Limits.png)(w:300 convert:false)
+
+By then translating and rotating the Physics Constraint, you can align the limits to give the desired results. In the most basic of terms, when a Physics Constraint is limited, 
+you can see a yellow line suspended in a green arc or cone structure. The line will be "constrained" within this arc or cone.
+
+For more information on Physics Constraints and their properties, see the **[Physics Constraints User Guide](Engine/Physics\Constraints\ConstraintsUserGuide)**.
+
+
+### Deleting Physics Constraints
+
+[REGION:note]
+There is no easy way to recreate a Physics Constraint, be mindful of this when deleting them.
+[/REGION]
+
+1. **Enable** Constraint Editing Mode ![](Shared\Icons\icon_PHatMode_Joint_40x.png).
+1. **Select** the Physics Constraint you wish to delete.
+1. **Press** the **Delete** key.
+
+
+### Recreating a Physics Constraint
+
+[REGION:note]
+There is no easy way to recreate a Physics Constraint, be mindful of this when deleting them.
+[/REGION]
+
+Physics Constraints are only created on generation of a Physics Body and are only generated upstream. So if you remove a Physics Constraint from the shoulder, you would have to 
+remove the upper arm Physics Body (which will remove the elbow Physics Constraint), then re-create the upper arm Physics Body. This would create the shoulder Physics Constraint, 
+but not the elbow Physics Constraint, so you would have to continue this process down the arm.
+
+
+### Constraint Snapping
+
+Snapping a Physics Constraint will take the blue locator for the Physics Constraint and snap its location and rotation to the red one's location and rotation.
+
+You can snap the currently selected Physics Constraint(s) by using the Snap Constraint ![](Shared\Icons\icon_PhAT_SnapConstraint_40x.png) button.
+
+
+## Both Physics Bodies and Physics Constraints
+
+### Copying Bodies and Constraints
+
+While in either mode, properties from one Physics Body or Constraint can be copied to other Physics Bodies or Constraints:
+
+1. **Select** the Physics Body or Constraint you wish to copy data from.
+1. **Press Ctrl + C**.
+1. **Select** the Physics Bodies or Constraints you wish to apply data to.
+1. **Press Ctrl + V**.
+
+
+This will not overwrite the shape of a Physics Body.
+
+
+## Testing
+
+Clicking the Simulation ![](Shared\Icons\icon_PhAT_PlaySim_40x.png) button will allow you to test your Physics Asset.
+
+![](Engine/Physics/PhAT/HowTo/SimPhAT.png)(w:600)
+
+The downward triangle next to the Simulation button will give you access to **Real Simulation** and **No Gravity**:
+
+* **Real Simulation** - Simulates the entire Physics Asset with gravity enabled.
+* **No Gravity** - Simulates the entire Physics Asset but does not turn on gravity, leaving you to **ctrl+click** to poke the Physics Bodies in a zero gravity environment.
+
+	This is useful for finding any interpenetrating Physics Bodies or Limited Physics Constraints already outside of their limits.
+
+You can also simulate a chain of joints by toggling **Selected Simulation** ![](Engine/Physics/PhAT/HowTo/selectedToggle_button.png) on, then clicking the **Simulation** ![](Shared\Icons\icon_PhAT_PlaySim_40x.png) button.
+
+* **Selected Simulation** - Does mostly what it sounds like, it only simulates the Physics Bodies you have selected (using **ctrl+left click** you can select more than one 
+Physics Body) and those down the hierarchy from the selected Physics Bodies. So if you select the shoulder, the entire arm will simulate.
+
+
+
+## Exploding
+
+This is caused by two Physics Bodies interpenetrating and the physics system attempting to correct it by applying a massive amount of force to expel the Physics Bodies from each other. 
+If a Physics Constraint is keeping the offending Physics Bodies together, the physics system will continually apply force to separate them, causing very erratic and extreme movement.
+
+![](Engine/Physics/PhAT/HowTo/brokenPhAT.png)
+
+This is normally fixed by disabling collision between the two offending Physics Bodies, or changing their positions and/or scale to make sure they do not interpenetrate. Also, welded 
+Physics Bodies act as one regardless of interpenetration.
+
+## Jitter
+
+If the Physics Asset is mostly "collapsed" but still shaking and twitching on the ground, but not violently bouncing around, there are a few values you can tweak to get it to stop 
+and sleep. 
+
+Before doing anything, try simulating with No Gravity, this will show you if any of your Physics Constraints are miss aligned and trying to correct before the Physics Asset has 
+even hit the ground.
+
+Often times, a small amount of Linear and Angular Damping on all the Physics Asset's Physics Bodies is more than enough to stop the asset from shaking. However, using large 
+values in Linear Damping will cause a Physics Body slow its movement through the world, even due to gravity. Damping, as it is implemented here, is not "drag" per se. It is 
+more like moving through a viscous substance.
+
+If jitter is still an issue, perhaps you have a number of really small Physics Bodies, check the **[Substepping](Engine/Physics/Substepping)** documentation for information on having the 
+physics simulation calculate in-between solutions.
+
+See the **[](Engine/Physics/PhAT/Reference)** and **[](Engine/Physics/Constraints/ConstraintsReference)** for more details on the properties for Physics Constraints.
+
+
+
+
+
+
diff --git a/Engine/Documentation/Source/Engine/Physics/PhAT/PhAT.INT.udn b/Engine/Documentation/Source/Engine/Physics/PhAT/PhAT.INT.udn
--- a/Engine/Documentation/Source/Engine/Physics/PhAT/PhAT.INT.udn
+++ b/Engine/Documentation/Source/Engine/Physics/PhAT/PhAT.INT.udn
@@ -1,7 +1,7 @@
 Availability: Public
 Title:Physics Asset Tool (PhAT)
 Crumbs: %ROOT%, Engine, Engine/Physics
-Description:Tool to set up physics bodies and joints to use for physical simulation and collision for SkeletalMeshes.
+Description: Tool to set up physics bodies and joints to use for physical simulation and collision for SkeletalMeshes.
 Version: 4.9
 
 
@@ -44,382 +44,7 @@ The PhAT Physics Editor can be opened by **double-clicking** any Physics Asset o
 		Physical Materials Topics
 	[/PARAM]
 	[PARAM:links]
-		* [](Engine/Physics/PhAT/UserGuide "%Engine/Physics/PhAT/UserGuide:description%")
-		* [](Engine/Physics/PhAT/Reference "%Engine/Physics/PhAT/Reference:description%")
-		* [](Engine/Physics/PhAT/Properties "%Engine/Physics/PhAT/Properties:description%")
+		* [](Engine/Physics/PhAT/HowTo "%Engine/Physics/PhAT/HowTo:Description%")
+		* [](Engine/Physics/PhAT/Reference "%Engine/Physics/PhAT/Reference:Description%")
 	[/PARAM]
 [/OBJECT]
-
-
-
-
-<!--
-
-## Opening PhAT
-
-The PhAT Physics Editor can be opened by double-clicking any Physics Asset or through the right-click context menu of a physics asset in the Content Browser. Either of these will open up that particular physics asset in the PhAT Physics Editor for editing. 
-
-
-
-### Preview Pane
-
-![previewpane.jpg](previewpane.jpg)
-
-The Preview Pane shows a rendered view of the default skeletal mesh asset associated with the physics asset. It also shows any physical bodies associated with the ones in the skeleton of the default mesh, any constraints between those bodies, visual tools to manipulate those bodies and constraints, and information about those bodies and constraints. The rendering mode for the mesh, the bodies, and the constraints can all be adjusted individually to get just the right preview for the task at hand. Perhaps most importantly, the physics simulation of the bodies and constraints can be previewed directly in the viewport and the mesh can be manipulated (nudging it or grabbing amd moving it) to see how they react to particular actions.
-
-
-### Properties Pane
-
-![propertiespane.jpg](propertiespane.jpg)
-
-The Properties pane is a standard Unreal properties window which shows either a set of basic editor properties if nothing is selected or the properties of the body or constraint currently selected in the Tree Pane, depending on what mode the editor is in. When in Body editing mode, the Properties Pane can be toggled between displaying the body's properties and the instance properties. When in Constraint editing mode, the Properties Pane can be toggled between showing the constraint's properties and the instance properties.
-
-
-#### Editor Properties
-
-**Anim**
-
-
-* **Preview Anim Set** - Sets the AnimSet used for previewing animations in the Preview Pane during simulation preview.
-* **Physics Blend** - Sets the blending between physics driving the skeleton and the preview animation dirving the skeleton.
-* **Blend On Poke** - If true, the preview mesh will be driven by the preview animation until poked (Ctrl + LMB) in the Preview Pane. It then snaps to the physics simulation, pauses for _Poke Pause Time_, and blends back to the animation over _Poke Blend Time_.
-* **Poke Pause Time** - Sets the amount of time to show the physics simulation when _Blend On Poke_ is true.
-* **Poke Blend Time** - Sets the amount of time it takes to blend back to the animation when _Blend On Poke_ is true.
-* **Angular Spring Scale** - Sets an overall spring scale multiplier applied to all motors.
-* **Angular Damping Scale** - Sets an overall damping scale multiplier applied to all motors.
-
-
-**Simulation**
-
-
-* **Draw Contacts** -
-* **Sim Speed** - Sets the playback speed of the preview (physics simulation or preview animation).
-* **Floor Gap** - Sets the distance above the floor the skeletal mesh is position when at rest.
-* **Grav Scale** - Sets a multiplier for the gravity of the simulation preview.
-
-
-**Advanced**
-
-
-* **Prompt On Bone Delete** -
-* **Show Constraints As Points** - If true, constraints will be displayed as points in the Preview Pane. Otherwise, they are displayed as crosses.
-* **Show Names In Hierarchy** - If true, the names of the bones will be displayed in the Preview Pane when the skeletal mesh's skeleton is being displayed.
-
-
-**Mouse Spring**
-
-
-* **Handle Linear Damping** - Sets the linear damping of the spring used when grabbing (Ctrl + RMB + Drag) the physics asset during the simulation preview.
-* **Handle Linear Stiffness** - Sets the linear stiffness of the spring used when grabbing (Ctrl + RMB + Drag) the physics asset during the simulation preview.
-* **Handle Angular Damping** - Sets the angular damping of the spring used when grabbing (Ctrl + RMB + Drag) the physics asset during the simulation preview.
-* **Handle Angular Stiffness** - Sets the angular stiffness of the spring used when grabbing (Ctrl + RMB + Drag) the physics asset during the simulation preview.
-
-
-**Poking**
-
-
-* **Poke Strength** - Sets the strength of the impulse applied when poking (Ctrl + LMB) the physics asset in during the simulation preview.
-
-
-**Lighting**
-
-
-* **Sky Brightness** - Sets the brightness of the ambient light in the Preview Pane.
-* **Brightness** - Sets the brightness of the directional light in the Preview Pane.
-
-
-**Snap**
-
-
-* **Angular Snap** - Sets the snapping increment for rotations in the Preview Pane when snapping is enabled.
-* **Linear Snap** - Sets the snapping increment for translations in the Preview Pane when snapping is enabled.
-
-
-**Drawing**
-
-
-* **Constraint Draw Size** - Sets the size at which constraints are drawn in the Preview Pane.
-
-
-
-### Tree Pane
-
-![treepane.jpg](treepane.jpg)
-
-The Tree Pane shows a hierarchical view of the bodies or constraints of the physics asset, depending on which mode the editor is currelty in. The tree view shows all bones which have a body or constraint associated with it or whose children have a body or constraint associated with them. Bones with bodies or constraints directly associated with them are shown in bold to distinguish them. Selecting bodies or constraints in the Tree Pane will select them in the Preview Pane allowing them to be transformed with the transformation widget and will make their properties visible in the Properties Pane for editing as well.
-
-
-### Controls
-
-
-#### Mouse Controls
-
-| **LMB** | Select limb/joint or widget axis |
-| --- | --- |
-| **SHIFT + LMB** | will do a non-uniform transform if possible (scaling along one axis) |
-| *X* | Lock Selection |
-| **CTRL + LMB** | Poke the physics asset (only in a live simulation) |
-| **CTRL + RMB + Drag** | Grab onto the physics asset (only in a live simulation) |
-
-
-#### Keyboard Controls
-
-| --- | --- |
-| *W* | Move Mode |
-| *E* | Rotate Mode |
-| *R* | Scale Mode |
-| **Spacebar** | Cycle Widget Mode |
-| *B* | Switch between Body and Constraint mode |
-| *A* | Turn on Snap |
-| *C* | Copy properties from/to |
-| *S* | Simulation start/stop |
-| *H* | Cycle render mode for asset |
-| *J* | Cycle render mode for bodies |
-| *K* | Cycle render mode for constraints |
-| *D* | Weld bodies together (must be parent-child pair) |
-| *Q* | Cycle constraint X -> Y -> Z |
-| **DEL** | Delete constraint/body |
-| **HOME** | Center view on selected Body or Constraint |
-| **L + DRAG** | Orbit directional light |
-| *[* | Enable collision between selected body and another |
-| *]* | Disable collision with the selected body and another |
-
-
-#### Hot Keys
-
-
-
-## Creating a Physics Asset
-
-To create a physics asset, navigate to the [Content Browser](Engine/Content/Browser) and find the skeletal mesh asset you wish to have a physical setup for. Right click on it and select _Create Physics Asset_ from the context menu. This will open a dialog that provides you with options for the creation of a physics asset based on the skeletal structure and vertex weighting for that mesh. This dialog can be reached again while inside of PhAT by using the _Reset Entire Asset_ option in the Tools menu.
-
-
-
-## Modifying a physics asset
-
-Once PhAT is open, you will see your asset floating in space over a default _floor_. The object should also be enveloped in several boxes or sphere/sphyls depending on the options chosen in the collision primitives section of the _Create physics asset_ dialog. Most likely, these collision primitives are not optimally placed and may be interpenetrating in an undesired manner. If this is the case, PhAT has several tools for manipulating these primitives. For these examples, we will use Phys_Capsule3_Physics, which is in PhysicsObjects.upk, which should be in the Packages folder.
-
-
-
-## Collision primitives / bodies
-
-There are three different collision primitives that can be created and modified inside of PhAT. These three are boxes (that can be uniformly and non-uniformly scaled), capsules (can be scaled uniformly or along their length or width), and spheres (that can only be scaled uniformly).  Modifying collision primitives for physical assets is easily accomplished using PhAT. Open Phys_Capsule3_Physics in PhAT, and notice that there are wireframe spheres and capsules surrounding the object. These wireframe objects are the collision primitives for the bones in this object. To see how these are used, press the _Simulate_ button and watch the behavior of the object. When you are ready, press _Simulate_ again to stop it. Select one of the collision primitives by left clicking on it, and it should highlight orange, while primitives immediately adjacent to it turn white. This signifies that collision has been disabled between these collision bodies, as the constraint is handling the interaction between them. Press the _X_ key to lock the selection so we don't accidentally unselect the primitive or select another one. The word **LOCK** will appear in the lower left corner of the window.  
-
-![Fig1_1.jpg](Fig1_1.jpg)
-
-There is an axis gizmo at the center of this primitive, which can be used to move, rotate, and scale the collision bodies. Make sure you are in _move_ mode by either pressing the _W_ key or clicking on the Move mode button. Now left click and drag on the blue gizmo line (Z axis) and the body should move along that axis. As previously mentioned, the same controls can be used to rotate and scale bodies. To rotate a body, either press _E_ to enter rotation mode, or click on the _Rotate_ button in the toolbar. The translation gizmo should turn into a rotation gizmo with 3 wire rings representing each axis. If you have trouble seeing the gizmo, you can Cycle Mesh Rendering Mode (or press _H_) to wireframe or invisible mode.
-
-![Fig1_2.jpg](Fig1_2.jpg)
-
-Each of these rings is used to rotate along that axis the same manner in which movement was done, by left clicking and dragging. This isn't very obvious when rotating a spherical body, but should be very obvious when done to a box or capsule body.
-
-Scaling is done the same way as translation; however it supports non-uniform scaling as well. Switch to _Scale_ mode by pressing _R_ or clicking on the _Scale_ button.  The gizmo should appear the same as it did with _Move_, however left click will scale the entire body up or down depending on which direction you drag. Shift-Left click however will only scale along the axis you click on. It is important to note that attempting to non-uniformly scale a sphere body will cause the sphere to scale uniformly.
-
-Adding and removing bodies is also easy to do. Selecting a body and then clicking on any of the _Add Primitive_ buttons will drop that primitive at the bone's origin. Once that is done, the body can be moved, rotated, and scaled to the desired position. There is no limit to the number of collision bodies per bone, but performance will become an issue with a large number of bones and bodies in a single physical asset.
-
-
-### Body Properties
-
-**RB_BodySetup**
-
-
-* **Sleep Family** - Sets the method used to determine when to put the body to sleep.
-* **Bone Name** - Displays the name of the bone with which the body is associated.
-* **Fixed** - If true, the body will have no physics applied to it. it will be fixed relative to the world.
-* **No Collision** - If true, the body will not collide with anything.
-* **Block Zero Extent** - If true, the body will block zero extent checks, such as instant hit weapon traces.
-* **Block Non Zero Extent** - If true, the body will block non-zero extent checks, such as player collisions.
-* **Enable Continuous Collision Detection** - If true, the continuous collision detection system is used. This can be useful to keep bodies moving at high speeds from passing through geometry, but comes at the expense of slower calculations.
-* **Always Full Anim Weight** - If true, the body will always have physics applied to it regardless of the _Physics Weight_ value. This is useful for bodies that should always be driven by physics such as cloth or ponytails.
-* **Consider For Bounds** - If true, this body will factor into the bounds determination. The fewer bodies that are used in determining the bounds, the faster the calculation will be each frame.
-* **Phys Material** - Sets the Physical Material used by this body. See the [Physical Material](Engine/Physics/PhysicalMaterials/) page for more information.
-* **Mass Scale** - Sets a multiplier for the mass of the body. This value is multiplied by the default mass which is based on the volume of the body.
-* **Pre Cached Phys Scale** - Sets a list of scales at which to pre-cache physics for this body.
-
-
-**KMeshProps**
-
-
-* **COMNudge** - An offset applied on top of the body's PhysX-calculated Center of Mass.
-* **Agg Geom**
-    * **No RBCollision** - If true, this shape is completely ignored for physics. It will not only have no collision, but will also not contribute to the mass/inertia tensor of the body. Setting this on all shapes within a body is invalid.
-    * **Per Poly Shape** - If true, this shape will be considered by Unreal line traces even if this bone is defined as using per-poly collision with the PerPolyCollisionBones option.
-
-
-
-### Body Instance Properties
-
-**Bone Spring**
-
-
-* **Enable Bone Spring Linear** - If true, a linear 'spring' between the physics body for this bone and the world-space location of the animation bone is enabled.
-* **Enable Bone Spring Angular** - If true, a angular 'spring' between the physics body for this bone and the world-space location of the animation bone is enabled.
-* **Disable On Overextension** - If true, the bone spring will automatically disable if it ever gets longer than the _Overextension Threshold_.
-* **Notify Owner On Overextension** - If true, a notification is sent to the Owning Actor when overextended is sent by calling _Actor::OnRigidBody_.
-* **Teleport On Overextension** - If true, this will teleport the whole Physics Asset Instance if this spring gets too long to reduce the error to zero. **Note:** Having this enabled on more than one body in a single physics asset will have unpredictable results.
-* **Use Kin Actor For Bone Spring** - If true, when using bone springs, connect them to a kinematic body and animate that rather than animating the attachment location on the 'world'. This adds some overhead, but tracks rapidly moving targets better.
-* **Make Spring To Base Collision Component** - If true, when using bone springs, connect them to the physics body of the Base's Collision Component. When using this option, SetBoneSpringTarget must be given a matrix for this bone relative to the other bone, rather than relative to the world.
-* **Bone Linear Spring** - Sets the strength of the linear spring to the animated bone.
-* **Bone Linear Damping** - Sets the damping of the linear spring to the animated bone.
-* **Bone Angular Spring** - Sets the strength of the angular spring to the animated bone.
-* **Bone Angular Damping** - Sets the damping of the linear spring to the animated bone.
-* **Overextension Threshold** - If _Disable On Overextension_ is true, the bone spring will be disabled if it stretches more than this amount.
-
-
-**Physics**
-
-
-* **Only Collide With Pawns** - If true, this body should only collide with other bodies with their components marked by _bConsiderPawnForRBCollision_. This is useful for flappy bits which should not collide with the world.
-* **Enable Collision Repsonse** - If true, this enables a physics response for this body. If false, contacts are still generated and reported. This is useful for sensor bodies.
-* **Push Body** - If true, this body will be denoted as a "push" body. This also disables collision response.
-* **Phys Material Override** - Sets the physical material to use for this body, overriding that of the actor using the physics asset or the material assigned to it.
-* **Contact Report Force Threshold** - Sets the threshold used by the force-based collision reporting system. if the value is negative, the system will be disabled for this body.
-* **Instance Mass Scale** - Per-instance multiplier to scale the mass.
-* **Instance Damping Scale** - Per-instance scaling of linear/angular damping for this body.
-
-
-**RB_BodyInsatnce**
-
-
-* **Custom Gravity Factor** - A multiplier for the custom gravity applied to the body, allowing bodies to individually control how custom gravity settings affect them.
-
-
-
-
-## Constraints
-
-Toggling the Body mode button will switch to _Constraint_ mode, and the rendering of the physical asset will change to show purple wireframe crosses at the bone locations. The constraints are the representations of how the bones are going to be held together when the simulation is run.
-
-![Constr_1.jpg](Constr_1.jpg)
-
-Select a constraint the same way a body is selected, and you'll see the limits of that constraint. By default, a constraint is unlimited, and is represented as just the normal gizmo, but once limits are set, this will change to reflect the values entered by the user. For twist limits, a line is created along the axis of twist, with an X at the end if there is no value set for the limit, but if a limit is set the representation is a green arc at the end of the line with another yellow line bisecting the arc. This yellow line shows the current position of the bone within the twist limits. For swing limits, the representation without a value set in either _Swing1LimitAngle_ or _Swing2LimitAngle_ is a green line extending equally in both directions along the axis of the constraint. Pressing _Q_ with a constraint selected will cycle to the next axis if that needs to be changed. If a value is entered into one of the _SwingLimits_, the green line will change to a green cone, and the yellow line inside of it shows the current position of the bone again. If both swing and twist limits are entered, the twist limits will appear at the end of the yellow line inside the green cone.
-
-With the constraint selected, change the values so you can see the limits change. Once you've done that, simulate the asset to see the results of your tweaks. You can turn on rendering the constraint (either the one selected or all of them) by cycling through modes in the Cycle Constraint Rendering Mode button or pressing _K_. With this enabled, you'll be able to see the yellow line move around within the constraint and stop when it reaches the limits.  
-
-Constraints can be rotated and moves in the same way that bodies can, however, the difference here is that holding down shift to move a constraint will separate the starting location from the constraint location (causing the constraint to snap back together when the simulation starts), and rotating the constraint will rotate the limits but not the current rotation of the joint, which is useful for setting up non-symmetrical joints such as elbows. It is important to note that if you rotate a constraint so that the limits are beyond the current position of the joint, when the asset is simulated all of the joints will move to satisfy the limits in the first frame, thus causing motion immediately upon simulating the asset.
-
-![RotatedConstraint.jpg](RotatedConstraint.jpg)
-
-
-### Constraint Properties
-
-**RB_ConstraintSetup**
-
-
-* **Joine Name** - The name of the joint. This will be the same as the _Constraint Bone 1_ in a physics asset.
-* **Constraint Bone 1** - The name of the first bone that the joint should connect.
-* **Constraint Bone 2** - The name of the second bone that the joint should connect.
-* **Enable Projection** - If true, when distance error between bodies exceeds 0.1 units or rotation error exceeds 10 degrees, the body will be projected to fix this.
-
-
-**Linear**
-
-
-* **Linear Limit Soft** - By default the linear limits are `hard'. If you set this to true they are treated as springy instead and you can `push' into them.
-* **Linear Breakable** - If true, this joint can be broken by applying a linear force to it (basically pulling the joint apart).
-* **Linear X/Y/ZSetup**
-    * **b Limited** - If true, there is a limit on this degree of freedom. If false, linear movement is not limited at all.
-    * **Limit Size** - If _b Limited_ is true, how much movement is allowed in this direction. Set this to zero to completely `lock' this degree of freedom.
-* **Linear Limit Stiffness** - If _Linear Limit Soft_ is true, this controls the stiffness of the limit spring.
-* **Linear Limit Damping** - If _Linear Limit Soft_ is true, this controls the damping of the limit spring.
-* **Linear Break Threshold** - If _Linear Breakable_ is true, this is the force that is needed to break the joint.
-
-
-**Angular**
-
-
-* **Swing Limited** - If true, the angle between the two bodies is limited to a cone.
-* **Twist Limited** - If true, the twist between the two bodies is limited.
-* **Swing Limit Soft** - By default the swing cone limit is `hard'. If this is set to true it is springy instead and can be pushed into.
-* **Twist Limit Soft** - By default the twist limit is `hard'. If this is set to true it is springy instead and can be pushed into.
-* **Angular Breakable** - If true, this joint can be broken by applying a torque to it (basically twisting the joint apart).
-* **Swing 1LImit Angle** - Defines one dimension of the swing limit cone in degrees.
-* **Swing 2Limit Angle** - Defines the other dimension of the swing limit cone in degrees.
-* **Twist Limit Angle** - Defines the amount of twist that is allowed between the two constrained bodies.
-* **Swing Limit Stiffness** - If _Swing Limit Soft_ is true, this controls the stiffness of the cone.
-* **Swing Limit Damping** - If _Swing Limit Soft_ is true, this controls the damping of the cone.
-* **Twist Limit Stiffness** - If _Twist Limit Soft_ is true, this controls the stiffness of the twist limit.
-* **Twist Limit Damping** - If _Twist Limit Soft_ is true, this controls the damping of the twist limit
-* **Angular Break Threshold** - If _Angular Breakable_ is true, this is the torque required to break this joint.
-
-
-**Pulley**
-
-
-* **Is Pulley** - If true, this joint will be treated as a pulley type
-* **Maintain Min Distance** - If true, no `slack' is allowed in the `ropes' and they are treated a bit like rods instead, so you can push up on one object and the other will be pushed down.
-* **Pulley Ratio** - The gear ratio on the pulley. If its not 1.0, pulling on one side will result in less movement on the other, but with more force.
-
-
-
-### Constraint Instance Properties
-
-**Linear**
-
-
-* **Linear XPosition Drive** - Enable position drive on the linear X axis.
-* **Linear XVelocity Drive** - Enable velocity drive on the linear X axis.
-* **Linear YPosition Drive** - Enable position drive on the linear Y axis.
-* **Linear YVelocity Drive** - Enable velocity drive on the linear Y axis.
-* **Linear ZPosition Drive** - Enable position drive on the linear Z axis.
-* **Linear ZVelocity Drive** - Enable velocity drive on the linear Z axis.
-* **Linear Position Target** - If linear position drive is enabled, this is the target relative position that the joint will try to achieve.
-* **Linear Velocity Target** - If linear velocity drive is enabled, this is the target relative linear velocity that the joint will try to achieve.
-* **Linear Drive Spring** - If linear position drive is enabled, this is the stiffness of the drive spring.
-* **Linear Drive Damping** - If linear position drive is enabled, this is the damping of the drive spring.
-* **Linear Drive Force Limit** - If linear velocity drive is enabled, this is the maximum force that will be applied to attain LinearVelocityTarget.
-
-
-**Angular**
-
-
-* **Swing Position Drive** - Enable position drive on the swing axes.
-* **Swing Velocity Drive** - Enable velocity drive on the swing axes.
-* **Twist Position Drive** - Enable position drive around the twist axis.
-* **Twist Velocity Drive** - Enable velocity drive around the twist axis.
-* **Angular Slerp Drive** - Advanced option. Specifies the method that is used for applying drive to the joint.
-* **Angular Position Target** - If angular position drive is enabled, this is the target relative quaternion that the joint will try to achieve.
-* **Angular Velocity Target** - If angular velocity drive is enabled, this is the target relative angular velocity that the joint will try to achieve.
-* **Angular Drive Spring** - If angular position drive is enabled, this is the stiffness of the drive spring.
-* **Angular Drive Damping** - If angular position drive is enabled, this is the damping of the drive spring.
-* **Angular Drive Force Limit** - If angular velocity drive is enabled, this is the maximum torque that will be applied to attain AngularVelocityTarget.
-
-
-You can find this, and more, information on the various constraint settings on the [Physics Constraint Reference](Engine/Physics/Physx/Constraints) page.
-
-
-
-## Miscellaneous
-
-Disabling collision between adjacent primitives is done automatically when the physics asset is created, as it is much faster to rely on the constraints to keep them separate than to use collision to do so.  
-
-Deleting all of the collision bodies for a given bone will cause the bone and constraint to be removed from the asset. A better method of doing this is to weld the bodies together using the Weld button (or by pressing _D_).
-
-
-
-## Previewing Animations
-
-It is possible to preview skeletal animations on the default skeletal mesh as well as blending between animations and physics simulation of the bodies of the physics asset. To do this, first you must assign an AnimSet containing the animations you wish to preview to the _Preview Anim Set_ in the editor proeprties (proeprties visible with nothing selected). This can be done simply by selecting the AnimSet in the Content Browser and then pressing the ![useselected.jpg](useselected.jpg)button next to the _Preview Anim Set_ property. Once assigned, the listbox in the toolbar should change from being empty to displaying _-None-_.
-
-At this point, the listbox is still disabled so no animation can be selected. To select an animation to play, you must begin the physics simulation preview by pressing the ![simstart.jpg](simstart.jpg)button in the toolbar. The listbox should now be enabled showing a list of all the animations contained in the _Preview Anim Set_.
-
-![animlist.jpg](animlist.jpg)
-
-You can now select an animation from the list and then press the ![animplay.jpg](animplay.jpg)button to cause the skeletal mesh to begin playing the animation. Of course, the physics simulation is taking precedence so you won't see much effect in the Preview Pane. To see the animation playing, you can adjust the _Physics Blend_ property in the editor properties. 
-
-![physweight.jpg](physweight.jpg)
-
-Another option is to enable _Blend On Poke_. This will cause the animation to have full weighting until the physics asset is 'poked' in the Preview Pane.
-
-![blendpoke.jpg](blendpoke.jpg)
-
- 
- -->
-
-
-
-
-
-
-
-
diff --git a/Engine/Documentation/Source/Engine/Physics/PhAT/Properties/Properties.INT.udn b/Engine/Documentation/Source/Engine/Physics/PhAT/Properties/Properties.INT.udn
deleted file mode 100644
diff --git a/Engine/Documentation/Source/Engine/Physics/PhAT/Reference/PhAT_Reference.INT.udn b/Engine/Documentation/Source/Engine/Physics/PhAT/Reference/PhAT_Reference.INT.udn
deleted file mode 100644
diff --git a/Engine/Documentation/Source/Engine/Physics/PhAT/Reference/Properties.INT.udn b/Engine/Documentation/Source/Engine/Physics/PhAT/Reference/Properties.INT.udn
new file mode 100644
--- /dev/null
+++ b/Engine/Documentation/Source/Engine/Physics/PhAT/Reference/Properties.INT.udn
@@ -0,0 +1,476 @@
+Availability: Public
+Title:Physics Asset Properties Reference
+Crumbs: %ROOT%, Engine, Engine/Physics, Engine/Physics/PhAT
+Description:A listing of properties and descriptions for Physics Bodies, Physics Constraints, and Physics Assets accessible in the Physics Asset Tool.
+Version: 4.9
+
+[VAR:Topic]
+[OBJECT:Topic]
+	[PARAM:image]
+		![%Engine/Physics/PhAT/Properties:title%](Engine/Physics/physics_topic.png)
+	[/PARAM]
+	[PARAM:icon]
+		![](%ROOT%/physics_icon.png)(convert:false)
+	[/PARAM]
+	[PARAM:title]
+		%Engine/Physics/PhAT/Properties:title%
+	[/PARAM]
+	[PARAM:description]
+		%Engine/Physics/PhAT/Properties:description%
+	[/PARAM]
+	[PARAM:path]
+		[RELATIVE:Engine/Physics/PhAT/Properties]
+	[/PARAM]
+[/OBJECT]
+[/VAR]
+[VAR:TopicCompact]
+[OBJECT:TopicCompact]
+	[PARAM:image]
+		![%Engine/Physics/PhAT/Properties:title%](Engine/Physics/physics_topic.png)
+	[/PARAM]
+	[PARAM:icon]
+		![](%ROOT%/physics_icon.png)(convert:false)
+	[/PARAM]
+	[PARAM:title]
+		%Engine/Physics/PhAT/Properties:title%
+	[/PARAM]
+	[PARAM:description]
+		%Engine/Physics/PhAT/Properties:description%
+	[/PARAM]
+	[PARAM:path]
+		[RELATIVE:Engine/Physics/PhAT/Properties]
+	[/PARAM]
+[/OBJECT]
+[/VAR]
+
+
+[VAR:spacer]![](%ROOT%/spacer.png)(w:150 convert:false)[/VAR]
+
+[VAR:spacer2]![](%ROOT%/spacer.png)(w:40 convert:false)[/VAR]
+
+This page contains a reference listing of properties you will find while working in the Physics Asset Tool.
+
+## Properties
+
+Below are the properties for Physics Asset Tool, separated into 3 expandable sections based on what is selected in the viewport.
+
+[OBJECT:modulesection]
+	[PARAMLITERAL:id]
+		PhatP
+	[/PARAMLITERAL]
+	[PARAM:heading]
+		Physics Asset Tool Properties
+	[/PARAM]
+	[PARAM:content]
+		%c1%
+	[/PARAM]
+[/OBJECT]
+
+[VAR:c1]
+
+When no Physics Body or Physics Constraint is selected (regardless of mode), the properties for the Physics Asset Tool are show in the Details panel.
+
+[OBJECT:modulesectionexpanded]
+	[PARAMLITERAL:id]
+		Anim1
+	[/PARAMLITERAL]
+	[PARAM:heading]
+		Anim
+	[/PARAM]
+	[PARAM:content]
+		%p11%
+	[/PARAM]
+[/OBJECT]
+
+[OBJECT:modulesectionexpanded]
+	[PARAMLITERAL:id]
+		Sim1
+	[/PARAMLITERAL]
+	[PARAM:heading]
+		Simulation
+	[/PARAM]
+	[PARAM:content]
+		%p12%
+	[/PARAM]
+[/OBJECT]
+
+[OBJECT:modulesectionexpanded]
+	[PARAMLITERAL:id]
+		MousS
+	[/PARAMLITERAL]
+	[PARAM:heading]
+		Mouse Spring
+	[/PARAM]
+	[PARAM:content]
+		%p13%
+	[/PARAM]
+[/OBJECT]
+
+[OBJECT:modulesectionexpanded]
+	[PARAMLITERAL:id]
+		Poke
+	[/PARAMLITERAL]
+	[PARAM:heading]
+		Poking
+	[/PARAM]
+	[PARAM:content]
+		%p14%
+	[/PARAM]
+[/OBJECT]
+
+[OBJECT:modulesectionexpanded]
+	[PARAMLITERAL:id]
+		Adv
+	[/PARAMLITERAL]
+	[PARAM:heading]
+		Advanced
+	[/PARAM]
+	[PARAM:content]
+		%p15%
+	[/PARAM]
+[/OBJECT]
+
+[OBJECT:modulesectionexpanded]
+	[PARAMLITERAL:id]
+		Draw
+	[/PARAMLITERAL]
+	[PARAM:heading]
+		Drawing
+	[/PARAM]
+	[PARAM:content]
+		%p16%
+	[/PARAM]
+[/OBJECT]
+
+[/VAR]
+
+[VAR:p11]
+
+![](PhATDef1.png)
+
+[REGION:raw]
+| Property | %spacer% | Description |
+| --- | :---: | :--- |
+| **Physics Blend** || %doxygen:UPhATSimOptions::PhysicsBlend% |
+| **Update Joints from Animation** || |
+| **Poke Pause Time** || %doxygen:UPhATSimOptions::PokePauseTime% |
+| **Poke Blend Time** || %doxygen:UPhATSimOptions::PokeBlendTime% |
+[/REGION]
+
+[/VAR]
+
+[VAR:p12]
+
+![](PhATDef2.png)
+
+[REGION:raw]
+| Property | %spacer% | Description |
+| --- | :---: | :--- |
+|**Floor Gap** || %doxygen:UPhATSimOptions::FloorGap% |
+|**Grav Scale** || %doxygen:UPhATSimOptions::GravScale% |
+|**Max FPS** || Max FPS for simulation in PhAT. This is helpful for targeting the same FPS as your game. -1 means this is disabled.  |
+|**Time Dilation** || This will dilate the time by scale.  |
+[/REGION]
+
+[/VAR]
+
+[VAR:p13]
+
+![](PhATDef3.png)
+
+[REGION:raw]
+| Property | %spacer% | Description |
+| --- | :---: | :--- |
+|**Handle Linear Damping** || %doxygen:UPhATSimOptions::HandleLinearDamping% |
+|**Handle Linear Stiffness** || %doxygen:UPhATSimOptions::HandleLinearStiffness% |
+|**Handle Angular Damping** || %doxygen:UPhATSimOptions::HandleAngularDamping% |
+|**Handle Angular Stiffness** || %doxygen:UPhATSimOptions::HandleAngularStiffness% |
+|**Interpolation Speed** || %doxygen:UPhATSimOptions::InterpolationSpeed%|
+[/REGION]
+
+[/VAR]
+
+[VAR:p14]
+
+![](PhATDef4.png)
+
+[REGION:raw]
+| Property | %spacer% | Description |
+| --- | :---: | :--- |
+| **Poke Strength** || %doxygen:UPhATSimOptions::PokeStrength% |
+[/REGION]
+
+[/VAR]
+
+[VAR:p15]
+
+![](PhATDef5.png)
+
+[REGION:raw]
+| Property | %spacer% | Description |
+| --- | :---: | :--- |
+| **Show Constraints as Points** || %doxygen:UPhATSimOptions::bShowConstraintsAsPoints% |
+| **Show Names in Hierarchy** || %doxygen:UPhATSimOptions::bShowNamesInHierarchy% |
+[/REGION]
+
+[/VAR]
+
+[VAR:p16]
+
+![](PhATDef6.png)
+
+[REGION:raw]
+| Property | %spacer% | Description |
+| --- | :---: | :--- |
+| **Constraint Draw Size** || %doxygen:UPhATSimOptions::ConstraintDrawSize% |
+[/REGION]
+
+[/VAR]
+
+---
+
+[OBJECT:modulesection]
+	[PARAMLITERAL:id]
+		BodyM
+	[/PARAMLITERAL]
+	[PARAM:heading]
+		Body Mode
+	[/PARAM]
+	[PARAM:content]
+		%c2%
+	[/PARAM]
+[/OBJECT]
+
+[VAR:c2]
+
+When the Physics Asset Tool is in Body Mode and a Physics Body is selected, the Details panel will display the properties for that Physics Body.
+
+[OBJECT:modulesectionexpanded]
+	[PARAMLITERAL:id]
+		Body1
+	[/PARAMLITERAL]
+	[PARAM:heading]
+		Body Setup
+	[/PARAM]
+	[PARAM:content]
+		%p21%
+	[/PARAM]
+[/OBJECT]
+
+[OBJECT:modulesectionexpanded]
+	[PARAMLITERAL:id]
+		Phys1
+	[/PARAMLITERAL]
+	[PARAM:heading]
+		Physics
+	[/PARAM]
+	[PARAM:content]
+		%p22%
+	[/PARAM]
+[/OBJECT]
+
+[OBJECT:modulesectionexpanded]
+	[PARAMLITERAL:id]
+		Coll
+	[/PARAMLITERAL]
+	[PARAM:heading]
+		Collision
+	[/PARAM]
+	[PARAM:content]
+		%p23%
+	[/PARAM]
+[/OBJECT]
+
+[/VAR]
+
+[VAR:p21]
+
+![](PhATBody1.png)
+
+[REGION:raw]
+| Property | %spacer% | Description |
+| --- | :---: | :--- |
+| **Bone Name** || %doxygen:UBodySetup::BoneName% |
+| **Consider for Bounds** || %doxygen:UBodySetup::bConsiderForBounds% |
+[/REGION]
+
+[/VAR]
+
+[VAR:p22]
+
+![](PhATBody2.png)
+
+[REGION:raw]
+| Property | %spacer% | Description |
+| --- | :---: | :--- |
+| **Mass in KG** || Mass of the body in KG. |
+| **Linear Damping** || %doxygen:FBodyInstance::LinearDamping% |
+| **Angular Damping** || %doxygen:FBodyInstance::AngularDamping% |
+| **Enable Gravity** || %doxygen:FBodyInstance::bEnableGravity% |
+| **Physics Type** || %doxygen:UBodySetup::PhysicsType% |
+| **Double Sided Geometry** || If true, the physics triangle mesh will use double sided faces when doing scene queries. Thsi is useful for planes and single sided meshes that need traces to work on both sides. |
+| **Simple Collision Physical Material** || %doxygen:UBodySetup::PhysMaterial% |
+[/REGION]
+
+### Advanced
+
+![](PhATBody2_adv.png)
+
+[REGION:raw]
+| Property | %spacer% | Description |
+| --- | :---: | :--- |
+| **Start Awake** || %doxygen:FBodyInstance::bStartAwake% |
+| **Center of Mass Offset** || %doxygen:FBodyInstance::COMNudge% |
+| **Mass Scale** || %doxygen:FBodyInstance::MassScale% |
+| **Max Angular Velocity** || The maximum angular velocity for this instance. |
+| **Use Async Scene** || %doxygen:FBodyInstance::bUseAsyncScene% |
+| **Generate Wake Events** ||  |
+| **Max Depenetration Velocity** || The maximum velocity used to depenetrate this object.  |
+| **Walkable Slope Override** || %doxygen:UBodySetup::WalkableSlopeOverride% |
+| **Walkable Slope Behavior** || Behavior of this surface (whether we affect the walkable slope). Determines how the Walkable Slope Angle operates on the existing walkable angle, whether it increases or decreases it. |
+| **Walkable Slope Angle** || %doxygen:FWalkableSlopeOverride::WalkableSlopeAngle% |
+| **Custom Sleep Threshold Multiplier** || |
+| **Position Solver Iteration Count** || %doxygen:FBodyInstance::PositionSolverIterationCount% |
+| **Velocity Solver Iteration Count** || %doxygen:FBodyInstance::VelocitySolverIterationCount% |
+| **Sleep Family** || %doxygen:FBodyInstance::SleepFamily% |
+[/REGION]
+
+[/VAR]
+
+[VAR:p23]
+
+![](PhATBody3.png)
+
+[REGION:raw]
+| Property | %spacer% | Description |
+| --- | :---: | :--- |
+| **Simulation Generates Hit Events** || %doxygen:FBodyInstance::bNotifyRigidBodyCollision% |
+| **Phys Material Override** || %doxygen:FBodyInstance::PhysMaterialOverride% |
+| **Collision Response** || %doxygen:UBodySetup::CollisionReponse% |
+
+[/REGION]
+### Advanced
+
+![](PhATBody3_adv.png)
+
+[REGION:raw]
+| Property | %spacer% | Description |
+| --- | :---: | :--- |
+| **Use CCD** || %doxygen:FBodyInstance::bUseCCD% |
+[/REGION]
+
+[/VAR]
+
+
+---
+
+[OBJECT:modulesection]
+	[PARAMLITERAL:id]
+		ConsM
+	[/PARAMLITERAL]
+	[PARAM:heading]
+		Constraint Mode
+	[/PARAM]
+	[PARAM:content]
+		%c3%
+	[/PARAM]
+[/OBJECT]
+
+[VAR:c3]
+
+When the Physics Asset Tool is in Body Mode and a Physics Body is selected, the Details panel will display the properties for that Physics Body.
+
+[OBJECT:modulesectionexpanded]
+	[PARAMLITERAL:id]
+		Cons2
+	[/PARAMLITERAL]
+	[PARAM:heading]
+		Constraint
+	[/PARAM]
+	[PARAM:content]
+		%p31%
+	[/PARAM]
+[/OBJECT]
+
+[OBJECT:modulesectionexpanded]
+	[PARAMLITERAL:id]
+		Line1
+	[/PARAMLITERAL]
+	[PARAM:heading]
+		Linear Limits
+	[/PARAM]
+	[PARAM:content]
+		%Engine/Physics/Constraints/ConstraintsReference:p2%
+	[/PARAM]
+[/OBJECT]
+
+[OBJECT:modulesectionexpanded]
+	[PARAMLITERAL:id]
+		Proj2
+	[/PARAMLITERAL]
+	[PARAM:heading]
+		Projection
+	[/PARAM]
+	[PARAM:content]
+		%Engine/Physics/Constraints/ConstraintsReference:p3%
+	[/PARAM]
+[/OBJECT]
+
+[OBJECT:modulesectionexpanded]
+	[PARAMLITERAL:id]
+		Angu2
+	[/PARAMLITERAL]
+	[PARAM:heading]
+		Angular Limits
+	[/PARAM]
+	[PARAM:content]
+		%Engine/Physics/Constraints/ConstraintsReference:p4%
+	[/PARAM]
+[/OBJECT]
+
+[OBJECT:modulesectionexpanded]
+	[PARAMLITERAL:id]
+		Lmoto
+	[/PARAMLITERAL]
+	[PARAM:heading]
+		Linear Motor
+	[/PARAM]
+	[PARAM:content]
+		%Engine/Physics/Constraints/ConstraintsReference:p5%
+	[/PARAM]
+[/OBJECT]
+
+[OBJECT:modulesectionexpanded]
+	[PARAMLITERAL:id]
+		Amoto
+	[/PARAMLITERAL]
+	[PARAM:heading]
+		Angular Motor
+	[/PARAM]
+	[PARAM:content]
+		%Engine/Physics/Constraints/ConstraintsReference:p6%
+	[/PARAM]
+[/OBJECT]
+
+[/VAR]
+
+[VAR:p31]
+
+![](PhATCon1.png)
+
+[REGION:raw]
+| Property | %spacer% | Description |
+| --- | :---: | :--- |
+| **Joint Name** || %doxygen:FConstraintInstance::JointName% |
+| **Constraint Bone 1** || %doxygen:FConstraintInstance::ConstraintBone1% |
+| **Constraint Bone 2** || %doxygen:FConstraintInstance::ConstraintBone2% |
+| **Disable Collision** || Disable collision between bodies joined by this constraint. |
+[/REGION]
+[/VAR]
+
+
+
+
+
+
+
diff --git a/Engine/Documentation/Source/Engine/Physics/PhAT/UserGuide/UserGuide.INT.udn b/Engine/Documentation/Source/Engine/Physics/PhAT/UserGuide/UserGuide.INT.udn
deleted file mode 100644
diff --git a/Engine/Documentation/Source/Engine/Physics/PhysicalMaterials/HowTo/AddSurfaceType/AddSurfaceType.INT.udn b/Engine/Documentation/Source/Engine/Physics/PhysicalMaterials/HowTo/AddSurfaceType/AddSurfaceType.INT.udn
new file mode 100644
--- /dev/null
+++ b/Engine/Documentation/Source/Engine/Physics/PhysicalMaterials/HowTo/AddSurfaceType/AddSurfaceType.INT.udn
@@ -0,0 +1,60 @@
+Availability: Public
+Title:Add a Surface Type
+Crumbs: %ROOT%, Engine, Engine/Physics, Engine/Physics/PhysicalMaterials
+Description:
+Version: 4.9
+SkillLevel: Intermediate
+
+[VAR:Topic]
+[OBJECT:Topic]
+	[PARAM:image]
+		![%Engine/Physics/PhysicalMaterials/HowTo/Editing:title%](Engine/Physics/physics_topic.png)
+	[/PARAM]
+	[PARAM:icon]
+		![](%ROOT%/physics_icon.png)(convert:false)
+	[/PARAM]
+	[PARAM:title]
+		%Engine/Physics/PhysicalMaterials/HowTo/Editing:title%
+	[/PARAM]
+	[PARAM:description]
+		%Engine/Physics/PhysicalMaterials/HowTo/Editing:description%
+	[/PARAM]
+	[PARAM:path]
+		[RELATIVE:Engine/Physics/PhysicalMaterials/HowTo/Editing]
+	[/PARAM]
+[/OBJECT]
+[/VAR]
+[VAR:TopicCompact]
+[OBJECT:TopicCompact]
+	[PARAM:image]
+		![%Engine/Physics/PhysicalMaterials/HowTo/Editing:title%](Engine/Physics/physics_topic.png)
+	[/PARAM]
+	[PARAM:icon]
+		![](%ROOT%/physics_icon.png)(convert:false)
+	[/PARAM]
+	[PARAM:title]
+		%Engine/Physics/PhysicalMaterials/HowTo/Editing:title%
+	[/PARAM]
+	[PARAM:description]
+		%Engine/Physics/PhysicalMaterials/HowTo/Editing:description%
+	[/PARAM]
+	[PARAM:path]
+		[RELATIVE:Engine/Physics/PhysicalMaterials/HowTo/Editing]
+	[/PARAM]
+[/OBJECT]
+[/VAR]
+
+The follwing steps detail how to add a Physical Surface Type to your project.
+
+1. Edit -> Project Settings... -> Physics -> Physical Surface Category
+1. Change the "None" in the Name field next to a "SurfaceType#" to a name descriptive of what you'd like the Surface Type to represent (concrete, flesh, wood, etc...)
+
+[REGION:raw]
+![](addSurfaceTop.png)(convert:false)
+[/REGION]
+
+
+
+
+
+
diff --git a/Engine/Documentation/Source/Engine/Physics/PhysicalMaterials/HowTo/Creation/Creation.INT.udn b/Engine/Documentation/Source/Engine/Physics/PhysicalMaterials/HowTo/Creation/Creation.INT.udn
new file mode 100644
--- /dev/null
+++ b/Engine/Documentation/Source/Engine/Physics/PhysicalMaterials/HowTo/Creation/Creation.INT.udn
@@ -0,0 +1,68 @@
+Availability: Public
+Title:Create a Physical Material
+Crumbs: %ROOT%, Engine, Engine/Physics, Engine/Physics/PhysicalMaterials
+Description:
+Version: 4.9
+SkillLevel: Basic
+
+[VAR:Topic]
+[OBJECT:Topic]
+	[PARAM:image]
+		![%Engine/Physics/PhysicalMaterials/HowTo/Creation:title%](Engine/Physics/physics_topic.png)
+	[/PARAM]
+	[PARAM:icon]
+		![](%ROOT%/physics_icon.png)(convert:false)
+	[/PARAM]
+	[PARAM:title]
+		%Engine/Physics/PhysicalMaterials/HowTo/Creation:title%
+	[/PARAM]
+	[PARAM:description]
+		%Engine/Physics/PhysicalMaterials/HowTo/Creation:description%
+	[/PARAM]
+	[PARAM:path]
+		[RELATIVE:Engine/Physics/PhysicalMaterials/HowTo/Creation]
+	[/PARAM]
+[/OBJECT]
+[/VAR]
+[VAR:TopicCompact]
+[OBJECT:TopicCompact]
+	[PARAM:image]
+		![%Engine/Physics/PhysicalMaterials/HowTo/Creation:title%](Engine/Physics/physics_topic.png)
+	[/PARAM]
+	[PARAM:icon]
+		![](%ROOT%/physics_icon.png)(convert:false)
+	[/PARAM]
+	[PARAM:title]
+		%Engine/Physics/PhysicalMaterials/HowTo/Creation:title%
+	[/PARAM]
+	[PARAM:description]
+		%Engine/Physics/PhysicalMaterials/HowTo/Creation:description%
+	[/PARAM]
+	[PARAM:path]
+		[RELATIVE:Engine/Physics/PhysicalMaterials/HowTo/Creation]
+	[/PARAM]
+[/OBJECT]
+[/VAR]
+
+
+1. From the Content Browser. **click** ![](Engine/Audio\SoundCues\Editor\button_CB_NewAsset.png) **New Asset** -> **Physics** -> **Physical Material** or **right-click** in the **Content Browser** -> **Physics** -> **Physical Material**.
+
+	![](Engine\Physics\PhysicalMaterials\HowTo\PhysMatCreate.jpg)
+
+1. **Double click** the new Physical Material to edit its properties.
+
+	![](Engine\Physics\PhysicalMaterials\HowTo\PhysMatProperties.jpg)(w:250)
+
+1. **Adjust Properties**.
+
+	![](Engine\Physics\PhysicalMaterials\HowTo\AdjustProperties.jpg)
+
+1. **Click Save** ![](Shared\Icons\icon_SaveAsset_40x.png)
+
+See the [**Physical Material Reference**](Engine/Physics\PhysicalMaterials\Reference) for information on the properties in the Physical Material.
+
+
+
+
+
+
diff --git a/Engine/Documentation/Source/Engine/Physics/PhysicalMaterials/HowTo/Editing/Editing.INT.udn b/Engine/Documentation/Source/Engine/Physics/PhysicalMaterials/HowTo/Editing/Editing.INT.udn
new file mode 100644
--- /dev/null
+++ b/Engine/Documentation/Source/Engine/Physics/PhysicalMaterials/HowTo/Editing/Editing.INT.udn
@@ -0,0 +1,65 @@
+Availability: Public
+Title:Edit a Physical Material
+Crumbs: %ROOT%, Engine, Engine/Physics, Engine/Physics/PhysicalMaterials
+Description:
+Version: 4.9
+SkillLevel: Intermediate
+
+[VAR:Topic]
+[OBJECT:Topic]
+	[PARAM:image]
+		![%Engine/Physics/PhysicalMaterials/HowTo/Editing:title%](Engine/Physics/physics_topic.png)
+	[/PARAM]
+	[PARAM:icon]
+		![](%ROOT%/physics_icon.png)(convert:false)
+	[/PARAM]
+	[PARAM:title]
+		%Engine/Physics/PhysicalMaterials/HowTo/Editing:title%
+	[/PARAM]
+	[PARAM:description]
+		%Engine/Physics/PhysicalMaterials/HowTo/Editing:description%
+	[/PARAM]
+	[PARAM:path]
+		[RELATIVE:Engine/Physics/PhysicalMaterials/HowTo/Editing]
+	[/PARAM]
+[/OBJECT]
+[/VAR]
+[VAR:TopicCompact]
+[OBJECT:TopicCompact]
+	[PARAM:image]
+		![%Engine/Physics/PhysicalMaterials/HowTo/Editing:title%](Engine/Physics/physics_topic.png)
+	[/PARAM]
+	[PARAM:icon]
+		![](%ROOT%/physics_icon.png)(convert:false)
+	[/PARAM]
+	[PARAM:title]
+		%Engine/Physics/PhysicalMaterials/HowTo/Editing:title%
+	[/PARAM]
+	[PARAM:description]
+		%Engine/Physics/PhysicalMaterials/HowTo/Editing:description%
+	[/PARAM]
+	[PARAM:path]
+		[RELATIVE:Engine/Physics/PhysicalMaterials/HowTo/Editing]
+	[/PARAM]
+[/OBJECT]
+[/VAR]
+
+The following steps detail how to edit an already existing Physical Material.
+
+1. **Double click** the new Physical Material to edit its properties.
+
+	![](Engine\Physics\PhysicalMaterials\HowTo\PhysMatProperties.jpg)(w:250)
+
+1. **Adjust Properties**.
+
+	![](Engine\Physics\PhysicalMaterials\HowTo\AdjustProperties.jpg)
+
+1. **Click Save** ![](Shared\Icons\icon_SaveAsset_40x.png)
+
+See the **[](Engine/Physics/PhysicalMaterials/Reference)** for information on the properties in the Physical Material.
+
+
+
+
+
+
diff --git a/Engine/Documentation/Source/Engine/Physics/PhysicalMaterials/HowTo/HowTo.INT.udn b/Engine/Documentation/Source/Engine/Physics/PhysicalMaterials/HowTo/HowTo.INT.udn
new file mode 100644
--- /dev/null
+++ b/Engine/Documentation/Source/Engine/Physics/PhysicalMaterials/HowTo/HowTo.INT.udn
@@ -0,0 +1,65 @@
+Availability: Public
+Title:Physical Material How To Guides
+Crumbs: %ROOT%, Engine, Engine/Physics, Engine/Physics/PhysicalMaterials
+Description:Various How To Guides focused on the creation and usage of Physical Materials.
+Version: 4.9
+SkillLevel: Intermediate
+
+[VAR:Topic]
+[OBJECT:Topic]
+	[PARAM:image]
+		![%Engine/Physics/PhysicalMaterials/HowTo:title%](Engine/Physics/physics_topic.png)
+	[/PARAM]
+	[PARAM:icon]
+		![](%ROOT%/physics_icon.png)(convert:false)
+	[/PARAM]
+	[PARAM:title]
+		%Engine/Physics/PhysicalMaterials/HowTo:title%
+	[/PARAM]
+	[PARAM:description]
+		%Engine/Physics/PhysicalMaterials/HowTo:description%
+	[/PARAM]
+	[PARAM:path]
+		[RELATIVE:Engine/Physics/PhysicalMaterials/HowTo]
+	[/PARAM]
+[/OBJECT]
+[/VAR]
+[VAR:TopicCompact]
+[OBJECT:TopicCompact]
+	[PARAM:image]
+		![%Engine/Physics/PhysicalMaterials/HowTo:title%](Engine/Physics/physics_topic.png)
+	[/PARAM]
+	[PARAM:icon]
+		![](%ROOT%/physics_icon.png)(convert:false)
+	[/PARAM]
+	[PARAM:title]
+		%Engine/Physics/PhysicalMaterials/HowTo:title%
+	[/PARAM]
+	[PARAM:description]
+		%Engine/Physics/PhysicalMaterials/HowTo:description%
+	[/PARAM]
+	[PARAM:path]
+		[RELATIVE:Engine/Physics/PhysicalMaterials/HowTo]
+	[/PARAM]
+[/OBJECT]
+[/VAR]
+
+This section of documentation contains a number of short "how to" guides focused on getting things done with **Physical Materials**.
+
+[OBJECT:TopicButtonList]
+	[PARAM:icon]
+		![](%ROOT%/start_icon.png)(convert:false)
+	[/PARAM]
+	[PARAM:title]
+		How To
+	[/PARAM]
+	[PARAM:links]
+		* [](Engine/Physics/PhysicalMaterials/HowTo/Creation)
+		* [](Engine/Physics/PhysicalMaterials/HowTo/Editing)
+		* [](Engine/Physics/PhysicalMaterials/HowTo/AddSurfaceType)
+		* [](Engine/Physics/PhysicalMaterials/HowTo/UsageMaterial)
+		* [](Engine/Physics/PhysicalMaterials/HowTo/UsageMaterialInstance)
+		* [](Engine/Physics/PhysicalMaterials/HowTo/UsagePhysAsset)
+		* [](Engine/Physics/PhysicalMaterials/HowTo/UsageStaticMesh)
+	[/PARAM]
+[/OBJECT]
\ No newline at end of file
diff --git a/Engine/Documentation/Source/Engine/Physics/PhysicalMaterials/HowTo/UsageMaterial/UsageMaterial.INT.udn b/Engine/Documentation/Source/Engine/Physics/PhysicalMaterials/HowTo/UsageMaterial/UsageMaterial.INT.udn
new file mode 100644
--- /dev/null
+++ b/Engine/Documentation/Source/Engine/Physics/PhysicalMaterials/HowTo/UsageMaterial/UsageMaterial.INT.udn
@@ -0,0 +1,60 @@
+Availability: Public
+Title:Assign a Physical Material to a Material
+Crumbs: %ROOT%, Engine, Engine/Physics, Engine/Physics/PhysicalMaterials
+Description:
+Version: 4.9
+SkillLevel: Intermediate
+
+[VAR:Topic]
+[OBJECT:Topic]
+	[PARAM:image]
+		![%Engine/Physics/PhysicalMaterials/HowTo/UsageMaterial:title%](Engine/Physics/physics_topic.png)
+	[/PARAM]
+	[PARAM:icon]
+		![](%ROOT%/physics_icon.png)(convert:false)
+	[/PARAM]
+	[PARAM:title]
+		%Engine/Physics/PhysicalMaterials/HowTo/UsageMaterial:title%
+	[/PARAM]
+	[PARAM:description]
+		%Engine/Physics/PhysicalMaterials/HowTo/UsageMaterial:description%
+	[/PARAM]
+	[PARAM:path]
+		[RELATIVE:Engine/Physics/PhysicalMaterials/HowTo/UsageMaterial]
+	[/PARAM]
+[/OBJECT]
+[/VAR]
+[VAR:TopicCompact]
+[OBJECT:TopicCompact]
+	[PARAM:image]
+		![%Engine/Physics/PhysicalMaterials/HowTo/UsageMaterial:title%](Engine/Physics/physics_topic.png)
+	[/PARAM]
+	[PARAM:icon]
+		![](%ROOT%/physics_icon.png)(convert:false)
+	[/PARAM]
+	[PARAM:title]
+		%Engine/Physics/PhysicalMaterials/HowTo/UsageMaterial:title%
+	[/PARAM]
+	[PARAM:description]
+		%Engine/Physics/PhysicalMaterials/HowTo/UsageMaterial:description%
+	[/PARAM]
+	[PARAM:path]
+		[RELATIVE:Engine/Physics/PhysicalMaterials/HowTo/UsageMaterial]
+	[/PARAM]
+[/OBJECT]
+[/VAR]
+
+
+1. **Open** or **create** a new Material.
+
+	![](Engine\Physics\PhysicalMaterials\HowTo\PhysMat_Material_01.png)
+
+1. **Select** the main material node.
+
+	![](Engine\Physics\PhysicalMaterials\HowTo\PhysMatNone.jpg)
+
+1. **Change** the Physical Material.
+
+	![](Engine\Physics\PhysicalMaterials\HowTo\PhysMatGlass2.jpg)
+
+1. **Click Save** ![](Shared\Icons\icon_SaveAsset_40x.png)
diff --git a/Engine/Documentation/Source/Engine/Physics/PhysicalMaterials/HowTo/UsageMaterialInstance/UsageMaterialInstance.INT.udn b/Engine/Documentation/Source/Engine/Physics/PhysicalMaterials/HowTo/UsageMaterialInstance/UsageMaterialInstance.INT.udn
new file mode 100644
--- /dev/null
+++ b/Engine/Documentation/Source/Engine/Physics/PhysicalMaterials/HowTo/UsageMaterialInstance/UsageMaterialInstance.INT.udn
@@ -0,0 +1,62 @@
+Availability: Public
+Title:Assign a Physical Material to a Material Instance
+Crumbs: %ROOT%, Engine, Engine/Physics, Engine/Physics/PhysicalMaterials
+Description:
+Version: 4.9
+SkillLevel: Intermediate
+
+[VAR:Topic]
+[OBJECT:Topic]
+	[PARAM:image]
+		![%Engine/Physics/PhysicalMaterials/HowTo/UsageMaterialInstance:title%](Engine/Physics/physics_topic.png)
+	[/PARAM]
+	[PARAM:icon]
+		![](%ROOT%/physics_icon.png)(convert:false)
+	[/PARAM]
+	[PARAM:title]
+		%Engine/Physics/PhysicalMaterials/HowTo/UsageMaterialInstance:title%
+	[/PARAM]
+	[PARAM:description]
+		%Engine/Physics/PhysicalMaterials/HowTo/UsageMaterialInstance:description%
+	[/PARAM]
+	[PARAM:path]
+		[RELATIVE:Engine/Physics/PhysicalMaterials/HowTo/UsageMaterialInstance]
+	[/PARAM]
+[/OBJECT]
+[/VAR]
+[VAR:TopicCompact]
+[OBJECT:TopicCompact]
+	[PARAM:image]
+		![%Engine/Physics/PhysicalMaterials/HowTo/UsageMaterialInstance:title%](Engine/Physics/physics_topic.png)
+	[/PARAM]
+	[PARAM:icon]
+		![](%ROOT%/physics_icon.png)(convert:false)
+	[/PARAM]
+	[PARAM:title]
+		%Engine/Physics/PhysicalMaterials/HowTo/UsageMaterialInstance:title%
+	[/PARAM]
+	[PARAM:description]
+		%Engine/Physics/PhysicalMaterials/HowTo/UsageMaterialInstance:description%
+	[/PARAM]
+	[PARAM:path]
+		[RELATIVE:Engine/Physics/PhysicalMaterials/HowTo/UsageMaterialInstance]
+	[/PARAM]
+[/OBJECT]
+[/VAR]
+
+
+1. **Open** or **create** a new Material Instance.
+
+	![](Engine\Physics\PhysicalMaterials\HowTo\PhysMat_MaterialInst_01.png)
+
+1. **Change** the Physical Material.
+
+	![](Engine\Physics\PhysicalMaterials\HowTo\StoneMat.jpg)
+
+1. **Click Save** ![](Shared\Icons\icon_SaveAsset_40x.png)
+
+
+
+
+
+
diff --git a/Engine/Documentation/Source/Engine/Physics/PhysicalMaterials/HowTo/UsagePhysAsset/UsagePhysAsset.INT.udn b/Engine/Documentation/Source/Engine/Physics/PhysicalMaterials/HowTo/UsagePhysAsset/UsagePhysAsset.INT.udn
new file mode 100644
--- /dev/null
+++ b/Engine/Documentation/Source/Engine/Physics/PhysicalMaterials/HowTo/UsagePhysAsset/UsagePhysAsset.INT.udn
@@ -0,0 +1,72 @@
+Availability: Public
+Title:Assign a Physical Material to a Physics Asset
+Crumbs: %ROOT%, Engine, Engine/Physics, Engine/Physics/PhysicalMaterials
+Description:How to assign a Physical Material to the entirety of a Physics Asset using the Physics Asset Tool.
+Version: 4.9
+SkillLevel: Intermediate
+
+[VAR:Topic]
+[OBJECT:Topic]
+	[PARAM:image]
+		![%Engine/Physics/PhysicalMaterials/HowTo/UsagePhysAsset:title%](Engine/Physics/physics_topic.png)
+	[/PARAM]
+	[PARAM:icon]
+		![](%ROOT%/physics_icon.png)(convert:false)
+	[/PARAM]
+	[PARAM:title]
+		%Engine/Physics/PhysicalMaterials/HowTo/UsagePhysAsset:title%
+	[/PARAM]
+	[PARAM:description]
+		%Engine/Physics/PhysicalMaterials/HowTo/UsagePhysAsset:description%
+	[/PARAM]
+	[PARAM:path]
+		[RELATIVE:Engine/Physics/PhysicalMaterials/HowTo/UsagePhysAsset]
+	[/PARAM]
+[/OBJECT]
+[/VAR]
+[VAR:TopicCompact]
+[OBJECT:TopicCompact]
+	[PARAM:image]
+		![%Engine/Physics/PhysicalMaterials/HowTo/UsagePhysAsset:title%](Engine/Physics/physics_topic.png)
+	[/PARAM]
+	[PARAM:icon]
+		![](%ROOT%/physics_icon.png)(convert:false)
+	[/PARAM]
+	[PARAM:title]
+		%Engine/Physics/PhysicalMaterials/HowTo/UsagePhysAsset:title%
+	[/PARAM]
+	[PARAM:description]
+		%Engine/Physics/PhysicalMaterials/HowTo/UsagePhysAsset:description%
+	[/PARAM]
+	[PARAM:path]
+		[RELATIVE:Engine/Physics/PhysicalMaterials/HowTo/UsagePhysAsset]
+	[/PARAM]
+[/OBJECT]
+[/VAR]
+
+
+The following steps detail how to set the Physical Material of all Physics Bodies in a Physics Asset at once.
+
+[EXCERPT:PhysMatSkelMesh1]
+1. **Open** a Physics Asset in the Physics Asset Tool.
+
+	![](Engine\Physics\PhysicalMaterials\HowTo\PhysicsAsset1.jpg)
+
+1. **Select** a Physical Material in the Content Browser.
+
+	![](Engine\Physics\PhysicalMaterials\HowTo\M_Glass.jpg)
+
+1. In the Physics Asset Tool, **click** the **Apply PhysMat** button from the **Asset** dropdown in the toolbar: ![](Engine\Physics\PhysicalMaterials\HowTo\AssetButton.jpg).
+	
+	[REGION:raw]
+	![](Engine\Physics\PhysicalMaterials\HowTo\ApplyPhysMat.jpg)
+	[/REGION]
+
+1. **Click Save** ![](Shared\Icons\icon_SaveAsset_40x.png)
+[/EXCERPT:PhysMatSkelMesh1]
+
+
+
+
+
+
diff --git a/Engine/Documentation/Source/Engine/Physics/PhysicalMaterials/HowTo/UsagePhysAssetBody/UsagePhysAssetBody.INT.udn b/Engine/Documentation/Source/Engine/Physics/PhysicalMaterials/HowTo/UsagePhysAssetBody/UsagePhysAssetBody.INT.udn
new file mode 100644
--- /dev/null
+++ b/Engine/Documentation/Source/Engine/Physics/PhysicalMaterials/HowTo/UsagePhysAssetBody/UsagePhysAssetBody.INT.udn
@@ -0,0 +1,77 @@
+Availability: Public
+Title:Assign a Physical Material to a Physics Asset Body
+Crumbs: %ROOT%, Engine, Engine/Physics, Engine/Physics/PhysicalMaterials
+Description:How to assign a Physical Material to a specific Physics Body in a Physics Asset using the Physics Asset Tool.
+Version: 4.9
+SkillLevel: Intermediate
+
+[VAR:Topic]
+[OBJECT:Topic]
+	[PARAM:image]
+		![%Engine/Physics/PhysicalMaterials/HowTo/UsagePhysAssetBody:title%](Engine/Physics/physics_topic.png)
+	[/PARAM]
+	[PARAM:icon]
+		![](%ROOT%/physics_icon.png)(convert:false)
+	[/PARAM]
+	[PARAM:title]
+		%Engine/Physics/PhysicalMaterials/HowTo/UsagePhysAssetBody:title%
+	[/PARAM]
+	[PARAM:description]
+		%Engine/Physics/PhysicalMaterials/HowTo/UsagePhysAssetBody:description%
+	[/PARAM]
+	[PARAM:path]
+		[RELATIVE:Engine/Physics/PhysicalMaterials/HowTo/UsagePhysAssetBody]
+	[/PARAM]
+[/OBJECT]
+[/VAR]
+[VAR:TopicCompact]
+[OBJECT:TopicCompact]
+	[PARAM:image]
+		![%Engine/Physics/PhysicalMaterials/HowTo/UsagePhysAssetBody:title%](Engine/Physics/physics_topic.png)
+	[/PARAM]
+	[PARAM:icon]
+		![](%ROOT%/physics_icon.png)(convert:false)
+	[/PARAM]
+	[PARAM:title]
+		%Engine/Physics/PhysicalMaterials/HowTo/UsagePhysAssetBody:title%
+	[/PARAM]
+	[PARAM:description]
+		%Engine/Physics/PhysicalMaterials/HowTo/UsagePhysAssetBody:description%
+	[/PARAM]
+	[PARAM:path]
+		[RELATIVE:Engine/Physics/PhysicalMaterials/HowTo/UsagePhysAssetBody]
+	[/PARAM]
+[/OBJECT]
+[/VAR]
+
+
+If specific Physics Bodies need differing Physical Materials, you can then adjust them individually.
+
+1. **Open** a Physics Asset in the Physics Asset Tool.
+
+	![](Engine\Physics\PhysicalMaterials\HowTo\PhysicsAsset1.jpg)
+
+1. **Enable Body Edit Mode** in the Physics Asset Tool.
+1. **Select** a Physics Body.
+1. In the Details pane, in the Physics Category, change the **Simple Collision Physical Material**.
+
+	[REGION:raw]
+	![](Engine\Physics\PhysicalMaterials\HowTo\SimpleCollision.png)(w:700)
+	[/REGION]
+
+1. **Click Save** ![](Shared\Icons\icon_SaveAsset_40x.png)
+
+<!--
+The default behavior for Physics interactions with Skeletal Meshes is to only interact with its associated Physics Asset, so the Skeletal Mesh's material's Physical Material will 
+not be used.
+
+[REGION:note]
+With Physics Assets, though you are setting the Simple Collision Physical Material, when you trace against a Physics Asset, you will need to do a complex trace and the complex trace 
+will return the Physical Material listed in the Simple Collision Physical Material property of the hit Physics Body.
+[/REGION]
+-->
+
+
+
+
+
diff --git a/Engine/Documentation/Source/Engine/Physics/PhysicalMaterials/HowTo/UsageStaticMesh/UsageStaticMesh.INT.udn b/Engine/Documentation/Source/Engine/Physics/PhysicalMaterials/HowTo/UsageStaticMesh/UsageStaticMesh.INT.udn
new file mode 100644
--- /dev/null
+++ b/Engine/Documentation/Source/Engine/Physics/PhysicalMaterials/HowTo/UsageStaticMesh/UsageStaticMesh.INT.udn
@@ -0,0 +1,69 @@
+Availability: Public
+Title:Assign a Physical Material in the Static Mesh Editor
+Crumbs: %ROOT%, Engine, Engine/Physics, Engine/Physics/PhysicalMaterials
+Description:Physical Materials User Guide: the creation and application of Physical Materials.
+Version: 4.9
+SkillLevel: Intermediate
+
+[VAR:Topic]
+[OBJECT:Topic]
+	[PARAM:image]
+		![%Engine/Physics/PhysicalMaterials/HowTo/UsageStaticMesh:title%](Engine/Physics/physics_topic.png)
+	[/PARAM]
+	[PARAM:icon]
+		![](%ROOT%/physics_icon.png)(convert:false)
+	[/PARAM]
+	[PARAM:title]
+		%Engine/Physics/PhysicalMaterials/HowTo/UsageStaticMesh:title%
+	[/PARAM]
+	[PARAM:description]
+		%Engine/Physics/PhysicalMaterials/HowTo/UsageStaticMesh:description%
+	[/PARAM]
+	[PARAM:path]
+		[RELATIVE:Engine/Physics/PhysicalMaterials/HowTo/UsageStaticMesh]
+	[/PARAM]
+[/OBJECT]
+[/VAR]
+[VAR:TopicCompact]
+[OBJECT:TopicCompact]
+	[PARAM:image]
+		![%Engine/Physics/PhysicalMaterials/HowTo/UsageStaticMesh:title%](Engine/Physics/physics_topic.png)
+	[/PARAM]
+	[PARAM:icon]
+		![](%ROOT%/physics_icon.png)(convert:false)
+	[/PARAM]
+	[PARAM:title]
+		%Engine/Physics/PhysicalMaterials/HowTo/UsageStaticMesh:title%
+	[/PARAM]
+	[PARAM:description]
+		%Engine/Physics/PhysicalMaterials/HowTo/UsageStaticMesh:description%
+	[/PARAM]
+	[PARAM:path]
+		[RELATIVE:Engine/Physics/PhysicalMaterials/HowTo/UsageStaticMesh]
+	[/PARAM]
+[/OBJECT]
+[/VAR]
+
+
+
+
+
+StaticMeshes have both Simple collision (the Physics Bodies that you create in your 3D art package or in the StaticMesh Editor) and Complex collision (per poly collision) which can be composed of a number of different materials each with their own differing Physical Materials. To set the Complex Collision Physical Material see: [](Engine/Physics/PhysicalMaterials/HowTo/UsageMaterial).
+
+To set the Simple Collision Physical Material for a StaticMesh:
+
+1. **Double Click** a Static Mesh in the Content Browser to bring up the Static Mesh Editor.
+
+	![](Engine\Physics\PhysicalMaterials\HowTo\PhysMat_StaticMeshEd_01.png)
+
+1. Change the Simple Collision Physical Material property in the **Static Mesh Settings** category to the desired Physical Material.
+
+	![](Engine\Physics\PhysicalMaterials\HowTo\PhysMat_StaticMeshEd_02.png)
+
+1. **Click Save** ![](Shared\Icons\icon_SaveAsset_40x.png).
+
+
+
+
+
+
diff --git a/Engine/Documentation/Source/Resources/ContentExamples/Physics/1_1/1_1.INT.udn b/Engine/Documentation/Source/Resources/ContentExamples/Physics/1_1/1_1.INT.udn
--- a/Engine/Documentation/Source/Resources/ContentExamples/Physics/1_1/1_1.INT.udn
+++ b/Engine/Documentation/Source/Resources/ContentExamples/Physics/1_1/1_1.INT.udn
@@ -3,7 +3,7 @@ Title:1.1 - Simulating Physics
 Crumbs:%ROOT%, Resources, Resources/ContentExamples, Resources/ContentExamples/Physics
 Description:An overview of the Physics example level, example 1.1: Object SImulating Physics
 Related: Engine/Physics/PhysicsBodies
-Related: Engine/Physics/PhAT/UserGuide
+Related: Engine/Physics/PhAT/HowTo
 Related: Engine/Physics/Constraints/ConstraintsReference
 Related: Engine/Physics/Constraints/ConstraintsUserGuide
 Related: Engine/Physics/Collision
diff --git a/Engine/Documentation/Source/Resources/ContentExamples/Physics/1_2/1_2.INT.udn b/Engine/Documentation/Source/Resources/ContentExamples/Physics/1_2/1_2.INT.udn
--- a/Engine/Documentation/Source/Resources/ContentExamples/Physics/1_2/1_2.INT.udn
+++ b/Engine/Documentation/Source/Resources/ContentExamples/Physics/1_2/1_2.INT.udn
@@ -3,7 +3,7 @@ Title:1.2 - SkeletalMesh RigidBody
 Crumbs:%ROOT%, Resources, Resources/ContentExamples, Resources/ContentExamples/Physics
 Description:An overview of the Physics example level, example 1.2: SkeletalMesh RigidBody Simulating Physics
 Related: Engine/Physics/PhysicsBodies
-Related: Engine/Physics/PhAT/UserGuide
+Related: Engine/Physics/PhAT/HowTo
 Related: Engine/Physics/Constraints/ConstraintsReference
 Related: Engine/Physics/Constraints/ConstraintsUserGuide
 Related: Engine/Physics/Collision
diff --git a/Engine/Documentation/Source/Resources/ContentExamples/Physics/1_3/1_3.INT.udn b/Engine/Documentation/Source/Resources/ContentExamples/Physics/1_3/1_3.INT.udn
--- a/Engine/Documentation/Source/Resources/ContentExamples/Physics/1_3/1_3.INT.udn
+++ b/Engine/Documentation/Source/Resources/ContentExamples/Physics/1_3/1_3.INT.udn
@@ -3,7 +3,7 @@ Title:1.3 - Radial Force / Impulse
 Crumbs:%ROOT%, Resources, Resources/ContentExamples, Resources/ContentExamples/Physics
 Description:An overview of the Physics example level, example 1.3: Radial Force / Impulse
 Related: Engine/Physics/PhysicsBodies
-Related: Engine/Physics/PhAT/UserGuide
+Related: Engine/Physics/PhAT/HowTo
 Related: Engine/Physics/Constraints/ConstraintsReference
 Related: Engine/Physics/Constraints/ConstraintsUserGuide
 Related: Engine/Physics/Collision
diff --git a/Engine/Documentation/Source/Resources/ContentExamples/Physics/1_4/1_4.INT.udn b/Engine/Documentation/Source/Resources/ContentExamples/Physics/1_4/1_4.INT.udn
--- a/Engine/Documentation/Source/Resources/ContentExamples/Physics/1_4/1_4.INT.udn
+++ b/Engine/Documentation/Source/Resources/ContentExamples/Physics/1_4/1_4.INT.udn
@@ -3,7 +3,7 @@ Title:1.4 - Thruster
 Crumbs:%ROOT%, Resources, Resources/ContentExamples, Resources/ContentExamples/Physics
 Description:An overview of the Physics example level, example 1.4: Thruster
 Related: Engine/Physics/PhysicsBodies
-Related: Engine/Physics/PhAT/UserGuide
+Related: Engine/Physics/PhAT/HowTo
 Related: Engine/Physics/Constraints/ConstraintsReference
 Related: Engine/Physics/Constraints/ConstraintsUserGuide
 Related: Engine/Physics/Collision
diff --git a/Engine/Documentation/Source/Resources/ContentExamples/Physics/1_5/1_5.INT.udn b/Engine/Documentation/Source/Resources/ContentExamples/Physics/1_5/1_5.INT.udn
--- a/Engine/Documentation/Source/Resources/ContentExamples/Physics/1_5/1_5.INT.udn
+++ b/Engine/Documentation/Source/Resources/ContentExamples/Physics/1_5/1_5.INT.udn
@@ -2,6 +2,12 @@ Availability: Public
 Title:1.5 - PhysicsConstraints
 Crumbs:%ROOT%, Resources, Resources/ContentExamples, Resources/ContentExamples/Physics
 Description:An overview of the Physics example level, example 1.5: PhysicsContraints
+Related: Engine/Physics/PhysicsBodies
+Related: Engine/Physics/PhAT/HowTo
+Related: Engine/Physics/Constraints/ConstraintsReference
+Related: Engine/Physics/Constraints/ConstraintsUserGuide
+Related: Engine/Physics/Collision
+Related: Engine/Physics/PhysicalMaterials/Reference
 
 [VAR:Steps]
 [OBJECT:Navigation]
@@ -67,7 +73,10 @@ Description:An overview of the Physics example level, example 1.5: PhysicsContra
 [/REGION]
 
 
-[INCLUDE:Engine/Physics\Constraints\ConstraintsUserGuide#PhysXConstraints_001]
+A constraint is basically a type of joint. It allows you to connect two Actors together (presumably one physically simulating), and also apply limits or forces. Unreal Engine 4 
+has a very flexible and data-driven constraint system that allows designers to make many different types of joints simply by changing some options. The engine has a few default 
+joint types (ball-and-socket, hinge, prismatic) but they only differ in their settings. You should feel free to start with any joint type and experiment with making your own 
+configurations.
 
 In the content example, working from left to right:
 
diff --git a/Engine/Documentation/Source/Resources/ContentExamples/Physics/1_6/1_6.INT.udn b/Engine/Documentation/Source/Resources/ContentExamples/Physics/1_6/1_6.INT.udn
--- a/Engine/Documentation/Source/Resources/ContentExamples/Physics/1_6/1_6.INT.udn
+++ b/Engine/Documentation/Source/Resources/ContentExamples/Physics/1_6/1_6.INT.udn
@@ -3,7 +3,7 @@ Title:1.6 - Angular Motors
 Crumbs:%ROOT%, Resources, Resources/ContentExamples, Resources/ContentExamples/Physics
 Description:An overview of the Physics example level, example 1.6: PhysicsConstraints Angular Motors
 Related: Engine/Physics/PhysicsBodies
-Related: Engine/Physics/PhAT/UserGuide
+Related: Engine/Physics/PhAT/HowTo
 Related: Engine/Physics/Constraints/ConstraintsReference
 Related: Engine/Physics/Constraints/ConstraintsUserGuide
 Related: Engine/Physics/Collision
diff --git a/Engine/Documentation/Source/Resources/ContentExamples/Physics/1_7/1_7.INT.udn b/Engine/Documentation/Source/Resources/ContentExamples/Physics/1_7/1_7.INT.udn
--- a/Engine/Documentation/Source/Resources/ContentExamples/Physics/1_7/1_7.INT.udn
+++ b/Engine/Documentation/Source/Resources/ContentExamples/Physics/1_7/1_7.INT.udn
@@ -3,7 +3,7 @@ Title:1.7 - Linear Motors
 Crumbs:%ROOT%, Resources, Resources/ContentExamples, Resources/ContentExamples/Physics
 Description:An overview of the Physics example level, example 1.7: PhysicsContraints Linear Motors
 Related: Engine/Physics/PhysicsBodies
-Related: Engine/Physics/PhAT/UserGuide
+Related: Engine/Physics/PhAT/HowTo
 Related: Engine/Physics/Constraints/ConstraintsReference
 Related: Engine/Physics/Constraints/ConstraintsUserGuide
 Related: Engine/Physics/Collision
diff --git a/Engine/Documentation/Source/Resources/ContentExamples/Physics/1_8/1_8.INT.udn b/Engine/Documentation/Source/Resources/ContentExamples/Physics/1_8/1_8.INT.udn
--- a/Engine/Documentation/Source/Resources/ContentExamples/Physics/1_8/1_8.INT.udn
+++ b/Engine/Documentation/Source/Resources/ContentExamples/Physics/1_8/1_8.INT.udn
@@ -3,7 +3,7 @@ Title:1.8 - Breakable Constraints
 Crumbs:%ROOT%, Resources, Resources/ContentExamples, Resources/ContentExamples/Physics
 Description:An overview of the Physics example level, example 1.8: Breakable Physics Constraints
 Related: Engine/Physics/PhysicsBodies
-Related: Engine/Physics/PhAT/UserGuide
+Related: Engine/Physics/PhAT/HowTo
 Related: Engine/Physics/Constraints/ConstraintsReference
 Related: Engine/Physics/Constraints/ConstraintsUserGuide
 Related: Engine/Physics/Collision
