#UE4 #DOCS: Walkable Slope doc update

diff --git a/Engine/Documentation/Source/Engine/Physics/WalkableSlope/WalkableSlope.INT.udn b/Engine/Documentation/Source/Engine/Physics/WalkableSlope/WalkableSlope.INT.udn
--- a/Engine/Documentation/Source/Engine/Physics/WalkableSlope/WalkableSlope.INT.udn
+++ b/Engine/Documentation/Source/Engine/Physics/WalkableSlope/WalkableSlope.INT.udn
@@ -1,7 +1,7 @@
 Availability: Public
 Title:Walkable Slope
 Crumbs: %ROOT%, Engine, Engine/Physics
-Description:
+Description:What the Walkable Slope Override is and the situations you may want to use it in.
 Version: 4.9
 
 [VAR:Topic]
@@ -43,7 +43,7 @@ Version: 4.9
 [/OBJECT]
 [/VAR]
 
-
+[TOC(start:2)]
 
 The **Walkable Slope Override** on **Physics Bodies** enable you to tweak what surfaces **Characters** can walk up. Perhaps a stair case is too steep or maybe you want to enforce the "no walking on the grass" signs, these settings will enable you to do so.
 
@@ -60,11 +60,7 @@ Depending on the editor you are in, or if you are looking at an Actor in the lev
 
 ### Increase Walkable Slope
 
-This setting will increase the walkable slope of a **Physics Body** up to the value listed in the **Walkable Slope Angle** property. If **Walkable Slope Angle** is set to 75, and the angle of the surface of the **Physics Body** is 65, the Character will be able to walk up it, regardless of the Character's **Walkable Floor Angle**:
-
-![](increased1.png)
-
-_The blue angle is the default **Character Movement Component Walkable Floor Angle**, while the green represents the new **Increased Walkable Slope** angle._
+This setting will increase the walkable slope of a **Physics Body** up to the value listed in the **Walkable Slope Angle** property. If **Walkable Slope Angle** is set to 75, and the angle of the surface of the **Physics Body** is 65, the Character will be able to walk up it, regardless of the Character's **Walkable Floor Angle**.
 
 While **Walkable Slope Behavior** is set to "Increase Walkable Slope":
 
@@ -73,18 +69,27 @@ While **Walkable Slope Behavior** is set to "Increase Walkable Slope":
 
 ### Decrease Walkable Slope
 
-This setting will "cap" the walkable slope of a **Physics Body** to the value listed in the **Walkable Slope Angle** property. If **Walkable Slope Angle** is set to 25, and the angle of the surface of the **Physics Body** is 35 degrees, the Character will be unable to walk up it regardless of the Character's **Walkable Floor Angle**:
-
-![](Decreased1.png)
+This setting will "cap" the walkable slope of a **Physics Body** to the value listed in the **Walkable Slope Angle** property. If **Walkable Slope Angle** is set to 25, and the angle of the surface of the **Physics Body** is 35 degrees, the Character will be unable to walk up it regardless of the Character's **Walkable Floor Angle**.
 
-_The blue angle is the default **Character Movement Component Walkable Floor Angle**, while the red represents the new **Decreased Walkable Slope** angle._
 
 While **Walkable Slope Behavior** is set to "Decrease Walkable Slope":
 
 * A value of 0.0 in the **Walkable Slope Angle** will result in the Character being unable to walk across the **Physics Body** surface. This can result in some odd behavior for mostly flat surfaces, as the character will skate across them but will not be able to change direction.
 * A value of 90.0 in the **Walkable Slope Angle** is essentially the same as **No Change** in the **Walkable Slope Behavior** property.
 
+## Examples
 
+[REGION:imagetable]
+| %inc% | %dec% |
+| --- | --- |
+| The blue angle is the default **Character Movement Component Walkable Floor Angle**, while the green represents the new **Increased Walkable Slope** angle. | The blue angle is the default **Character Movement Component Walkable Floor Angle**, while the red represents the new **Decreased Walkable Slope** angle. |
+[/REGION]
 
 
+[VAR:inc]
+[REGION:lightbox][![](increased.png)(h:470)](increased.png)[/REGION]
+[/VAR]
 
+[VAR:dec]
+[REGION:lightbox][![](Decreased.png)(h:470)](Decreased.png)[/REGION]
+[/VAR]
