Updated documentation on static shadows on translucency

diff --git a/Engine/Documentation/Source/Engine/Rendering/LightingAndShadows/LightMobility/StationaryLights/StationaryLights.INT.udn b/Engine/Documentation/Source/Engine/Rendering/LightingAndShadows/LightMobility/StationaryLights/StationaryLights.INT.udn
--- a/Engine/Documentation/Source/Engine/Rendering/LightingAndShadows/LightMobility/StationaryLights/StationaryLights.INT.udn
+++ b/Engine/Documentation/Source/Engine/Rendering/LightingAndShadows/LightMobility/StationaryLights/StationaryLights.INT.udn
@@ -124,6 +124,7 @@ Translucency also receives shadowing very cheaply with Stationary lights - Light
 	[/PARAM]
 [/OBJECT]
 
+The resolution of the static shadow depth map is controlled by StaticShadowDepthMapTransitionSampleDistanceX and StaticShadowDepthMapTransitionSampleDistanceY in BaseLightmass.ini, with a default setting of 100 meaning one texel every meter.
 
 ### Dynamic Shadowing
 
@@ -185,10 +186,10 @@ In the following image you see the difference between the two shadowing methods.
 
 [OBJECT:ComparisonSlider]
 	[PARAM:before]
-	![Area Shadows Enabled](Area_Shadows_On.png)(w:900)
+	![Area Shadows Enabled](AreaShadows.png)(w:900)
 	[/PARAM]
 	[PARAM:after]
-	![Area Shadows Disabled](Area_Shadows_Off.png)(w:900)
+	![Area Shadows Disabled](UniformPenumbraSharp.png)(w:900)
 	[/PARAM]
 [/OBJECT]
 
diff --git a/Engine/Documentation/Source/Engine/Rendering/LightingAndShadows/LitTranslucency/LitTranslucency.INT.udn b/Engine/Documentation/Source/Engine/Rendering/LightingAndShadows/LitTranslucency/LitTranslucency.INT.udn
--- a/Engine/Documentation/Source/Engine/Rendering/LightingAndShadows/LitTranslucency/LitTranslucency.INT.udn
+++ b/Engine/Documentation/Source/Engine/Rendering/LightingAndShadows/LitTranslucency/LitTranslucency.INT.udn
@@ -32,6 +32,11 @@ Translucent effects generally fall into a few categories: volumetric, volumetric
 
 Lit translucency gets most of its lighting through a series of cascaded volume textures oriented around the view frustum.  This allows lighting to be known in a single forward pass for any point inside the volumes, but has the downside that the volume texture is fairly low resolution, and can only cover a limited depth range from the viewer.  
 
+The volume is configured through cvars that can be set differently based on the scalability level:
+* r.TranslucencyLightingVolumeDim, which defaults to 64.  Raising this by a factor of 2 increases the cost to light volume by a factor of 8.
+* r.TranslucencyLightingVolumeInnerDistance, which defaults to 1500.  Raising this increases lighting volume coverage but reduces effective resolution.
+* r.TranslucencyLightingVolumeOuterDistance, which defaults to 5000.  Raising this increases lighting volume coverage but reduces effective resolution.
+
 Shadowed direct lighting from all movable light types is injected into the translucency lighting volume.  Light functions are also taken into account.
 
 [REGION:fullwidth]
@@ -58,6 +63,10 @@ Translucent self-shadowing uses per-object shadows, which means that it needs us
 | | |
 [/REGION]
 
+### Static shadowing
+
+Translucency can receive static shadowing from [stationary lights](Engine/Rendering/LightingAndShadows/LightMobility/StationaryLights) through a special static shadow depth map generated by Lightmass at lighting build time.
+
 ## Translucent Surfaces
 
 ### Reflection Captures
@@ -70,6 +79,5 @@ _The left sphere is translucent, the right is opaque, and they are both setup as
 
 ## Limitations
 
-* Static shadowing from static lights is currently not handled for lit translucency. However, dynamic shadows from stationary lights is supported.
 * Lit translucent surfaces are missing direct specular.
 * Lit translucent surfaces get all their direct lighting through the translucency volume lighting texture, which causes them to be lower resolution than needed for most surface materials (glass, water).
