#UE4doc #docs:
Added engine version field
Replaced almost all images on page due to being out of date.
Fixed some steps as they were referring to incorrect information.

diff --git a/Engine/Documentation/Source/Gameplay/SaveGame/Blueprints/BP_SaveGame.INT.udn b/Engine/Documentation/Source/Gameplay/SaveGame/Blueprints/BP_SaveGame.INT.udn
--- a/Engine/Documentation/Source/Gameplay/SaveGame/Blueprints/BP_SaveGame.INT.udn
+++ b/Engine/Documentation/Source/Gameplay/SaveGame/Blueprints/BP_SaveGame.INT.udn
@@ -2,12 +2,13 @@ Availability:Public
 Title: Saving Your Game with Blueprints
 Crumbs:%ROOT%,  Gameplay, Gameplay/SaveGame
 Description:How to save and load your game using Blueprints
+version: 4.9
 
 
 ### Create a SaveGame Object
 
-To create a new SaveGame object, [create a new Blueprint Class](Engine\Blueprints\UserGuide\Types\ClassBlueprint\Creation). When the Pick Parent Class dialog pops up, expand the **Custom Classes** 
-dropdown, then select **SaveGame**. You can use the search box to jump directly to SaveGame.
+To create a new SaveGame object, [create a new Blueprint Class](Engine\Blueprints\UserGuide\Types\ClassBlueprint\Creation). When the **Pick Parent Class** dialog pops up, expand the **Custom Classes** 
+dropdown, then select **SaveGame**. You can use the search box to jump directly to SaveGame. Name your new blueprint MySaveGame.
 
 ![](savegame.png) 
 
@@ -25,21 +26,20 @@ You can set default values for the variables after the Blueprint is compiled.
 
 ### Saving a Variable
 
-**This is a snapshot of the entire saving setup. You can click the image to enlarge it.**
-
 [REGION:lightbox]
 	[![](SaveGameBP.png)(w:880)](SaveGameBP.png)
 [/REGION]
+_This is a snapshot of the entire saving setup. You can click the image to enlarge it._
 
 [REGION:fullwidth]
 First, create an instance, or copy, of your SaveGame object, using the **Create Save Game Object** node. Make sure you set the Save Game Class dropdown to the name of your new SaveGame object Blueprint, 
 in this case, **MySaveGame**. Because the Create Save Game Object node creates a generic SaveGame object copy, you have to drag off of the **Return Value** pin and **Cast to MySaveGame**. Save
-the result of the casting to a variable so that you can reuse the SaveGame object easily later on.
+the result of the casting to a variable by using **Promote to Variable** so that you can reuse the SaveGame object easily later on.
 
 ![](SaveGameBP_1.png)
 
 Now that you have a copy of your custom SaveGame object, MySaveGame, saved to the **Save Game Instance** variable, you can send information to it. Here, the **Player Name** is being set to 
-"Player One". You would repeat this step for any other variables you wanted to store information in.
+"PlayerOne". You would repeat this step for any other variables you wanted to store information in.
 
 ![](SaveGameBP_2.png)
 
@@ -51,11 +51,10 @@ default values created earlier.
 
 ### Loading a Variable
 
-**This is a snapshot of the entire loading setup. You can click the image to enlarge it.**
-
 [REGION:lightbox]
 	[![](LoadGameBP.png)(w:880)](LoadGameBP.png)
 [/REGION]
+_This is a snapshot of the entire saving setup. You can click the image to enlarge it._
 
 [REGION:fullwidth]
 
