#UE4doc #docs:
Updated engine version field
Replaced an image
Fixed a few names of variables as they were inconsistent
Added a warning region field due to a compilation error

diff --git a/Engine/Documentation/Source/Gameplay/SaveGame/Code/Code_SaveGame.INT.udn b/Engine/Documentation/Source/Gameplay/SaveGame/Code/Code_SaveGame.INT.udn
--- a/Engine/Documentation/Source/Gameplay/SaveGame/Code/Code_SaveGame.INT.udn
+++ b/Engine/Documentation/Source/Gameplay/SaveGame/Code/Code_SaveGame.INT.udn
@@ -2,7 +2,7 @@ Availability:Public
 Title: Saving Your Game with C++
 Crumbs:%ROOT%, Gameplay, Gameplay/SaveGame
 Description:How to save and load your game with C++ code
-Version: 4.6
+Version: 4.9
 
 [TOC (start:2 end:2)]
 
@@ -57,7 +57,7 @@ MySaveGame.h
 		UPROPERTY(VisibleAnywhere, Category = Basic)
 		uint32 UserIndex;
 
-		UMySaveGame(const FObjectInitializer& ObjectInitializer);
+		UMySaveGame();
 	};
 
 ### Source
@@ -75,8 +75,7 @@ MySaveGame.cpp
 	#include "[ProjectName].h"
 	#include "MySaveGame.h"
 
-	UMySaveGame::UMySaveGame(const FObjectInitializer& ObjectInitializer)
-		: Super(ObjectInitializer)
+	UMySaveGame::UMySaveGame()
 	{
 		SaveSlotName = TEXT("TestSaveSlot");
 		UserIndex = 0;
@@ -85,19 +84,22 @@ MySaveGame.cpp
 ### Module Header File
 
 To be able to easily access both your `SaveGame` object and the creating, saving, and loading functions from `GameplayStatics`, you should add the following lines to your game module's header file,
-under any other `#include` statements.
+under any other `#include` statements. Your module header file will be named [ProjectName].h.
 
 [REGION:codetitlebar]
 MyProject.h
 [/REGION:codetitlebar]
 	#include "MySaveGame.h"
 	#include "Kismet/GameplayStatics.h"
+[REGION:warning]
+Ensure that these #include statements are added after `#include "Engine.h"` as the compilation will fail otherwise
+[/REGION]
 
 ## Saving a Variable
 
 Whenever you want to save out a variable to your `SaveGame` object, you must create an instance of your `SaveGame` object, and then set the variable within it.
 
-	MyPlayerName = TEXT("PlayerOne");
+	PlayerName = TEXT("PlayerOne");
 	UMySaveGame* SaveGameInstance = Cast<UMySaveGame>(UGameplayStatics::CreateSaveGameObject(UMySaveGame::StaticClass()));
 	SaveGameInstance->PlayerName = MyPlayerName;
 	UGameplayStatics::SaveGameToSlot(SaveGameInstance, SaveGameInstance->SaveSlotName, SaveGameInstance->UserIndex);
