#UE4doc #docs:
Added engine version field
Changed font settings for multiple terms
Fixed some capitalization issues.

diff --git a/Engine/Documentation/Source/Gameplay/UnrealTerminology/UnrealTerminology.INT.udn b/Engine/Documentation/Source/Gameplay/UnrealTerminology/UnrealTerminology.INT.udn
--- a/Engine/Documentation/Source/Gameplay/UnrealTerminology/UnrealTerminology.INT.udn
+++ b/Engine/Documentation/Source/Gameplay/UnrealTerminology/UnrealTerminology.INT.udn
@@ -2,6 +2,7 @@ Availability:Public
 Title:Unreal Projects and Gameplay
 Crumbs: %ROOT%, Gameplay
 Description:Introduction to projects, levels, classes, and Actors in Unreal Engine.
+version: 4.9
 
 Everything for your project is contained within the project directory.  You can create as many projects as you would like, but each one is self-contained. Using Unreal Engine's **Project Browser**
 to [create new projects](Engine/Basics/Projects/Browser#creatingnewprojects) will set up the necessary project framework such as the [directory structure](Engine\Basics\DirectoryStructure) and 
@@ -17,13 +18,13 @@ displayed in the **Viewport**.
 
 <!-- **Actors** are the most important concept in Unreal Engine 4. Your game will be filled with them, as every interactive entity in your game is an Actor. They usually contain one or more components, can be placed in a **Level** or **Spawned** in during gameplay, and support network replication for multiplayer games. -->
 
-At the most fundamental level, an **Actor** is a gameplay entity that (usually) contains one or more components, can be placed in a **Level** or **Spawned** in during gameplay, and support network replication for multiplayer games. In the above level, the ground platform, the blocks in the center of the level, and the PlayerStart visible at the bottom of the image are all Actors. All Actors extend from the AActor class, which is the base class of _spawnable_ gameplay Objects.
+At the most fundamental level, an **Actor** is a gameplay entity that (usually) contains one or more components, can be placed in a **Level** or **Spawned** in during gameplay, and support network replication for multiplayer games. In the above level, the ground platform, the blocks in the center of the level, and the PlayerStart visible at the bottom of the image are all Actors. All Actors extend from the AActor class, which is the base class of spawnable gameplay Objects.
 
-Actors can be thought of, in one sense, as containers that hold special types of **Objects** called components. For instance, a _CameraActor_ contains a _CameraComponent_.
+Actors can be thought of, in one sense, as containers that hold special types of **Objects** called **Components**. For instance, a CameraActor contains a CameraComponent.
 
 ![](camera_actor.png)
 
-The functionality of a camera, like the field of view, is all contained within the _CameraComponent_.  That means that the _CameraComponent_ can be included in other Actors, like a Character, 
+The functionality of a camera, like the field of view, is all contained within the CameraComponent.  That means that the CameraComponent can be included in other Actors, like a Character, 
 to give the same camera functionality to those Objects.
 
 ![](pawn_actor.png)
@@ -32,7 +33,7 @@ Different types of components can be used to control how Actors move, how they a
 UObject class, which is the base class of all gameplay Objects.  This means they cannot be directly instanced into the world; they must belong to an Actor.
 
 Each Actor or Object is a single instance of a class.  The class sets up the template for the Actor or Object. It defines the variables that can be set for that Actor or Object, and the functions that
-can be carried out within that Actor or Object.  You can create new classes, or types of Objects and Actors, with C++ code. _Blueprint Classes_ primarily allow you to create classes that set up new Actors, 
-although you can extend a few Objects with _Blueprint Classes_ as well. You can also combine the two, by creating a new C++ class and then making a _Blueprint Class_ derived from that C++ class. To learn more about 
+can be carried out within that Actor or Object.  You can create new classes, or types of Objects and Actors, with C++ code. Blueprint Classes primarily allow you to create classes that set up new Actors, 
+although you can extend a few Objects with Blueprint Classes as well. You can also combine the two, by creating a new C++ class and then making a Blueprint Class derived from that C++ class. To learn more about 
 creating classes so you can make new kinds of Actors and Objects, see the [Class Creation Basics](Gameplay\ClassCreation) page.
 
