#UE4doc #docs:
Unbolded a term in an excerpt

diff --git a/Engine/Documentation/Source/Engine/UI/LevelEditor/LevelEditor.INT.udn b/Engine/Documentation/Source/Engine/UI/LevelEditor/LevelEditor.INT.udn
--- a/Engine/Documentation/Source/Engine/UI/LevelEditor/LevelEditor.INT.udn
+++ b/Engine/Documentation/Source/Engine/UI/LevelEditor/LevelEditor.INT.udn
@@ -105,7 +105,7 @@ The **Level Editor** provides the core level creation functionality for Unreal E
 [/EXCERPT:Intro]
 
 [EXCERPT:Levels]
-In Unreal Editor, the scenes in which you create your game experience are generally referred to as [**Levels**](Engine/Levels). You can think of a level as a 3D environment into which you place a series of objects and geometry to define the world your players will experience. Any object that is placed in your world, be it a light, a mesh, or a character, is considered to be an Actor. Technically speaking, _Actor_ is a programming class used within the Unreal Engine to define an object that has 3D position, rotation, and scale data. For sake of ease, however, it will be easiest to think of an Actor as _any object that can be placed in your levels_.
+In Unreal Editor, the scenes in which you create your game experience are generally referred to as [Levels](Engine/Levels). You can think of a level as a 3D environment into which you place a series of objects and geometry to define the world your players will experience. Any object that is placed in your world, be it a light, a mesh, or a character, is considered to be an Actor. Technically speaking, _Actor_ is a programming class used within the Unreal Engine to define an object that has 3D position, rotation, and scale data. For sake of ease, however, it will be easiest to think of an Actor as _any object that can be placed in your levels_.
 [/EXCERPT:Levels]
 
 [EXCERPT:Actors]
