#UE4doc #docs:
Added engine version field

diff --git a/Engine/Documentation/Source/GettingStarted/GettingStarted.INT.udn b/Engine/Documentation/Source/GettingStarted/GettingStarted.INT.udn
--- a/Engine/Documentation/Source/GettingStarted/GettingStarted.INT.udn
+++ b/Engine/Documentation/Source/GettingStarted/GettingStarted.INT.udn
@@ -3,6 +3,7 @@ Title: Get Started with UE4
 Crumbs: %ROOT%
 Description:Introductory information for developers starting out creating games with Unreal Engine.
 Navigation:start
+version: 4.9
 
 [VAR:Params]
 [PARAM:image]
