#UE4doc #docs:
Added engine version field

diff --git a/Engine/Documentation/Source/GettingStarted/DownloadingUnrealEngine/DownloadingUnrealEngine4.INT.udn b/Engine/Documentation/Source/GettingStarted/DownloadingUnrealEngine/DownloadingUnrealEngine4.INT.udn
--- a/Engine/Documentation/Source/GettingStarted/DownloadingUnrealEngine/DownloadingUnrealEngine4.INT.udn
+++ b/Engine/Documentation/Source/GettingStarted/DownloadingUnrealEngine/DownloadingUnrealEngine4.INT.udn
@@ -5,7 +5,7 @@ Description:Steps for connecting to the source code repository and downloading t
 Prereq: GettingStarted/RecommendedSpecifications
 Related: GettingStarted/DownloadingUnrealEngine/GitHub
 Related: GettingStarted/DownloadingUnrealEngine/DownloadZip
-Version: 4.8
+Version: 4.9
 SkillLevel: Advanced
 
 [TOC (start:2 end:3)]
