#UE4doc #docs:
Added engine version field

diff --git a/Engine/Documentation/Source/GettingStarted/DownloadingUnrealEngine/Perforce/Perforce.INT.udn b/Engine/Documentation/Source/GettingStarted/DownloadingUnrealEngine/Perforce/Perforce.INT.udn
--- a/Engine/Documentation/Source/GettingStarted/DownloadingUnrealEngine/Perforce/Perforce.INT.udn
+++ b/Engine/Documentation/Source/GettingStarted/DownloadingUnrealEngine/Perforce/Perforce.INT.udn
@@ -2,6 +2,7 @@ Availability: Licensee
 Title:Using Perforce with Unreal Engine
 Crumbs:%ROOT%, Engine, Programming, Programming/Basics
 Description:Guide to using Perforce source control with Unreal Engine.
+version: 4.9
 
 [TOC(start:2)]
 
