#UE4doc #docs:
Added engine version field
Changed font settings for multiple terms
Fixed a few capitalization errors

diff --git a/Engine/Documentation/Source/GettingStarted/DownloadingUnrealEngine/Perforce/AutoMerge/AutoMerge.INT.udn b/Engine/Documentation/Source/GettingStarted/DownloadingUnrealEngine/Perforce/AutoMerge/AutoMerge.INT.udn
--- a/Engine/Documentation/Source/GettingStarted/DownloadingUnrealEngine/Perforce/AutoMerge/AutoMerge.INT.udn
+++ b/Engine/Documentation/Source/GettingStarted/DownloadingUnrealEngine/Perforce/AutoMerge/AutoMerge.INT.udn
@@ -2,10 +2,11 @@ Availability: Licensee
 Title:Perforce Auto-Merge Tool
 Crumbs:%ROOT%, GettingStarted, GettingStarted/DownloadingUnrealEngine, GettingStarted/DownloadingUnrealEngine/Perforce
 Description:Script that automates the process of performing an integration using Perforce.
+version: 4.9
 
 [TOC(start:2)]
 
-The Perforce AutoMerge script exists to automate the tedious and/or error-prone steps involved in performing a Perforce integration (aka merge). It incorporates several best-practices for merging as described in the [Perforce Integration](GettingStarted/DownloadingUnrealEngine/Perforce/Integration) page, so is the recommended way to merge in the Unreal Engine Perforce database.
+The Perforce Auto-Merge script exists to automate the tedious and/or error-prone steps involved in performing a Perforce integration (aka merge). It incorporates several best-practices for merging as described in the [Perforce Integration](GettingStarted/DownloadingUnrealEngine/Perforce/Integration) page, so is the recommended way to merge in the Unreal Engine Perforce database.
 
 
 
@@ -18,8 +19,8 @@ The screenshot below shows how it looks in P4V:
 ![P4AutoMergeCustomToolMenu.jpg](P4AutoMergeCustomToolMenu.jpg)
 
 
-* For any branch, you can choose a _Safe Merge_ or a _NULL Merge_ (described below). 
-* There is a _General_ tool that allows you to choose an arbitrary branch (with a pop-up dialog). 
+* For any branch, you can choose a **Safe Merge** or a **NULL Merge** (described below). 
+* There is a General tool that allows you to choose an arbitrary branch (with a pop-up dialog). 
 * You can add your own custom branches by matching the format used in the custom tool editor in P4V or P4Win.
 
 
@@ -44,14 +45,14 @@ The description contains:
 
 ### Safe Merge
 
-The _Safe Merge_ tool merges a changelist and performs an _auto-safe_ resolve, which only resolves files that have changed in one branch or the other, but not both. This type of resolve is guaranteed to succeed. When this process is finished, you will have a new changelist with the required checkin description already in place. If there are any files left to resolve (files that need to have a merge performed on them), you will have to handle them manually before checking in.
+The Safe Merge tool merges a changelist and performs an auto-safe resolve, which only resolves files that have changed in one branch or the other, but not both. This type of resolve is guaranteed to succeed. When this process is finished, you will have a new changelist with the required checkin description already in place. If there are any files left to resolve (files that need to have a merge performed on them), you will have to handle them manually before checking in.
 
 See the [Perforce Integration](GettingStarted/DownloadingUnrealEngine/Perforce/Integration#Resolve) page for help with manual resolves.
 
 
 ### NULL Merge
 
-The _NULL Merge_ tool merges a changelist and performs an _accept yours_ resolve, which means that the files will be integrated, but will not be changed. This will unambiguously indicate that a changelist has been considered for integration, but it was decided not to incorporate those changes.
+The NULL Merge tool merges a changelist and performs an accept yours resolve, which means that the files will be integrated, but will not be changed. This will unambiguously indicate that a changelist has been considered for integration, but it was decided not to incorporate those changes.
 
 For more information on NULL merges, see the [Perforce Integration](GettingStarted/DownloadingUnrealEngine/Perforce/Integration#NullMerge) page.
 
@@ -85,11 +86,11 @@ The script is a Python script that uses the P4Python API, so you will need to in
 
 * For Python:
     * The [ActiveState ActivePython](http://www.activestate.com/activepython/) distribution is recommended.
-    * The script is currently tested with Python version **2.5.2**.
+    * The script is currently tested with Python version 2.5.2.
 * For P4Python:
     * You can download the latest P4Python from the [Related Software](http://www.perforce.com/perforce/loadsupp.html) page of the Perforce site. 
     * You can get the [API documentation](http://www.perforce.com/perforce/doc.current/manuals/p4script/p4script.pdf) from the Perforce site.
-    * The script is currently tested with P4Python version **2008.2**.
+    * The script is currently tested with P4Python version 2008.2.
 
 
 
@@ -98,17 +99,17 @@ The script is a Python script that uses the P4Python API, so you will need to in
 Once Python and P4Python is installed:
 
 * Download the [](P4AutoMerge.py.txt) python script and place it in your Python path.
-* In the P4V _Tools_ menu, select _Manage custom tools..._
-* Select _New > Tool..._ (optionally creating a folder to group your tools).
+* In the P4V **Tools** menu, select **Manage custom tools...**
+* Select **New > Tool...** (optionally creating a folder to group your tools).
 * Set up the options like the image below.
     * Replace the indicated string with your branch name.
-    * Click the _select..._ button for more details on the options available for P4V's custom tools.
+    * Click the **select...** button for more details on the options available for P4V's custom tools.
 
 
 ![P4AutoMergeSetup.jpg](P4AutoMergeSetup.jpg)
 
 
-* From there, you should be able to **right-click** on a submitted changelist and select your tool.
+* From there, you should be able to **Right-click** on a submitted changelist and select your tool.
 
 
 
