#UE4doc #docs:
Added engine version field
Changed grammar for one line

diff --git a/Engine/Documentation/Source/GettingStarted/DownloadingUnrealEngine/Perforce/CommonMergeProblems/CommonMergeProblems.INT.udn b/Engine/Documentation/Source/GettingStarted/DownloadingUnrealEngine/Perforce/CommonMergeProblems/CommonMergeProblems.INT.udn
--- a/Engine/Documentation/Source/GettingStarted/DownloadingUnrealEngine/Perforce/CommonMergeProblems/CommonMergeProblems.INT.udn
+++ b/Engine/Documentation/Source/GettingStarted/DownloadingUnrealEngine/Perforce/CommonMergeProblems/CommonMergeProblems.INT.udn
@@ -2,6 +2,7 @@ Availability: Licensee
 Title:Common Merge Problems
 Crumbs:%ROOT%, Engine, Programming, Programming/Basics, GettingStarted/DownloadingUnrealEngine/Perforce
 Description:Explanation of many of the common problems encountered when merging new builds of the engine.
+version: 4.9
 
 [TOC(start:2)]
 
@@ -59,7 +60,7 @@ You can then wrap that around custom code that is in base engine code:
 
 Ideally, you would want to be able to undefine everything and have a generic build; but that is not always possible.
 
-However, these techniques should make merging a lot more bearable... perhaps even enjoyable?
+However, these techniques should make merging a lot more stable.
 
 
 
