#UE4doc #docs:
Added engine version field
Changed font settings for many terms
Added note regions in multiple places

diff --git a/Engine/Documentation/Source/GettingStarted/DownloadingUnrealEngine/Perforce/Integration/Integration.INT.udn b/Engine/Documentation/Source/GettingStarted/DownloadingUnrealEngine/Perforce/Integration/Integration.INT.udn
--- a/Engine/Documentation/Source/GettingStarted/DownloadingUnrealEngine/Perforce/Integration/Integration.INT.udn
+++ b/Engine/Documentation/Source/GettingStarted/DownloadingUnrealEngine/Perforce/Integration/Integration.INT.udn
@@ -2,13 +2,15 @@ Availability: Licensee
 Title:Perforce Integration FAQ
 Crumbs:%ROOT%, Engine, Programming, Programming/Basics, GettingStarted/DownloadingUnrealEngine/Perforce
 Description:Commonly aksed questions about Perforce terminology and usage details.
+version: 4.9
 
 [TOC(start:2)]
 
 This document provides an understanding of many of the common uses of Perforce with Unreal Engine 4, as well as related terminology.
 
-_[Note: To help simplify the merge process, an auto-merge tool has been created that can be integrated into P4V and P4Win. It automates several of the tedious steps and best practices described in this doc. For more information, see the [Perforce AutoMerge](GettingStarted/DownloadingUnrealEngine/Perforce/AutoMerge) page for more information.]_
-
+[REGION:note]
+To help simplify the merge process, an auto-merge tool has been created that can be integrated into P4V and P4Win. It automates several of the tedious steps and best practices described in this doc. For more information, see the [Perforce AutoMerge](GettingStarted/DownloadingUnrealEngine/Perforce/AutoMerge) page for more information.
+[/REGION]
 
 
 ## Frequently Asked Questions
@@ -16,7 +18,7 @@ _[Note: To help simplify the merge process, an auto-merge tool has been created
 
 ### What does integrate, merge, and branch mean in Perforce?
 
-They all mean the same thing and are accomplished using the _Integrate_ command. A _Branch_ is usually referred to as an integration that creates a new file in a new location. A _Merge_ is an integrate between two previously-branched files.
+They all mean the same thing and are accomplished using the Integrate command. A Branch is usually referred to as an integration that creates a new file in a new location. A Merge is an integrate between two previously-branched files.
 
 
 ### Why do all these complicated steps when I can just copy the file over and change it?
@@ -28,25 +30,27 @@ Perforce keeps a searchable, graphable database of all changes and integrations
 
 ### How do I integrate a single changelist?
 
-_Note: It is highly recommended to use P4V when integrating, as there have been many improvements to the UI to visualize the process. Everything ultimately resolves to the same underlying P4API calls, so both P4V and P4Win support the same feature set._
+[REGION:note]
+It is highly recommended to use P4V when integrating, as there have been many improvements to the UI to visualize the process. Everything ultimately resolves to the same underlying P4API calls, so both P4V and P4Win support the same feature set.
+[/REGION]
 
 To integrate a single change in P4V:
 
 
 * Make sure your active workspace contains the branch you are integrating to.
     * There is no need to include the source branch in your workspace. 
-* **Right-click** on the changelist you want to integrate and select _Integrate using submitted changelist 'change#'..._. You should see a dialog similar to this:
+* **Right-click** on the changelist you want to integrate and select **Integrate using submitted changelist 'change#'...**. You should see a dialog similar to this:
 
 
 (#IntegrationShot)
 ![P4VIntegrate.jpg](P4VIntegrate.jpg)
 
 
-* Under _Choose integration method_, select _Use branch specification_ and select the branch you want to use.
+* Under **Choose integration method**, select **Use branch specification** and select the branch you want to use.
 * If you deleted or re-added a file: 
-    * Click on _Advanced options..._ and section by clicking on the triangle icon. 
-    * Under _enable integrations around deleted revisions_, select the appropriate options.
-    * Generally you want _Delete target file when source is deleted._
+    * Click on **Advanced options...** and section by clicking on the triangle icon. 
+    * Under **Enable integrations around deleted revisions**, select the appropriate options.
+    * Generally you want **Delete target file when source is deleted**.
     * Do not select any other options unless you are sure of what you are doing!
 
 
@@ -56,7 +60,7 @@ To integrate a single change in P4V:
 * Always use the preview button to check for errors. You cannot complete the integration if there are errors.
     * P4V now supports visual inspection of integration errors.
     * You will see a treeview of all files integrated, with files that cannot be integrated highlighted in red.
-    * You can use the red arrows to iterate through the errors. The _Description_ field will explain exactly why each file cannot be integrated.
+    * You can use the red arrows to iterate through the errors. The Description field will explain exactly why each file cannot be integrated.
     * You can also switch between a classic flat view or the newer tree view.
     * If you cannot figure out the errors, get someone to help you.
 
@@ -64,8 +68,8 @@ To integrate a single change in P4V:
 ![P4VPreview.JPG](P4VPreview.JPG)
 
 
-* Select _New_ changelist to ensure you do not mix existing changes with this merge.
-* click _Finish_ to create the new changelist. 
+* Select **New changelist** to ensure you do not mix existing changes with this merge.
+* click **Finish** to create the new changelist. 
 * Resolve the integrated files by following the [instructions below](#Resolve). 
 * Test everything thoroughly before submitting.
 * Remember to use the [changelist description guidelines](#ChangeDesc) mentioned later in this guide.
@@ -73,16 +77,16 @@ To integrate a single change in P4V:
 
 To integrate a single change in P4Win:
 
-* Open Revision History and **right-click** on the changelist you want to merge, then select _Integrate Using|Branchspec..._
+* Open **Revision History** and **Right-click** on the changelist you want to merge, then select **Integrate Using|Branchspec...**
 * Select the branchspec from the list. 
 * If you deleted or re-added a file: 
-    * Select _Options...|Permit deletes/re-adds_. 
-    * Select _deletes only_, _re-adds only_ or _both_ radio button as appropriate. 
+    * Select **Options...|Permit deletes/re-adds**. 
+    * Select **deletes only**, **re-adds only** or both radio button as appropriate. 
     * Do not select any other options unless you are sure of what you are doing!
 * P4Win does not have a visualization tool for previewing, so you will have to read the text in the output window.
     * If you are having trouble with this, get someone to help you. 
-* Select _New_ changelist to ensure you do not mix existing changes with this merge 
-* Click _Finish_ to create the new changelist. 
+* Select **New changelist** to ensure you do not mix existing changes with this merge 
+* Click **Finish** to create the new changelist. 
 * Resolve the integrated files by following the [instructions below](#Resolve). 
 * Test everything thoroughly before submitting.
 * Remember to use the [changelist description guidelines](#ChangeDesc) mentioned later in this guide.
@@ -94,9 +98,9 @@ To integrate a single change in P4Win:
 ### What is the safest way to resolve files quickly?
 
 Each file integrated MUST be resolved before checking in. Difficult conflicts will sometimes require help from the original [TOC(start:2)]. To ease the process, Perforce has a few options for auto-resolving simple merges to pare the list down. P4V and P4Win both provide access to these options. The safest way to proceed is in the following order:
-1. _Safe automatic_ resolve - auto-resolve changes if the file has only changed in one branch since the last merge. This is always safe because it can just copy the changed file over.
-1. _Automatic_ resolve - resolves changes if the file has changed in both branches but no conflicts exist.
-1. _Interactive_ resolve - any files left will be the hard ones where changes in both branches conflicted.
+1. **Safe automatic resolve** - auto-resolve changes if the file has only changed in one branch since the last merge. This is always safe because it can just copy the changed file over.
+1. **Automatic resolve** - resolves changes if the file has changed in both branches but no conflicts exist.
+1. **Interactive resolve** - any files left will be the hard ones where changes in both branches conflicted.
     * To deal with these, you will use a merge tool such as:
         * [Araxis Merge](http://www.araxis.com/), or 
         * P4WinMerge.
@@ -105,7 +109,7 @@ Each file integrated MUST be resolved before checking in. Difficult conflicts wi
 
 Notes:
 
-* _Safe Automatic_ resolve will often take care of most files automatically.
+* Safe Automatic resolve will often take care of most files automatically.
 * Be very careful when merging files interactively; ask for help if you need to! 
 
 
@@ -115,7 +119,7 @@ P4V will display a summary count of the changes and conflicts found, and give yo
 
 ![P4VResolveIndividual.jpg](P4VResolveIndividual.jpg)
 
-When you are ready, click _Run Merge Tool_ to run your registered merge tool to perform the actual merge. P4V's built-in merge tool is recommended because it gives a 4-way pane that always displays the source, base, target, and merged result, color coded for the source of each change. Many other merge tools (like Araxis) only display 3 panes, so the base version is overwritten as you perform the merge. For more details on the color/icon scheme that P4Merge uses, see the P4Merge Help Viewer (press _F1_ while merging).
+When you are ready, click **Run Merge Tool** to run your registered merge tool to perform the actual merge. P4V's built-in merge tool is recommended because it gives a 4-way pane that always displays the source, base, target, and merged result, color coded for the source of each change. Many other merge tools (like Araxis) only display 3 panes, so the base version is overwritten as you perform the merge. For more details on the color/icon scheme that P4Merge uses, see the P4Merge Help Viewer (press F1 while merging).
 
 ![P4VMerge.jpg](P4VMerge.jpg)
 
@@ -144,7 +148,7 @@ Please do not. It is better if you use a separate submission for each integrated
 
 You should still perform the integration to inform the Perforce database that the change has been integrated. You should perform the integration as normal, but during the resolve phase, you should "accept yours". This will not change the local files, but will tell Perforce that those files have been integrated. 
 
-Then, before checking in, you should _reopen for edit_ the appropriate files and make the required changes in the new code. That way you create a single changelist that represents the integration of that change into your codeline.
+Then, before checking in, you should reopen for edit the appropriate files and make the required changes in the new code. That way you create a single changelist that represents the integration of that change into your codeline.
 
 
 ### There have been many changes to a file in either the source or destination branch since the changelist that I am interested in was submitted. Is it OK to merge this change?
@@ -161,14 +165,16 @@ It is not recommended that you do this, unless you have a really good reason to.
 
 No. You should always check in all files that you integrated from a changelist. Even if the file is not changing, checking it in will update the Perforce database with the integration information.
 
-Note: Perforce stores integration history on a per-file basis, not a per-changelist basis. If you integrate a changelist but then revert some files, it makes analyzing the database more difficult because some changelists will be listed as only having been partially integrated.
+[REGION:note]
+Perforce stores integration history on a per-file basis, not a per-changelist basis. If you integrate a changelist but then revert some files, it makes analyzing the database more difficult because some changelists will be listed as only having been partially integrated.
+[/REGION]
 
 (#NullMerge)
 
 
 ### I have decided that a changelist should not be integrated. Should I do anything?
 
-This is known as a _NULL Merge_. If you are absolutely sure that a change you submit will never be needed in the other branch, it is often helpful to go ahead and perform a NULL merge anyway. To perform a NULL merge, integrate the files and resolve them by selecting "accept yours." This will create an official record that that particular change has already been examined by someone and nothing was done. This simplifies the process for those maintaining the branch because that change will not be included in the list of pending changes to integrate in that branch. 
+This is known as a NULL Merge. If you are absolutely sure that a change you submit will never be needed in the other branch, it is often helpful to go ahead and perform a NULL merge anyway. To perform a NULL merge, integrate the files and resolve them by selecting "accept yours." This will create an official record that that particular change has already been examined by someone and nothing was done. This simplifies the process for those maintaining the branch because that change will not be included in the list of pending changes to integrate in that branch. 
 
 Whether this is necessary depends on the policy of the codeline and how integrations are being managed, so you should check with the curator of that codeline to see if this should be done for your changes.
 
@@ -181,8 +187,8 @@ Branchspecs should be set up so the normal flow of change is the default directi
 
 Perforce supports "reversing" a branchspec during integration:
 
-* In P4V, you can simply click on the green arrow in the _Select branch specification_ section. See the [screenshot](#IntegrationShot) in the [integration item](#IntegrateSection). 
-    * There is even text there in case you forget. :)
+* In P4V, you can simply click on the green arrow in the Select branch specification section. See the [screenshot](#IntegrationShot) in the [integration item](#IntegrateSection). 
+    * There is even text there in case you forget.
 * P4Win seems to automatically detect reverse requests.
     * Just integrate the changelist using the same branchspec and it will automatically determine that it should reverse the mapping.
 
