#UE4doc #docs - Added engine version field.

diff --git a/Engine/Documentation/Source/Resources/ContentExamples/LevelDesign/LevelDesign.INT.udn b/Engine/Documentation/Source/Resources/ContentExamples/LevelDesign/LevelDesign.INT.udn
--- a/Engine/Documentation/Source/Resources/ContentExamples/LevelDesign/LevelDesign.INT.udn
+++ b/Engine/Documentation/Source/Resources/ContentExamples/LevelDesign/LevelDesign.INT.udn
@@ -6,7 +6,7 @@ Related: Engine/Actors/Brushes
 Related:Engine/UI/LevelEditor
 Related: Engine/UI/LevelEditor/Modes
 Related: Engine/Actors/Volumes
-
+Version: 4.9
 
 
 [REGION:banner]
@@ -69,9 +69,9 @@ Each section below provides more information on what is included in each phase o
 
 | Example | What is Demonstrated |
 | ------- | ------------ |
-| **[1.1 Prototype Pass](Resources\ContentExamples\LevelDesign\1)** | Quick block out of a playable level using basic **Static Mesh Actors** or **Geometry Primitives** with rudamentry lighting. |
+| **[1.1 Prototype Pass](Resources\ContentExamples\LevelDesign\1)** | Quick block out of a playable level using basic Static Mesh Actors or Geometry Primitives with rudamentry lighting. |
 | **[1.2 Meshing Pass](Resources\ContentExamples\LevelDesign\2)** | Replacing the basic meshes or primitives from the prototype phase with near final asset and applying basic materials.  |
-| **[1.3 Lighting Pass](Resources\ContentExamples\LevelDesign\3)** | Emphasis is placed on improving lighting, adding **PostProcess** effects and updating materials. |
+| **[1.3 Lighting Pass](Resources\ContentExamples\LevelDesign\3)** | Emphasis is placed on improving lighting, adding PostProcess effects and updating materials. |
 | **[1.4 Polish Pass](Resources\ContentExamples\LevelDesign\4)** | Additional effects, audio and volumes are added, and final assets and details are tweaked. |
 
 
