#UE4doc #docs - added engine version field.

diff --git a/Engine/Documentation/Source/Resources/ContentExamples/LevelDesign/1/1.INT.udn b/Engine/Documentation/Source/Resources/ContentExamples/LevelDesign/1/1.INT.udn
--- a/Engine/Documentation/Source/Resources/ContentExamples/LevelDesign/1/1.INT.udn
+++ b/Engine/Documentation/Source/Resources/ContentExamples/LevelDesign/1/1.INT.udn
@@ -6,6 +6,7 @@ Related: Engine/Actors/Brushes
 Related:Engine/UI/LevelEditor
 Related: Engine/UI/LevelEditor/Modes
 Related: Engine/Actors/Volumes
+Version: 4.9
 
 [VAR:Steps]
 [OBJECT:Navigation]
