#UE4doc #docs - added engine version field.

diff --git a/Engine/Documentation/Source/Resources/ContentExamples/LevelDesign/2/2.INT.udn b/Engine/Documentation/Source/Resources/ContentExamples/LevelDesign/2/2.INT.udn
--- a/Engine/Documentation/Source/Resources/ContentExamples/LevelDesign/2/2.INT.udn
+++ b/Engine/Documentation/Source/Resources/ContentExamples/LevelDesign/2/2.INT.udn
@@ -6,6 +6,7 @@ Related: Engine/Content/Types/StaticMeshes
 Related: Engine/Rendering/Materials
 Related:Engine/UI/LevelEditor
 Related: Engine/UI/LevelEditor/Modes
+Version: 4.9
 
 [VAR:Steps]
 [OBJECT:Navigation]
