#UE4doc #docs - Added engine version field.

diff --git a/Engine/Documentation/Source/Resources/ContentExamples/LevelDesign/3/3.INT.udn b/Engine/Documentation/Source/Resources/ContentExamples/LevelDesign/3/3.INT.udn
--- a/Engine/Documentation/Source/Resources/ContentExamples/LevelDesign/3/3.INT.udn
+++ b/Engine/Documentation/Source/Resources/ContentExamples/LevelDesign/3/3.INT.udn
@@ -5,6 +5,7 @@ Description:An overview of the Level Design level, example 3: Lighting Pass
 Related: Engine/Rendering/LightingAndShadows/Basics
 Related: Engine/Rendering/PostProcessEffects
 Related: Engine/Rendering/ParticleSystems
+Version: 4.9
 
 [VAR:Steps]
 [OBJECT:Navigation]
