#UE4doc #docs - added engine version field.

diff --git a/Engine/Documentation/Source/Resources/ContentExamples/LevelDesign/4/4.INT.udn b/Engine/Documentation/Source/Resources/ContentExamples/LevelDesign/4/4.INT.udn
--- a/Engine/Documentation/Source/Resources/ContentExamples/LevelDesign/4/4.INT.udn
+++ b/Engine/Documentation/Source/Resources/ContentExamples/LevelDesign/4/4.INT.udn
@@ -5,6 +5,7 @@ Description:An overview of the Level Design level, example 4: Polish Pass
 Related: Engine/Rendering/LightingAndShadows/Lightmass
 Related: Engine/Rendering/LightingAndShadows/ReflectionEnvironment
 Related: Engine/Rendering/LightingAndShadows/LitTranslucency
+Version: 4.9
 
 [VAR:Steps]
 [OBJECT:Navigation]
