#UE4doc #docs - added engine version field.

diff --git a/Engine/Documentation/Source/Resources/ContentExamples/Lighting/Lighting.INT.udn b/Engine/Documentation/Source/Resources/ContentExamples/Lighting/Lighting.INT.udn
--- a/Engine/Documentation/Source/Resources/ContentExamples/Lighting/Lighting.INT.udn
+++ b/Engine/Documentation/Source/Resources/ContentExamples/Lighting/Lighting.INT.udn
@@ -11,6 +11,7 @@ Related: Engine/Rendering/LightingAndShadows/Lightmass
 Related: Engine/Rendering/LightingAndShadows/Shadows
 Related: Engine/Rendering/LightingAndShadows/IndirectLightingCache
 Related: Engine/Rendering/LightingAndShadows/LightMobility
+Version: 4.9
 
 [REGION:banner]
 ![](Lighting_Header.png)(convert:false)
