#UE4doc #docs - Added engine version field. Bolded term.

diff --git a/Engine/Documentation/Source/Resources/ContentExamples/Lighting/1_2/1_2.INT.udn b/Engine/Documentation/Source/Resources/ContentExamples/Lighting/1_2/1_2.INT.udn
--- a/Engine/Documentation/Source/Resources/ContentExamples/Lighting/1_2/1_2.INT.udn
+++ b/Engine/Documentation/Source/Resources/ContentExamples/Lighting/1_2/1_2.INT.udn
@@ -11,6 +11,7 @@ Related: Engine/Rendering/LightingAndShadows/Lightmass
 Related: Engine/Rendering/LightingAndShadows/Shadows
 Related: Engine/Rendering/LightingAndShadows/IndirectLightingCache
 Related: Engine/Rendering/LightingAndShadows/LightMobility
+Version: 4.9
 
 [VAR:Steps]
 [OBJECT:Navigation]
@@ -55,7 +56,7 @@ Related: Engine/Rendering/LightingAndShadows/LightMobility
 ![%Resources\ContentExamples\Lighting\1_2:title%](Resources\ContentExamples\Lighting\Lighting_1_2.png)
 [/REGION]
 
-A **Spot Light** emits light from a single point in a cone shape. Users are given two cones to shape the light - the **Inner Cone Angle** and Outer Cone Angle. Within the Inner Cone Angle, the light achieves full brightness. As you go from the extent of the inner radius to the extents of the Outer Cone Angle, a falloff takes place, creating a penumbra, or softening around the Spot Light's disc of illumination.
+A **Spot Light** emits light from a single point in a cone shape. Users are given two cones to shape the light - the **Inner Cone Angle** and **Outer Cone Angle**. Within the Inner Cone Angle, the light achieves full brightness. As you go from the extent of the inner radius to the extents of the Outer Cone Angle, a falloff takes place, creating a penumbra, or softening around the Spot Light's disc of illumination.
 
 ![](ConeRadius.png)(w:500)
 
