#UE4doc #docs - Added engine version field. Replaced image table with comparison slider.

diff --git a/Engine/Documentation/Source/Resources/ContentExamples/Lighting/2_1/2_1.INT.udn b/Engine/Documentation/Source/Resources/ContentExamples/Lighting/2_1/2_1.INT.udn
--- a/Engine/Documentation/Source/Resources/ContentExamples/Lighting/2_1/2_1.INT.udn
+++ b/Engine/Documentation/Source/Resources/ContentExamples/Lighting/2_1/2_1.INT.udn
@@ -11,6 +11,7 @@ Related: Engine/Rendering/LightingAndShadows/Lightmass
 Related: Engine/Rendering/LightingAndShadows/Shadows
 Related: Engine/Rendering/LightingAndShadows/IndirectLightingCache
 Related: Engine/Rendering/LightingAndShadows/LightMobility
+Version: 4.9
 
 [VAR:Steps]
 [OBJECT:Navigation]
@@ -68,12 +69,15 @@ On Static Mesh Components, lightmap resolution is set either on the Static Mesh
 
 On Brush surfaces, lightmap resolution is set via the _Lightmap Resolution_ property.  This is really setting a density, so lower values produce higher resolution.
 
-[REGION:imagetable]
-|![](LightmapRes_4.png)(w:%2XSize%)  |![](LightmapRes_16.png)(w:%2XSize%)  |
-|------------------------|-------------------------|
-|Lightmap Resolution: 4.0|Lightmap Resolution: 16.0|
-[/REGION]
-[/EXCERPT:StaticLighting02]
+[OBJECT:ComparisonSlider]
+	[PARAM:before]
+	![Lightmap Resolution: 4.0](LightmapRes_4.png)(w:%2XSize%)
+	[/PARAM]
+	[PARAM:after]
+	![Lightmap Resolution: 16.0](LightmapRes_16.png)(w:%2XSize%)
+	[/PARAM]
+[/OBJECT]
+
 
 [EXCERPT:ILC]
 Static lights can also affect dynamic objects due to the _Indirect Lighting Cache Sample Grid_. This is a 3D grid of lighting samples that get baked into the scene and are used by the rendering engine to determine what effects static lighting will have on dynamic objects. The grid can be visualized in the viewport by typing the following command into the Console Bar in the upper-right corner of the editor:
@@ -85,7 +89,7 @@ You will see the samples scattered throughout the scene like so:
 ![](IndirectLightingCache.png)
 
 [REGION:note]
-In order for the Indirect Lighting Cache samples to exist, your scene must have a _Lightmass Importance Volume_ around it, and lighting must have been built by Lightmass.
+In order for the Indirect Lighting Cache samples to exist, your scene must have a **Lightmass Importance Volume** around it, and lighting must have been built by Lightmass.
 [/REGION]
 
 
