#UE4doc #docs - Added engine version field.

diff --git a/Engine/Documentation/Source/Resources/ContentExamples/Lighting/2_3/2_3.INT.udn b/Engine/Documentation/Source/Resources/ContentExamples/Lighting/2_3/2_3.INT.udn
--- a/Engine/Documentation/Source/Resources/ContentExamples/Lighting/2_3/2_3.INT.udn
+++ b/Engine/Documentation/Source/Resources/ContentExamples/Lighting/2_3/2_3.INT.udn
@@ -10,7 +10,8 @@ Related: Engine/Rendering/LightingAndShadows/LightTypes/Directional
 Related: Engine/Rendering/LightingAndShadows/Lightmass
 Related: Engine/Rendering/LightingAndShadows/Shadows
 Related: Engine/Rendering/LightingAndShadows/IndirectLightingCache
-Related: Engine/Rendering/LightingAndShadows/LightMobility
+Related: Engine/Rendering/LightingAndShadows/LightMobility]
+Version: 4.9
 
 [VAR:Steps]
 [OBJECT:Navigation]
