#UE4doc #docs - Added engine version field. Unbolded hyperlink.

diff --git a/Engine/Documentation/Source/Resources/ContentExamples/Lighting/3_2/3_2.INT.udn b/Engine/Documentation/Source/Resources/ContentExamples/Lighting/3_2/3_2.INT.udn
--- a/Engine/Documentation/Source/Resources/ContentExamples/Lighting/3_2/3_2.INT.udn
+++ b/Engine/Documentation/Source/Resources/ContentExamples/Lighting/3_2/3_2.INT.udn
@@ -12,6 +12,7 @@ Related: Engine/Rendering/LightingAndShadows/Lightmass
 Related: Engine/Rendering/LightingAndShadows/Shadows
 Related: Engine/Rendering/LightingAndShadows/IndirectLightingCache
 Related: Engine/Rendering/LightingAndShadows/LightMobility
+Version: 4.9
 
 [VAR:Steps]
 [OBJECT:Navigation]
@@ -57,7 +58,7 @@ Related: Engine/Rendering/LightingAndShadows/LightMobility
 
 [INCLUDE:Engine/Rendering/LightingAndShadows/LightFunctions#LightFunctions]
 
-Since you are primarily using Light Functions to change the overall shape of a light, it is possible to confuse them with IES profiles, which are shown in **[Example 3.1](Resources\ContentExamples\Lighting\3_1)**. However, IES profiles are significantly cheaper to calculate, and should be used instead of Light Functions whenever possible.
+Since you are primarily using Light Functions to change the overall shape of a light, it is possible to confuse them with IES profiles, which are shown in [Example 3.1](Resources\ContentExamples\Lighting\3_1). However, IES profiles are significantly cheaper to calculate, and should be used instead of Light Functions whenever possible.
 
 
 
