#UE4doc #docs - Added engine version field. Bolded term.

diff --git a/Engine/Documentation/Source/Resources/ContentExamples/Lighting/4_2/4_2.INT.udn b/Engine/Documentation/Source/Resources/ContentExamples/Lighting/4_2/4_2.INT.udn
--- a/Engine/Documentation/Source/Resources/ContentExamples/Lighting/4_2/4_2.INT.udn
+++ b/Engine/Documentation/Source/Resources/ContentExamples/Lighting/4_2/4_2.INT.udn
@@ -11,6 +11,7 @@ Related: Engine/Rendering/LightingAndShadows/Lightmass
 Related: Engine/Rendering/LightingAndShadows/Shadows
 Related: Engine/Rendering/LightingAndShadows/IndirectLightingCache
 Related: Engine/Rendering/LightingAndShadows/LightMobility
+Version: 4.9
 
 [VAR:Steps]
 [OBJECT:Navigation]
@@ -59,7 +60,7 @@ Related: Engine/Rendering/LightingAndShadows/LightMobility
 
 Using Inverse Square Falloff also changes the fundamental behavior of a light. Firstly, Brightness changes from being an arbitrary value to being calculated in _lumens_. This means that a light with Brightness of 3 using Exponential Falloff (default) will appear bright, but if you switch to Inverse Square Falloff, the light is only 3 lumens, meaning it is practically invisible.
 
-Light radius also changes when using Inverse Square Falloff. When using default Exponential Falloff, the radius is used like the extents for light travel, with the light falling off as it approaches the radius. When using Inverse Square Falloff, radius works as more of a clamp, having no effect on how far the falloff actually goes. This means you _could_ set the radius very high for the most realistic light falloff. However, because of the problem of overlapping Stationary Lights, you will generally not want to do this.
+**Light Radius** also changes when using Inverse Square Falloff. When using default Exponential Falloff, the radius is used like the extents for light travel, with the light falling off as it approaches the radius. When using Inverse Square Falloff, radius works as more of a clamp, having no effect on how far the falloff actually goes. This means you _could_ set the radius very high for the most realistic light falloff. However, because of the problem of overlapping Stationary Lights, you will generally not want to do this.
 
 
 
