#UE4doc #docs - Added engine version field. Bolded and capitalized term.

diff --git a/Engine/Documentation/Source/Resources/ContentExamples/Lighting/4_3/4_3.INT.udn b/Engine/Documentation/Source/Resources/ContentExamples/Lighting/4_3/4_3.INT.udn
--- a/Engine/Documentation/Source/Resources/ContentExamples/Lighting/4_3/4_3.INT.udn
+++ b/Engine/Documentation/Source/Resources/ContentExamples/Lighting/4_3/4_3.INT.udn
@@ -11,6 +11,7 @@ Related: Engine/Rendering/LightingAndShadows/Lightmass
 Related: Engine/Rendering/LightingAndShadows/Shadows
 Related: Engine/Rendering/LightingAndShadows/IndirectLightingCache
 Related: Engine/Rendering/LightingAndShadows/LightMobility
+Version: 4.9
 
 [VAR:Steps]
 [OBJECT:Navigation]
@@ -54,7 +55,7 @@ Related: Engine/Rendering/LightingAndShadows/LightMobility
 ![%Resources\ContentExamples\Lighting\4_3:title%](Resources\ContentExamples\Lighting\Lighting_4_3.png)
 [/REGION]
 
-This example shows how indirect lighting can affect objects in the scene. Here we have a red light enclosed in a small space, with bounced light from the wall striking a static object and a dynamic one. Because indirect lighting is pre-calculated and baked into the static Lightmap, the static object's lighting and shadowing will be more accurate. For the dynamic object, however, lighting is dependent on the indirect lighting cache.
+This example shows how **Indirect Lighting** can affect objects in the scene. Here we have a red light enclosed in a small space, with bounced light from the wall striking a static object and a dynamic one. Because indirect lighting is pre-calculated and baked into the static Lightmap, the static object's lighting and shadowing will be more accurate. For the dynamic object, however, lighting is dependent on the indirect lighting cache.
 
 [INCLUDE:Resources\ContentExamples\Lighting\2_1#ILC]
 
