#UE4doc #docs - Added engine version field.

diff --git a/Engine/Documentation/Source/Resources/ContentExamples/Lighting/5_2/5_2.INT.udn b/Engine/Documentation/Source/Resources/ContentExamples/Lighting/5_2/5_2.INT.udn
--- a/Engine/Documentation/Source/Resources/ContentExamples/Lighting/5_2/5_2.INT.udn
+++ b/Engine/Documentation/Source/Resources/ContentExamples/Lighting/5_2/5_2.INT.udn
@@ -11,6 +11,7 @@ Related: Engine/Rendering/LightingAndShadows/Lightmass
 Related: Engine/Rendering/LightingAndShadows/Shadows
 Related: Engine/Rendering/LightingAndShadows/IndirectLightingCache
 Related: Engine/Rendering/LightingAndShadows/LightMobility
+Version: 4.9
 
 [VAR:Steps]
 [OBJECT:Navigation]
