#UE4doc #docs - Added engine version field. Removed outdated term.

diff --git a/Engine/Documentation/Source/Resources/ContentExamples/Lighting/6_1/6_1.INT.udn b/Engine/Documentation/Source/Resources/ContentExamples/Lighting/6_1/6_1.INT.udn
--- a/Engine/Documentation/Source/Resources/ContentExamples/Lighting/6_1/6_1.INT.udn
+++ b/Engine/Documentation/Source/Resources/ContentExamples/Lighting/6_1/6_1.INT.udn
@@ -11,6 +11,7 @@ Related: Engine/Rendering/LightingAndShadows/Lightmass
 Related: Engine/Rendering/LightingAndShadows/Shadows
 Related: Engine/Rendering/LightingAndShadows/IndirectLightingCache
 Related: Engine/Rendering/LightingAndShadows/LightMobility
+Version: 4.9
 
 [VAR:Steps]
 [OBJECT:Navigation]
@@ -55,7 +56,7 @@ Related: Engine/Rendering/LightingAndShadows/LightMobility
 [/REGION]
 
 
-This example shows off the Self Shadowing Accuracy, or **Shadow Bias**. This property exists primarily to help objects feel more grounded along a surface. There is a tradeoff, however, in that decreasing the value to make an object appear more grounded can sometimes lead to self-shadowing artifacts, as seen in the image below:
+This example shows off the **Shadow Bias**. This property exists primarily to help objects feel more grounded along a surface. There is a tradeoff, however, in that decreasing the value to make an object appear more grounded can sometimes lead to self-shadowing artifacts, as seen in the image below:
 
 ![](SelfShadowingArtifacts.png)(w:700)
 
