#UE4doc #docs - Added engine version field.

diff --git a/Engine/Documentation/Source/Resources/ContentExamples/MaterialInstances/1_1/1_1.INT.udn b/Engine/Documentation/Source/Resources/ContentExamples/MaterialInstances/1_1/1_1.INT.udn
--- a/Engine/Documentation/Source/Resources/ContentExamples/MaterialInstances/1_1/1_1.INT.udn
+++ b/Engine/Documentation/Source/Resources/ContentExamples/MaterialInstances/1_1/1_1.INT.udn
@@ -9,6 +9,7 @@ Related: Engine/Rendering/Materials/MaterialInputs
 Related: Engine/Rendering/Materials/LayeredMaterials
 Related: Engine/Rendering/Materials/MaterialInstances
 Related: Engine/Actors/DecalActor
+Version: 4.9
 
 [VAR:Steps]
 [OBJECT:Navigation]
diff --git a/Engine/Documentation/Source/Resources/ContentExamples/MaterialInstances/1_2/1_2.INT.udn b/Engine/Documentation/Source/Resources/ContentExamples/MaterialInstances/1_2/1_2.INT.udn
--- a/Engine/Documentation/Source/Resources/ContentExamples/MaterialInstances/1_2/1_2.INT.udn
+++ b/Engine/Documentation/Source/Resources/ContentExamples/MaterialInstances/1_2/1_2.INT.udn
@@ -9,6 +9,7 @@ Related: Engine/Rendering/Materials/ExpressionReference
 Related: Engine/Rendering/Materials/MaterialInputs
 Related: Engine/Rendering/Materials/LayeredMaterials
 Related: Engine/Actors/DecalActor
+Version: 4.9
 
 [VAR:Steps]
 [OBJECT:Navigation]
