#UE4doc #docs - Added engine version field. Removed Related header section and converted this to use the Related field. Updated image call out colors.

diff --git a/Engine/Documentation/Source/Resources/ContentExamples/MaterialInstances/1_4/1_4.INT.udn b/Engine/Documentation/Source/Resources/ContentExamples/MaterialInstances/1_4/1_4.INT.udn
--- a/Engine/Documentation/Source/Resources/ContentExamples/MaterialInstances/1_4/1_4.INT.udn
+++ b/Engine/Documentation/Source/Resources/ContentExamples/MaterialInstances/1_4/1_4.INT.udn
@@ -9,6 +9,7 @@ Related: Engine/Rendering/Materials/MaterialInputs
 Related: Engine/Rendering/Materials/LayeredMaterials
 Related: Engine/Rendering/Materials/MaterialInstances
 Related: Engine/Actors/DecalActor
+Version: 4.9
 
 [VAR:Steps]
 [OBJECT:Navigation]
diff --git a/Engine/Documentation/Source/Resources/ContentExamples/MaterialNodes/MaterialNodes.INT.udn b/Engine/Documentation/Source/Resources/ContentExamples/MaterialNodes/MaterialNodes.INT.udn
--- a/Engine/Documentation/Source/Resources/ContentExamples/MaterialNodes/MaterialNodes.INT.udn
+++ b/Engine/Documentation/Source/Resources/ContentExamples/MaterialNodes/MaterialNodes.INT.udn
@@ -2,6 +2,9 @@ Availability: Public
 Title:Material Nodes
 Crumbs:%ROOT%, Resources
 Description:An overview of the Material Nodes example level.
+Related:Engine/Rendering/Materials
+Related:Engine/Rendering/Materials\MaterialInputs
+Version: 4.9
 
 [TOC(start:2 end:2)]
 
@@ -84,9 +87,3 @@ As with all Content Examples, this level is meant to be interactive. With any of
 [/REGION]
 [/EXCERPT:Topics]
 
-
-## Related Topics and Further Reading
-
-* [](Engine/Rendering/Materials)
-* [](Engine/Rendering/Materials\MaterialInputs)
-
