#UE4doc #docs - added engine version field.

diff --git a/Engine/Documentation/Source/Resources/ContentExamples/MaterialNodes/1_13/1_13.INT.udn b/Engine/Documentation/Source/Resources/ContentExamples/MaterialNodes/1_13/1_13.INT.udn
--- a/Engine/Documentation/Source/Resources/ContentExamples/MaterialNodes/1_13/1_13.INT.udn
+++ b/Engine/Documentation/Source/Resources/ContentExamples/MaterialNodes/1_13/1_13.INT.udn
@@ -9,6 +9,7 @@ Related: Engine/Rendering/Materials/MaterialInputs
 Related: Engine/Rendering/Materials/LayeredMaterials
 Related: Engine/Rendering/Materials/MaterialInstances
 Related: Engine/Actors/DecalActor
+Version: 4.9
 
 [VAR:Steps]
 [OBJECT:Navigation]
diff --git a/Engine/Documentation/Source/Resources/ContentExamples/MaterialNodes/1_14/1_14.INT.udn b/Engine/Documentation/Source/Resources/ContentExamples/MaterialNodes/1_14/1_14.INT.udn
--- a/Engine/Documentation/Source/Resources/ContentExamples/MaterialNodes/1_14/1_14.INT.udn
+++ b/Engine/Documentation/Source/Resources/ContentExamples/MaterialNodes/1_14/1_14.INT.udn
@@ -9,6 +9,7 @@ Related: Engine/Rendering/Materials/MaterialInputs
 Related: Engine/Rendering/Materials/LayeredMaterials
 Related: Engine/Rendering/Materials/MaterialInstances
 Related: Engine/Actors/DecalActor
+Version: 4.9
 
 [VAR:Steps]
 [OBJECT:Navigation]
diff --git a/Engine/Documentation/Source/Resources/ContentExamples/MaterialNodes/1_7/1_7.INT.udn b/Engine/Documentation/Source/Resources/ContentExamples/MaterialNodes/1_7/1_7.INT.udn
--- a/Engine/Documentation/Source/Resources/ContentExamples/MaterialNodes/1_7/1_7.INT.udn
+++ b/Engine/Documentation/Source/Resources/ContentExamples/MaterialNodes/1_7/1_7.INT.udn
@@ -9,6 +9,7 @@ Related: Engine/Rendering/Materials/MaterialInputs
 Related: Engine/Rendering/Materials/LayeredMaterials
 Related: Engine/Rendering/Materials/MaterialInstances
 Related: Engine/Actors/DecalActor
+Version: 4.9
 
 [VAR:Steps]
 [OBJECT:Navigation]
