#UE4doc #docs - added engine version field.

diff --git a/Engine/Documentation/Source/Resources/ContentExamples/MaterialNodes/1_2/1_2.INT.udn b/Engine/Documentation/Source/Resources/ContentExamples/MaterialNodes/1_2/1_2.INT.udn
--- a/Engine/Documentation/Source/Resources/ContentExamples/MaterialNodes/1_2/1_2.INT.udn
+++ b/Engine/Documentation/Source/Resources/ContentExamples/MaterialNodes/1_2/1_2.INT.udn
@@ -10,6 +10,7 @@ Related: Engine/Rendering/Materials/MaterialInputs
 Related: Engine/Rendering/Materials/LayeredMaterials
 Related: Engine/Rendering/Materials/MaterialInstances
 Related: Engine/Actors/DecalActor
+Version: 4.9
 
 [VAR:Steps]
 [OBJECT:Navigation]
diff --git a/Engine/Documentation/Source/Resources/ContentExamples/MaterialNodes/1_3/1_3.INT.udn b/Engine/Documentation/Source/Resources/ContentExamples/MaterialNodes/1_3/1_3.INT.udn
--- a/Engine/Documentation/Source/Resources/ContentExamples/MaterialNodes/1_3/1_3.INT.udn
+++ b/Engine/Documentation/Source/Resources/ContentExamples/MaterialNodes/1_3/1_3.INT.udn
@@ -10,6 +10,7 @@ Related: Engine/Rendering/Materials/MaterialInputs
 Related: Engine/Rendering/Materials/LayeredMaterials
 Related: Engine/Rendering/Materials/MaterialInstances
 Related: Engine/Actors/DecalActor
+Version: 4.9
 
 [VAR:Steps]
 [OBJECT:Navigation]
diff --git a/Engine/Documentation/Source/Resources/ContentExamples/MaterialNodes/1_4/1_4.INT.udn b/Engine/Documentation/Source/Resources/ContentExamples/MaterialNodes/1_4/1_4.INT.udn
--- a/Engine/Documentation/Source/Resources/ContentExamples/MaterialNodes/1_4/1_4.INT.udn
+++ b/Engine/Documentation/Source/Resources/ContentExamples/MaterialNodes/1_4/1_4.INT.udn
@@ -9,6 +9,7 @@ Related: Engine/Rendering/Materials/MaterialInputs
 Related: Engine/Rendering/Materials/LayeredMaterials
 Related: Engine/Rendering/Materials/MaterialInstances
 Related: Engine/Actors/DecalActor
+Version: 4.9
 
 [VAR:Steps]
 [OBJECT:Navigation]
diff --git a/Engine/Documentation/Source/Resources/ContentExamples/MaterialNodes/1_5/1_5.INT.udn b/Engine/Documentation/Source/Resources/ContentExamples/MaterialNodes/1_5/1_5.INT.udn
--- a/Engine/Documentation/Source/Resources/ContentExamples/MaterialNodes/1_5/1_5.INT.udn
+++ b/Engine/Documentation/Source/Resources/ContentExamples/MaterialNodes/1_5/1_5.INT.udn
@@ -9,6 +9,7 @@ Related: Engine/Rendering/Materials/MaterialInputs
 Related: Engine/Rendering/Materials/LayeredMaterials
 Related: Engine/Rendering/Materials/MaterialInstances
 Related: Engine/Actors/DecalActor
+Version: 4.9
 
 [VAR:Steps]
 [OBJECT:Navigation]
