#UE4 #DOCS link correction for new physics docs

diff --git a/Engine/Documentation/Source/Resources/ContentExamples/Physics/Physics.INT.udn b/Engine/Documentation/Source/Resources/ContentExamples/Physics/Physics.INT.udn
--- a/Engine/Documentation/Source/Resources/ContentExamples/Physics/Physics.INT.udn
+++ b/Engine/Documentation/Source/Resources/ContentExamples/Physics/Physics.INT.udn
@@ -3,7 +3,7 @@ Title:Physics Content Examples
 Crumbs:%ROOT%, Resources, Resources/ContentExamples
 Description:Overview of the samples provided in the Physics Content Examples.
 Related: Engine/Physics/PhysicsBodies
-Related: Engine/Physics/PhAT/UserGuide
+Related: Engine/Physics/PhAT
 Related: Engine/Physics/Constraints/ConstraintsReference
 Related: Engine/Physics/Constraints/ConstraintsUserGuide
 Related: Engine/Physics/Collision
