#UE4doc #docs - Added engine version field.

diff --git a/Engine/Documentation/Source/Resources/ContentExamples/MaterialProperties/1_10/1_10.INT.udn b/Engine/Documentation/Source/Resources/ContentExamples/MaterialProperties/1_10/1_10.INT.udn
--- a/Engine/Documentation/Source/Resources/ContentExamples/MaterialProperties/1_10/1_10.INT.udn
+++ b/Engine/Documentation/Source/Resources/ContentExamples/MaterialProperties/1_10/1_10.INT.udn
@@ -11,6 +11,7 @@ Related: Engine/Rendering/Materials/MaterialInputs
 Related: Engine/Rendering/Materials/LayeredMaterials
 Related: Engine/Rendering/Materials/MaterialInstances
 Related: Engine/Actors/DecalActor
+Version: 4.9
 
 [VAR:Steps]
 [OBJECT:Navigation]
diff --git a/Engine/Documentation/Source/Resources/ContentExamples/MaterialProperties/1_11/1_11.INT.udn b/Engine/Documentation/Source/Resources/ContentExamples/MaterialProperties/1_11/1_11.INT.udn
--- a/Engine/Documentation/Source/Resources/ContentExamples/MaterialProperties/1_11/1_11.INT.udn
+++ b/Engine/Documentation/Source/Resources/ContentExamples/MaterialProperties/1_11/1_11.INT.udn
@@ -10,6 +10,7 @@ Related: Engine/Rendering/Materials/MaterialInputs
 Related: Engine/Rendering/Materials/LayeredMaterials
 Related: Engine/Rendering/Materials/MaterialInstances
 Related: Engine/Actors/DecalActor
+Version: 4.9
 
 [VAR:Steps]
 [OBJECT:Navigation]
diff --git a/Engine/Documentation/Source/Resources/ContentExamples/MaterialProperties/1_2/1_2.INT.udn b/Engine/Documentation/Source/Resources/ContentExamples/MaterialProperties/1_2/1_2.INT.udn
--- a/Engine/Documentation/Source/Resources/ContentExamples/MaterialProperties/1_2/1_2.INT.udn
+++ b/Engine/Documentation/Source/Resources/ContentExamples/MaterialProperties/1_2/1_2.INT.udn
@@ -11,6 +11,7 @@ Related: Engine/Rendering/Materials/MaterialInputs
 Related: Engine/Rendering/Materials/LayeredMaterials
 Related: Engine/Rendering/Materials/MaterialInstances
 Related: Engine/Actors/DecalActor
+Version: 4.9
 
 [VAR:Steps]
 [OBJECT:Navigation]
diff --git a/Engine/Documentation/Source/Resources/ContentExamples/MaterialProperties/1_3/1_3.INT.udn b/Engine/Documentation/Source/Resources/ContentExamples/MaterialProperties/1_3/1_3.INT.udn
--- a/Engine/Documentation/Source/Resources/ContentExamples/MaterialProperties/1_3/1_3.INT.udn
+++ b/Engine/Documentation/Source/Resources/ContentExamples/MaterialProperties/1_3/1_3.INT.udn
@@ -13,6 +13,7 @@ Related: Engine/Rendering/Materials/MaterialInputs
 Related: Engine/Rendering/Materials/LayeredMaterials
 Related: Engine/Rendering/Materials/MaterialInstances
 Related: Engine/Actors/DecalActor
+Version: 4.9
 
 [VAR:Steps]
 [OBJECT:Navigation]
@@ -81,5 +82,5 @@ If a **Translucent Material** is set to **Default Lit**, there are 3 more option
 | ---- | ----- |
 | **Volumetric Non Directional**| The mesh is lit as if it were being used to create a volumetric effect like smoke.  Only diffuse lighting will be displayed and the normal is not used.  This is the cheapest form of lit translucency. |
 | **Volumetric Directional**| The mesh is lit as if it were being used to create a volumetric effect like smoke, but with a normal that affects the lighting. Only diffuse lighting will be displayed. |
-| **Surface**| It is lit and receives a reflection, but the reflection is limited and not as good as on **Opaque** material. It simply picks the nearest **Reflection Capture Actor** and applies it to the entire object. |
+| **Surface**| It is lit and receives a reflection, but the reflection is limited and not as good as on Opaque material. It simply picks the nearest Reflection Capture Actor and applies it to the entire object. |
 
