#UE4doc #docs - Added engine version field. Removed bolded term.

diff --git a/Engine/Documentation/Source/Resources/ContentExamples/MaterialProperties/1_4/1_4.INT.udn b/Engine/Documentation/Source/Resources/ContentExamples/MaterialProperties/1_4/1_4.INT.udn
--- a/Engine/Documentation/Source/Resources/ContentExamples/MaterialProperties/1_4/1_4.INT.udn
+++ b/Engine/Documentation/Source/Resources/ContentExamples/MaterialProperties/1_4/1_4.INT.udn
@@ -11,6 +11,7 @@ Related: Engine/Rendering/Materials/MaterialInputs
 Related: Engine/Rendering/Materials/LayeredMaterials
 Related: Engine/Rendering/Materials/MaterialInstances
 Related: Engine/Actors/DecalActor
+Version: 4.9
 
 [VAR:Steps]
 [OBJECT:Navigation]
diff --git a/Engine/Documentation/Source/Resources/ContentExamples/MaterialProperties/1_5/1_5.INT.udn b/Engine/Documentation/Source/Resources/ContentExamples/MaterialProperties/1_5/1_5.INT.udn
--- a/Engine/Documentation/Source/Resources/ContentExamples/MaterialProperties/1_5/1_5.INT.udn
+++ b/Engine/Documentation/Source/Resources/ContentExamples/MaterialProperties/1_5/1_5.INT.udn
@@ -12,6 +12,7 @@ Related: Engine/Rendering/Materials/LayeredMaterials
 Related: Engine/Rendering/Materials/MaterialInstances
 Related: Engine/Actors/DecalActor
 Related: Engine/Rendering/Materials/MaterialProperties
+Version: 4.9
 
 [VAR:Steps]
 [OBJECT:Navigation]
@@ -74,13 +75,13 @@ Related: Engine/Rendering/Materials/MaterialProperties
 ![](MP_1_5_001.png)(w:900)
 [/REGION]
 
-**Domains** define the overall usage of a material. Most shaders will be **Surface** however, if a shader needed to be used as a decal, light function or post process, it has to belong to those **Domains**.
+**Domains** define the overall usage of a material. Most shaders will be **Surface** however, if a shader needed to be used as a decal, light function or post process, it has to belong to those Domains.
 
 | Domain | Description |
 | ---- | ---- |
 | **Surface**         | Is the default and is used for any kind geometry (static, skeletal, landscape, etc...). |
-| **Deferred Decals** | If the material is to be used by **Decal Actor**, it needs to be set to this or it will not work. |
-| **Light Function**  | Used by materials applied to lights as **LightFunctions**. |
-| **PostProcess**     | If the material is used as **Blendable Material** in a **PostProcessing** chain. |
-| **User Interface**     | This material will be used for **UMG** or **Slate UI** |
+| **Deferred Decals** | If the material is to be used by Decal Actor, it needs to be set to this or it will not work. |
+| **Light Function**  | Used by materials applied to lights as LightFunctions. |
+| **PostProcess**     | If the material is used as Blendable Material in a PostProcessing chain. |
+| **User Interface**     | This material will be used for UMG or Slate UI |
 
