#UE4 Docs: SteamVR page was missing the create a VR gamemode step so I added it.

diff --git a/Engine/Documentation/Source/Platforms/SteamVR/QuickStart/2/SteamVRQuickStart_2.INT.udn b/Engine/Documentation/Source/Platforms/SteamVR/QuickStart/2/SteamVRQuickStart_2.INT.udn
--- a/Engine/Documentation/Source/Platforms/SteamVR/QuickStart/2/SteamVRQuickStart_2.INT.udn
+++ b/Engine/Documentation/Source/Platforms/SteamVR/QuickStart/2/SteamVRQuickStart_2.INT.udn
@@ -40,6 +40,11 @@ In the following section we will go over how to setup a project to work with Ste
 	 If that still does not get your VR Headset to work with the VR Preview make sure to check out the [VR Documents](Platforms/VR) to ensure that nothing was missed during setup.
 	[/region]
 
+1. Right - Click inside of the **Content Browser** and select the create a new Blueprint.
+	From the **Pick Parent Class** window that is displayed select the **Game Mode** option and call it **VR_Gamemode**. 
+	
+	![](Create_New_Game_Mode.png)
+
 1. Right - Click inside of the **Content Browser** and select the create a new Blueprint.
 	From the **Pick Parent Class** window that is displayed use the search box at the bottom and search for **PlayerCameraManager**. 
 	Once PlayerCameraManager is located click on it to create it and name it **VR_PlayerCameraManager**.
