#UE4doc #docs - Added engine version field. Placed comparison images in the new comparison slider.

diff --git a/Engine/Documentation/Source/Resources/ContentExamples/MaterialProperties/1_6/1_6.INT.udn b/Engine/Documentation/Source/Resources/ContentExamples/MaterialProperties/1_6/1_6.INT.udn
--- a/Engine/Documentation/Source/Resources/ContentExamples/MaterialProperties/1_6/1_6.INT.udn
+++ b/Engine/Documentation/Source/Resources/ContentExamples/MaterialProperties/1_6/1_6.INT.udn
@@ -13,6 +13,7 @@ Related: Engine/Rendering/Materials/LayeredMaterials
 Related: Engine/Rendering/Materials/MaterialInstances
 Related: Engine/Actors/DecalActor
 Related: Engine/Rendering/Materials/MaterialProperties
+Version: 4.9
 
 [VAR:Steps]
 [OBJECT:Navigation]
@@ -72,7 +73,14 @@ Related: Engine/Rendering/Materials/MaterialProperties
 [/VAR]
 
 <div align="center">
-![](MP_1_6_001.png)(h:350) ![](MP_1_6_002.png)(h:350)
+[OBJECT:ComparisonSlider]
+	[PARAM:before]
+	![Separate Translucency No Post Process](MP_1_6_001.png)(h:400)
+	[/PARAM]
+	[PARAM:after]
+	![Separate Translucency with Depth of Field ](MP_1_6_002.png)(h:400)
+	[/PARAM]
+[/OBJECT]
 </div>
 
 Draws translucent objects separately. **Separate Translucency** shaders will be not be affected by **Depth of Field**.
