#UE4doc #docs:
Added engine version field
Changed font settings for multiple terms
Fixed capitalization issues
Replaced some images as they did not match the guidelines

diff --git a/Engine/Documentation/Source/GettingStarted/DownloadingUnrealEngine/Perforce/P4V/P4V.INT.udn b/Engine/Documentation/Source/GettingStarted/DownloadingUnrealEngine/Perforce/P4V/P4V.INT.udn
--- a/Engine/Documentation/Source/GettingStarted/DownloadingUnrealEngine/Perforce/P4V/P4V.INT.udn
+++ b/Engine/Documentation/Source/GettingStarted/DownloadingUnrealEngine/Perforce/P4V/P4V.INT.udn
@@ -2,13 +2,14 @@ Availability: Licensee
 Title:Setting up P4V for Unreal Engine 4
 Crumbs:%ROOT%, Engine, Programming, Programming/Basics, GettingStarted/DownloadingUnrealEngine/Perforce
 Description:Guide to installing, setting up, and syncing builds with the Perforce P4V windows client for use with Unreal Engine 4.
+version: 4.9
 
 [TOC(start:2)]
 
 This document covers the steps in setting up the P4V client to connect to the Unreal Engine 4 depot and get a build of the engine.
 
 [REGION:note]
-The version of P4V used in this document is _P4V/NTX86/2009.2/236331_ released _February 24, 2010_. If using a different version, the steps and images may differ.
+The version of P4V used in this document is P4V/NTX86/2009.2/236331 released February 24, 2010. If you're using a different version, the steps and images may differ.
 [/REGION]
 
 To connect to Perforce, you need to set up a secure Virtual Private Network (VPN) [Connection to Epic Games](GettingStarted/DownloadingUnrealEngine/VPNSetup).
@@ -19,25 +20,27 @@ For other general tips, see the [Setting up Perforce](GettingStarted/Downloading
 ### Installation
 
 1. Run the P4Vinst executable installer. 
-1. On the Perforce _Select Features_ page, all features are selected by default. 
+1. On the Perforce **Select Features** page, all features are selected by default. 
 
 	![install_features.jpg](install_features.jpg)(w:514 h:392)
 
-	Press the ![install_next_button.jpg](install_next_button.jpg) button to continue.
+	Press the **Next > ** button to continue.
 
 	If you previously used P4Win, you may be prompted that the SCC plugin will be replaced.
 
 	![install_plugin_warning.jpg](install_plugin_warning.jpg)(w:362 h:155)
 
-1. On the _Perforce Client Configuration_ page, enter the following information:
+1. On the **Perforce Client Configuration** page, enter the following information:
 
 	* **Server Port** - `p4-licensee-proxy.epicgames.net:1667` (available through the US and East Asian VPN) or `p4-licensee-proxy-asia.epicgames.net:1667` (primary for those connected through the East Asian VPN)
-	* **Username** - Provided in original UDN welcome email, typically your Team Name (**NOTE**: this is different from your UDN login information)
-
+	* **Username** - Provided in original UDN welcome email, typically your Team Name 
+	[REGION:note]
+	This is different from your UDN login information
+	[/REGION]
 	![install_client_config.jpg](install_client_config.jpg)(w:514 h:392)
 
-	Press the ![install_next_button.jpg](install_next_button.jpg) button to continue.
-1. You are ready to begin the installation. Press the ![install_start_button.jpg](install_start_button.jpg) button to begin.
+	Press the **Next >** button to continue.
+1. You are ready to begin the installation. Press the **Start** button to begin.
 
 	![install_complete.jpg](install_complete.jpg)(w:514 h:392)
 
@@ -45,7 +48,7 @@ For other general tips, see the [Setting up Perforce](GettingStarted/Downloading
 
 
 [REGION:note]
-You can later reconfigure the IP address, as needed, via the _Settings_ menu. Please use the name to connect to the server, not the resolved IP address. The IP address may change if the server computer is changed; but the name will remain the same!
+You can later reconfigure the IP address, as needed, via the **Settings** menu. Please use the name to connect to the server, not the resolved IP address. The IP address may change if the server computer is changed; but the name will remain the same!
 [/REGION]
 
 ### Configuration - Workspaces
@@ -54,23 +57,23 @@ You can later reconfigure the IP address, as needed, via the _Settings_ menu. Pl
 For those that are familiar with P4Win, ClientSpecs are known as Workspaces in P4V.
 [/REGION]
 
-1. Run P4V. Upon first start, you may be prompted to run the _Connection Setup Wizard_. You can choose _No_ since the connection was already set up during installation.
-1. The _Open Connection_ dialog will open.
+1. Run P4V. Upon first start, you may be prompted to run the **Connection Setup Wizard**. You can choose _No_ since the connection was already set up during installation.
+1. The **Open Connection** dialog will open.
 
 ![connect_dialog.jpg](connect_dialog.jpg)(w:512 h:265)
 
-Verify that the _Server_ and _User_ fields match up with the required information.
-1. For a client Workspace, if you previously used P4Win or reinstalled P4V, you can press the ![connect_browse_button.jpg](connect_browse_button.jpg) button to find a Workspace associated with your user.
+Verify that the **Server** and **User** fields match up with the required information.
+1. For a client Workspace, if you previously used P4Win or reinstalled P4V, you can press the **Browse...** button to find a Workspace associated with your user.
 
-	If no workspaces exist, press the ![connect_new_button.jpg](connect_new_button.jpg) button to create a new workspace. In the dialog that opens, enter the name for the workspace to create using the following information as a guide:
+	If no workspaces exist, press the **New** button to create a new workspace. In the dialog that opens, enter the name for the workspace to create using the following information as a guide:
 
-	* **IF YOU ARE A LICENSEE** _Client Name_: **UE4_Licensee_[TEAMNAME]** (or otherwise ending with your UDN username or team name)
-	* **IF YOU ARE A CONTRACTOR** _Client Name_: **UE4_Contractor_[USERNAME]** (or otherwise ending with your UDN username or team name)
+	* If yu are a licensee: **Client Name**: **UE4_Licensee_[TEAMNAME]** (or otherwise ending with your UDN username or team name)
+	* If you are a contractor: **Client Name**: **UE4_Contractor_[USERNAME]** (or otherwise ending with your UDN username or team name)
 
 	![connect_workspace.jpg](connect_workspace.jpg)(w:219 h:131)
 
-	Press the ![connect_ok_button.jpg](connect_ok_button.jpg) to create the new workspace.
-1. The workspace is now shown in the _Open Connection_ dialog.
+	Press the **Ok** button to create the new workspace.
+1. The workspace is now shown in the **Open Connection** dialog.
 
 	![connect_complete.jpg](connect_complete.jpg)(w:512 h:265)
 
@@ -78,10 +81,10 @@ Verify that the _Server_ and _User_ fields match up with the required informatio
 
 	![connect_startup.jpg](connect_startup.jpg)(w:140 h:26)
 
-Press the ![connect_ok_button.jpg](connect_ok_button.jpg) to open the connection and the P4V client.
-1. You may be prompted about your ticket expiring. Simply choose to not see the message again in the future, and press the ![connect_ok_button.jpg](connect_ok_button.jpg)button.
+Press the **Ok** button to open the connection and the P4V client.
+1. You may be prompted about your ticket expiring. Simply choose to not see the message again in the future, and press the **Ok** button.
 
-	If you are prompted to populate your recently-created workspace, decline the offer. You **really** do not want several gigabytes of data that you will have to replace shortly.
+	If you are prompted to populate your recently-created workspace, decline the offer. This will avoid receiving several gigabytes of data that you will have to replace shortly.
 
 
 Workspaces are created on a per-machine basis, because they keep track of the exact revision of each file that has been stored on that machine's hard drive in the selected directory. Thus, if you ever need to start syncing from a different PC, you will need to repeat this process to create a new workspace on that machine.
@@ -92,7 +95,7 @@ If you are an anonymous user, you will want to create a workspace that is named:
 
 #### Troubleshooting
 
-Selecting "depot view" with a filter of "Entire Depot tree" should show all folders that you have access to. Alternatively, when creating/editing a workspace, you can use the "View" tab to visually map any files that you have access to.
+Selecting Depot View with a filter of Entire Depot tree should show all folders that you have access to. Alternatively, when creating/editing a workspace, you can use the **View** tab to visually map any files that you have access to.
 
 
 
@@ -113,56 +116,56 @@ The tree pane in P4V contains your Workspace and the Depot Tree. In the Depot Tr
 ![p4v_depot_tree.jpg](p4v_depot_tree.jpg)(w:374 h:307)
 
 
-* _CodeDrop_ - Contains various snapshots of games created by Epic that can be used as examples.
-* _PhysX_ - Contains the common files for the physics middleware; depending on your agreement, you may also see additional directories and files.
-* _UnrealEngine3_ - Contains everything to build and run UnrealEngine3. Also includes the ExampleGame and UDKGame examples.
-* _UnrealMCP_ - Contains the source for UnrealMCP, which is a SOA (service-oriented architecture) that focuses on providing highly scalable and available set of services to enhance the online experience of players.
+* **CodeDrop** - Contains various snapshots of games created by Epic that can be used as examples.
+* **PhysX** - Contains the common files for the physics middleware; depending on your agreement, you may also see additional directories and files.
+* **UnrealEngine3** - Contains everything to build and run UnrealEngine3. Also includes the ExampleGame and UDKGame examples.
+* **UnrealMCP** - Contains the source for UnrealMCP, which is a SOA (service-oriented architecture) that focuses on providing highly scalable and available set of services to enhance the online experience of players.
 
 
 For more information, see the [directory layout](Engine/Basics/DirectoryStructure) page.
 
-In the _UnrealEngine3_ directory are several subdirectories: _Binaries_, _Development_, and _Engine_; and two game directories: _ExampleGame_ and _UDKGame_.
+In the UnrealEngine3 directory are several subdirectories: Binaries, Development, and Engine; and two game directories: ExampleGame and UDKGame.
 
 ![p4v_depot_tree_ue3.jpg](p4v_depot_tree_ue3.jpg)(w:374 h:307)
 
-_ExampleGame_ is the most minimal game implementation offered by Epic, and can be a good place for getting started in creating your own game project if your game is not very similar to _Unreal Tournament_.  
+ExampleGame is the most minimal game implementation offered by Epic, and can be a good place for getting started in creating your own game project if your game is not very similar to Unreal Tournament.  
 
-_UDKGame_ is Epic's current Unreal Engine 4 usage demo based off of _Unreal Tournament 3_. This provides a complex example showing exactly how Epic intends to use the engine. We recommend using _UDKGame_ as an example, rather than extending from _UDKGame_. You are welcome to use any C++ or UnrealScript code in _UDKGame_ for your game project, but you cannot ship any of Epic's content (textures, meshes, levels, sounds, etc.) in your game.
+UDKGame is Epic's current Unreal Engine 4 usage demo based off of Unreal Tournament 3. This provides a complex example showing exactly how Epic intends to use the engine. We recommend using UDKGame as an example, rather than extending from UDKGame. You are welcome to use any C++ or UnrealScript code in UDKGame for your game project, but you cannot ship any of Epic's content (textures, meshes, levels, sounds, etc.) in your game.
 
-For _UnrealEngine3_, most of Epic's content is located in the game-specific directories, so if all you want is the engine source, you can simply sync to the //depot/UnrealEngine3/Development directory. You can highlight specific directories as desired, for example, to grab the latest content for _ExampleGame_ without updating _UDKGame_ content.
+For UnrealEngine3, most of Epic's content is located in the game-specific directories, so if all you want is the engine source, you can simply sync to the //depot/UnrealEngine3/Development directory. You can highlight specific directories as desired, for example, to grab the latest content for ExampleGame without updating UDKGame content.
 
 ![p4v_depot_tree_ue3_examplegame.jpg](p4v_depot_tree_ue3_examplegame.jpg)(w:374 h:307)
 
-When getting a revision, lots of downloading will progress in the log window. _Sync completed with no warnings or errors_ is what you are looking for when it is done.
+When getting a revision, lots of downloading will progress in the log window. "Sync completed with no warnings or errors" is what you are looking for when it is done.
 
 
 ### Syncing to a QA Build
 
 To sync to a QA-approved build, follow the step below:
 
-1. In the Depot Tree, **right-click** on the UnrealEngine3 directory and then choose _Get Revision..._ from the context menu. Alternatively, you can select the _UnrealEngine3_ directory and choose _Get Revision..._ from the _Actions_ menu.
+1. In the Depot Tree, **Right-click** on the UnrealEngine3 directory and then choose **Get Revision...** from the context menu. Alternatively, you can select the UnrealEngine3 directory and choose **Get Revision...** from the **Actions** menu.
 
 	![p4v_depot_menu.jpg](p4v_depot_menu.jpg)(w:304 h:650)
-1. The _Get Revision_ window will open.
+1. The **Get Revision** window will open.
 
 	![p4v_get_revision.jpg](p4v_get_revision.jpg)(w:634 h:488)
 
-	Select the _Specify revision using_ radio button to enable the combobox and text field.
+	Select the **Specify revision using** radio button to enable the combobox and text field.
 
 	![p4v_specify_revision.jpg](p4v_specify_revision.jpg)(w:513 h:63)
-1. Select _label_ from the combobox.
+1. Select **Label** from the dropdown box.
 
 	![p4v_revision_label.jpg](p4v_revision_label.jpg)(w:100 h:103)
 
-	Press the ![p4v_browse_button.jpg](p4v_browse_button.jpg) button to open the _Select Label_ window.
+	Press the **Browse...** button to open the **Select Label** window.
 
 	![p4v_select_label.jpg](p4v_select_label.jpg)(w:800 h:526)
 1. Select the label you want to sync to. Usually this is the latest QA-approved build, designated by a label in the form of _QA_APPROVED_BUILD_[Month]_[Year]_. You may need to scroll down in the list a few times for all the labels to be loaded.
 
 	![p4v_label_build.jpg](p4v_label_build.jpg)(w:800 h:526)
 
-	Press the ![p4v_ok_button.jpg](p4v_ok_button.jpg) to select the label and close the window.
-1. Press the ![p4v_get_revision_button.jpg](p4v_get_revision_button.jpg) to close the _Get Revision_ window and begin the sync. The Log should show the progress of the sync.
+	Press the **Ok** button to select the label and close the window.
+1. Press the **Get Revision** button to close the **Get Revision** window and begin the sync. The **Log** should show the progress of the sync.
 
 	![p4v_sync_progress.jpg](p4v_sync_progress.jpg)(w:800 h:117)
 
@@ -173,13 +176,13 @@ Almost any file can be integrated into source control.
 
 The lifecycle of the files in your project can be managed via common commands:
 
-* check out
-* submit, or check in
-* revert, or roll back
-* revisions - history
+* Check out
+* Submit, or check in
+* Revert, or roll back
+* Revisions - history
 
 
-In addition to working with source control directly in Perforce, you can also use the editor tools. See [SCC Integration](Engine/UI/SourceControl) for information on how to manage content packages from within UnrealEd itself.
+In addition to working with source control directly in Perforce, you can also use the editor tools. See [SCC Integration](Engine/UI/SourceControl) for information on how to manage content packages from within the Editor itself.
 
  
 
