#UE4doc #docs:
Added engine version field
Changed font settings for multiple terms
Fixed capitalization errors

diff --git a/Engine/Documentation/Source/GettingStarted/DownloadingUnrealEngine/Perforce/P4Win/P4Win.INT.udn b/Engine/Documentation/Source/GettingStarted/DownloadingUnrealEngine/Perforce/P4Win/P4Win.INT.udn
--- a/Engine/Documentation/Source/GettingStarted/DownloadingUnrealEngine/Perforce/P4Win/P4Win.INT.udn
+++ b/Engine/Documentation/Source/GettingStarted/DownloadingUnrealEngine/Perforce/P4Win/P4Win.INT.udn
@@ -2,6 +2,7 @@ Availability: Licensee
 Title:Setting up P4Win for Unreal Engine 4
 Crumbs:%ROOT%, Engine, Programming, Programming/Basics, GettingStarted/DownloadingUnrealEngine/Perforce
 Description:Guide to installing, setting up, and syncing builds with the Perforce P4Win windows client for use with Unreal Engine 4.
+version: 4.9
 
 [TOC(start:2)]
 
@@ -17,17 +18,17 @@ For other general tips, see the [Setting up Perforce](GettingStarted/Downloading
 
 Run the P4Winst executable installer. 
 
-1. On the _Perforce Installation Configuration_ page, _User Installation_ is selected by default. Click _Next_. 
-1. On the _User Installation_ page, all of the components are selected by default. Click _Next_. 
-1. On the _Perforce Client Installation Directory_ page, determine the destination directory. **C:\Program Files\Perforce** is browsed to by default. Click _Next_. 
-1. On the _Perforce Client Configuration_ page, enter the following information:
-    * _Server Port_: **p4-licensee-proxy.epicgames.net:1667** (available through the US and East Asian VPN) or **p4-licensee-proxy-asia.epicgames.net:1667** (primary for those connected through the East Asian VPN).
-    * _Username_: provided in original UDN welcome email, typically your Team Name (**NOTE**: this is different from your UDN login information).
-1. Click _Next_. 
-1. On the _Perforce Installation Summary_ page, review the installation information. Click _Install_.
-1. On the _Perforce Installation Closing Summary_ page, choose whether or not you wish to reboot. Click _Finish_.
+1. On the **Perforce Installation Configuration** page, **User Installation** is selected by default. Click **Next**. 
+1. On the **User Installation** page, all of the components are selected by default. Click **Next**. 
+1. On the **Perforce Client Installation Directory** page, determine the destination directory. **C:\Program Files\Perforce** is browsed to by default. Click **Next**. 
+1. On the **Perforce Client Configuration** page, enter the following information:
+    * **Server Port**: **p4-licensee-proxy.epicgames.net:1667** (available through the US and East Asian VPN) or **p4-licensee-proxy-asia.epicgames.net:1667** (primary for those connected through the East Asian VPN).
+    * **Username**: provided in original UDN welcome email, typically your Team Name (**NOTE**: this is different from your UDN login information).
+1. Click **Next**. 
+1. On the **Perforce Installation Summary** page, review the installation information. Click **Install**.
+1. On the **Perforce Installation Closing Summary** page, choose whether or not you wish to reboot. Click **Finish**.
 
-NOTE: You can later reconfigure the IP address, as needed, via the _Settings_ menu. Please use the name to connect to the server, not the resolved IP address. The IP address may change if the server computer is changed; but the name will remain the same!
+NOTE: You can later reconfigure the IP address, as needed, via the **Settings** menu. Please use the name to connect to the server, not the resolved IP address. The IP address may change if the server computer is changed; but the name will remain the same!
 
 
 
@@ -39,23 +40,23 @@ If you are a content developer, or if you would simply like to run the example g
 
 1. Connect to Epic Games through the VPN connection.
 1. Run P4Win.
-1. If the initial _Client Workspace Wizard_ dialog appears, click _Cancel_.
-1. Select the _View Clientspecs_ button to show a list of Perforce Clients. If a list does not appear, click the _Refresh_ button.
+1. If the initial **Client Workspace Wizard** dialog appears, click **Cancel**.
+1. Select the **View Clientspecs** button to show a list of Perforce Clients. If a list does not appear, click the **Refresh** button.
 1. **Right-click** on the appropriate template:
-    * **IF YOU WANT SOURCE CODE ONLY** In the right _Perforce Clients_ pane, find **UE4_Template_Code**.
-    * **IF YOU WANT SOURCE CODE BUT NO UT3 CONTENT** In the right _Perforce Clients_ pane, find **UE4_Template_Code_NoUT3**.
-    * **IF YOU WANT EXECUTABLES AND NOT SOURCE** In the right _Perforce Clients_ pane, find **UE4_Template_Build**.
-1. Select _Create/Replace Client Using_ **UE4_Template_(Build/Code)** _as Template..._.  
+    * If you want source code only: In the right **Perforce Clients** pane, find **UE4_Template_Code**.
+    * If you want source code but no UT3 content: In the right **Perforce Clients** pane, find **UE4_Template_Code_NoUT3**.
+    * If you want executables and not source: In the right **Perforce Clients** pane, find **UE4_Template_Build**.
+1. Select **Create/Replace Client Using** **UE4_Template_(Build/Code)** **as Template...**.  
 1. A dialog will appear.  Fill in the following information:
-    * **IF YOU ARE A LICENSEE** _Client Name_: **UE4_Licensee_[TEAMNAME]** (or otherwise ending with your UDN username or team name)
-    * **IF YOU ARE A CONTRACTOR** _Client Name_: **UE4_Contractor_[USERNAME]** (or otherwise ending with your UDN username or team name)
-1. Click _OK_.
-1. A _Perforce Client Specification_ will appear. Set your _Root_ to a drive or directory on your local machine - something with many, many gigabytes free. **NOTE:** The full directory path should have no spaces, or else you may have trouble later on playing levels from the editor. All the other defaults are fine.
-1. Click _Update_.
-1. If the _Client Workspace Wizard_ dialog reappears, click _Cancel_ again.
-1. If you are prompted to do a full sync to head revision, decline the offer. You **really** do not want several gigabytes of data that you will have to replace shortly, do you?
-1. In the right pane, **right-click** on the clientspec you just created, and choose _Switch to UE4_Licensee_[USERNAME]_. (this may happen automatically)
-1. **Right-click** on the clientspec again, and choose _Set Default Client to UE4_Licensee_[USERNAME]_.
+    * If you are a licensee: **Client Name**: **UE4_Licensee_[TEAMNAME]** (or otherwise ending with your UDN username or team name)
+    * If you are a contractor: **Client Name**: **UE4_Contractor_[USERNAME]** (or otherwise ending with your UDN username or team name)
+1. Click **OK**.
+1. A **Perforce Client Specification** will appear. Set your **Root** to a drive or directory on your local machine - something with many, many gigabytes free. The full directory path should have no spaces, or else you may have trouble later on playing levels from the editor. All the other defaults are fine.
+1. Click **Update**.
+1. If the **Client Workspace Wizard** dialog reappears, click **Cancel** again.
+1. If you are prompted to do a full sync to head revision, decline the offer. This will prevent receiving several gigabytes of data that you will have to replace shortly.
+1. In the right pane, **Right-click** on the clientspec you just created, and choose **Switch to UE4_Licensee_[USERNAME]**. (this may happen automatically)
+1. **Right-click** on the clientspec again, and choose **Set Default Client to UE4_Licensee_[USERNAME]**.
 
 Now you are all set!
 
@@ -88,35 +89,35 @@ If you are using one of the defined templates, be sure that you are not viewing
 
 ### Depot - Client View of Directories
 
-To the left there should be _Client View of Depot_ pane. Expand the _//depot_ tree. There should be an _UnrealEngine3_ directory which contains everything to build and run UnrealEngine3. You may also see one or more of the following directories:
+To the left there should be **Client View of Depot** pane. Expand the **//depot** tree. There should be an UnrealEngine3 directory which contains everything to build and run UnrealEngine3. You may also see one or more of the following directories:
 
-* _PhysX_ - contains the common files for the physics middleware; depending on your agreement, you may also see additional directories and files.
+* PhysX - contains the common files for the physics middleware; depending on your agreement, you may also see additional directories and files.
 
 
 For more information, see the [directory layout](Engine/Basics/DirectoryStructure) page.
 
-When synchronizing, lots of downloading will progress in the log window. _Sync completed with no warnings or errors_ is what you are looking for when it is done.
+When synchronizing, lots of downloading will progress in the log window. "Sync completed with no warnings or errors" is what you are looking for when it is done.
 
 
 #### Unreal Engine 4
 
-In the _UnrealEngine3_ directory are several subdirectories: _Binaries_, _Development_, and _Engine_; and two game directories: _ExampleGame_ and _UTGame_.
+In the UnrealEngine3 directory are several subdirectories: Binaries, Development, and Engine; and two game directories: ExampleGame and UTGame.
 
-_ExampleGame_ is the most minimal game implementation offered by Epic, and can be a good place for getting started in creating your own game project if your game is not very similar to _Unreal Tournament_.  
+ExampleGame is the most minimal game implementation offered by Epic, and can be a good place for getting started in creating your own game project if your game is not very similar to Unreal Tournament.  
 
-_UTGame_ is based on Epic's _UT3_ &ndash; released in 2007. This is a good starting point for understanding how to create a commercial game in development, and thus provides a complex example showing exactly how Epic uses the engine. Be aware that the source available in UTGame is a lighter version than what was used in the final version. We recommend using _UTGame_ as an example, rather than extending from _UTGame_. You are welcome to use any C++ or Unrealscript code in _UTGame_ for your game project, but you cannot ship any of Epic's content (textures, meshes, levels, sounds, etc.) in your game.
+UTGame is based on Epic's UT3 &ndash; released in 2007. This is a good starting point for understanding how to create a commercial game in development, and thus provides a complex example showing exactly how Epic uses the engine. Be aware that the source available in UTGame is a lighter version than what was used in the final version. We recommend using UTGame as an example, rather than extending from UTGame. You are welcome to use any C++ or Unrealscript code in UTGame for your game project, but you cannot ship any of Epic's content (textures, meshes, levels, sounds, etc.) in your game.
 
-For _UnrealEngine3_, most of Epic's content is located in the game-specific directories, so if all you want is the engine source, you can simply sync to the //depot/UnrealEngine3/Development directory. You can highlight specific directories as desired, for example, to grab the latest content for _ExampleGame_ without updating _UTGame_ content. 
+For UnrealEngine3, most of Epic's content is located in the game-specific directories, so if all you want is the engine source, you can simply sync to the //depot/UnrealEngine3/Development directory. You can highlight specific directories as desired, for example, to grab the latest content for ExampleGame without updating UTGame content. 
 
 
 ### Sync - Download
 
-Once you have selected the directories to update, **right-click** and select _Sync to Revision, Label, Changelist, or Date..._. Choose the _Other_ radio button to sync to either a specific changelist or label.
+Once you have selected the directories to update, **Right-click** and select **Sync to Revision, Label, Changelist, or Date...**. Choose the **Other** radio button to sync to either a specific changelist or label.
 
 
 #### Get the Files
 
-Once either a changelist or label has been selected. Click _Sync Files_. Do NOT sync to the head revision as it may contain code that does not compile, or worse, is subtly broken.
+Once either a changelist or label has been selected. Click **Sync Files**. Do NOT sync to the head revision as it may contain code that does not compile, or worse, is subtly broken.
 
 
 
@@ -126,13 +127,13 @@ Almost any file can be integrated into source control.
 
 The lifecycle of the files in your project can be managed via common commands:
 
-* check out
-* submit, or check in
-* revert, or roll back
-* revisions - history
+* Check out
+* Submit, or check in
+* Revert, or roll back
+* Revisions - history
 
 
-In addition to working with source control directly in Perforce, you can also use the editor tools. See [SCC Integration](Engine/UI/SourceControl) for information on how to manage content packages from within UnrealEd itself.
+In addition to working with source control directly in Perforce, you can also use the editor tools. See [SCC Integration](Engine/UI/SourceControl) for information on how to manage content packages from within the Editor itself.
 
  
 
