#UE4doc #docs:
Added engine version field
Changed font settings for a few terms

diff --git a/Engine/Documentation/Source/GettingStarted/DownloadingUnrealEngine/Perforce/Setup/Setup.INT.udn b/Engine/Documentation/Source/GettingStarted/DownloadingUnrealEngine/Perforce/Setup/Setup.INT.udn
--- a/Engine/Documentation/Source/GettingStarted/DownloadingUnrealEngine/Perforce/Setup/Setup.INT.udn
+++ b/Engine/Documentation/Source/GettingStarted/DownloadingUnrealEngine/Perforce/Setup/Setup.INT.udn
@@ -2,6 +2,7 @@ Availability: Licensee
 Title:Setting Up a Perforce Connection
 Crumbs:%ROOT%, Engine, Programming, Programming/Basics, GettingStarted/DownloadingUnrealEngine/Perforce
 Description:Guide to connecting to Epic's Perforce server to obtain Unreal Engine 4 builds.
+version: 4.9
 
 [TOC(start:2)]
 
@@ -187,10 +188,10 @@ The script can be called as:
     
 
 
-
-NOTE: This only works on Windows.  Also, if running on Windows Server 2008 R2, Vista or Windows 7, last access date is 
+[REGION:note]
+This only works on Windows.  Also, if running on Windows Server 2008 R2, Vista or Windows 7, last access date is 
 disabled by default.  You need to "fsutil behavior set disablelastaccess 0" to reenable last accessed date.
-
+[/REGION]
 
 ### Scheduled Downtime
 
