#UE4doc #docs:
Added engine version field
Changed font settings for multiple terms
Replaced multiple images as they did not meet the guidelines

diff --git a/Engine/Documentation/Source/GettingStarted/DownloadingUnrealEngine/Perforce/Syncing/Syncing.INT.udn b/Engine/Documentation/Source/GettingStarted/DownloadingUnrealEngine/Perforce/Syncing/Syncing.INT.udn
--- a/Engine/Documentation/Source/GettingStarted/DownloadingUnrealEngine/Perforce/Syncing/Syncing.INT.udn
+++ b/Engine/Documentation/Source/GettingStarted/DownloadingUnrealEngine/Perforce/Syncing/Syncing.INT.udn
@@ -2,6 +2,7 @@ Availability: Licensee
 Title:Syncing Builds with Perforce
 Crumbs:%ROOT%, Engine, Programming, Programming/Basics, GettingStarted/DownloadingUnrealEngine/Perforce
 Description:Guide to obtaining builds of Unreal Engine 4 using Perforce.
+version: 4.9
 
 [TOC(start:2)]
 
@@ -21,27 +22,27 @@ The tree pane in P4V contains your Workspace and the Depot Tree. In the Depot Tr
 ![p4v_depot_tree.jpg](p4v_depot_tree.jpg)(w:374 h:307)
 
 
-* _CodeDrop_ - Contains various snapshots of games created by Epic that can be used as examples.
-* _PhysX_ - Contains the common files for the physics middleware; depending on your agreement, you may also see additional directories and files.
-* _UnrealEngine3_ - Contains everything to build and run UnrealEngine3. Also includes the ExampleGame and UDKGame examples.
-* _UnrealMCP_ - Contains the source for UnrealMCP, which is a SOA (service-oriented architecture) that focuses on providing highly scalable and available set of services to enhance the online experience of players.
+* CodeDrop - Contains various snapshots of games created by Epic that can be used as examples.
+* PhysX - Contains the common files for the physics middleware; depending on your agreement, you may also see additional directories and files.
+* UnrealEngine3 - Contains everything to build and run UnrealEngine3. Also includes the ExampleGame and UDKGame examples.
+* UnrealMCP - Contains the source for UnrealMCP, which is a SOA (service-oriented architecture) that focuses on providing highly scalable and available set of services to enhance the online experience of players.
 
 
 For more information, see the [directory layout](Engine/Basics/DirectoryStructure) page.
 
-In the _UnrealEngine3_ directory are several subdirectories: _Binaries_, _Development_, and _Engine_; and two game directories: _ExampleGame_ and _UDKGame_.
+In the UnrealEngine3 directory are several subdirectories: Binaries, Development, and Engine; and two game directories: ExampleGame and UDKGame.
 
 ![p4v_depot_tree_ue3.jpg](p4v_depot_tree_ue3.jpg)(w:374 h:307)
 
-_ExampleGame_ is the most minimal game implementation offered by Epic, and can be a good place for getting started in creating your own game project if your game is not very similar to _Unreal Tournament_.  
+ExampleGame is the most minimal game implementation offered by Epic, and can be a good place for getting started in creating your own game project if your game is not very similar to Unreal Tournament.  
 
-_UDKGame_ is Epic's current Unreal Engine 4 usage demo based off of _Unreal Tournament 3_. This provides a complex example showing exactly how Epic intends to use the engine. We recommend using _UDKGame_ as an example, rather than extending from _UDKGame_. You are welcome to use any C++ or UnrealScript code in _UDKGame_ for your game project, but you cannot ship any of Epic's content (textures, meshes, levels, sounds, etc.) in your game.
+UDKGame is Epic's current Unreal Engine 4 usage demo based off of Unreal Tournament 3. This provides a complex example showing exactly how Epic intends to use the engine. We recommend using UDKGame as an example, rather than extending from UDKGame. You are welcome to use any C++ or UnrealScript code in UDKGame for your game project, but you cannot ship any of Epic's content (textures, meshes, levels, sounds, etc.) in your game.
 
-For _UnrealEngine3_, most of Epic's content is located in the game-specific directories, so if all you want is the engine source, you can simply sync to the //depot/UnrealEngine3/Development directory. You can highlight specific directories as desired, for example, to grab the latest content for _ExampleGame_ without updating _UDKGame_ content.
+For UnrealEngine3, most of Epic's content is located in the game-specific directories, so if all you want is the engine source, you can simply sync to the //depot/UnrealEngine3/Development directory. You can highlight specific directories as desired, for example, to grab the latest content for ExampleGame without updating UDKGame content.
 
 ![p4v_depot_tree_ue3_examplegame.jpg](p4v_depot_tree_ue3_examplegame.jpg)(w:374 h:307)
 
-When getting a revision, lots of downloading will progress in the log window. _Sync completed with no warnings or errors_ is what you are looking for when it is done.
+When getting a revision, lots of downloading will progress in the log window. "Sync completed with no warnings or errors" is what you are looking for when it is done.
 
 
 
@@ -54,33 +55,33 @@ To sync to a QA-approved build, follow the step below:
 
 	![p4v_small.jpg](p4v_small.jpg)(w:800 h:525)
 
-1. In the Depot Tree, **right-click** on the UnrealEngine3 directory and then choose _Get Revision..._ from the context menu. Alternatively, you can select the _UnrealEngine3_ directory and choose _Get Revision..._ from the _Actions_ menu.
+1. In the Depot Tree, **right-click** on the UnrealEngine3 directory and then choose **Get Revision...** from the context menu. Alternatively, you can select the UnrealEngine3 directory and choose **Get Revision...** from the **Actions** menu.
 
 	![p4v_depot_menu.jpg](p4v_depot_menu.jpg)(w:304 h:650)
 
-1. The _Get Revision_ window will open.
+1. The **Get Revision** window will open.
 
 	![p4v_get_revision.jpg](p4v_get_revision.jpg)(w:634 h:488)
 
-	Select the _Specify revision using_ radio button to enable the combobox and text field.
+	Select the **Specify revision using** radio button to enable the combobox and text field.
 
 	![p4v_specify_revision.jpg](p4v_specify_revision.jpg)(w:513 h:63)
 
-1. Select _label_ from the combobox.
+1. Select **label** from the combobox.
 
 	![p4v_revision_label.jpg](p4v_revision_label.jpg)(w:100 h:103)
 
-	Press the ![p4v_browse_button.jpg](p4v_browse_button.jpg) button to open the _Select Label_ window.
+	Press the **Browse...** button to open the **Select Label** window.
 
 	![p4v_select_label.jpg](p4v_select_label.jpg)(w:800 h:526)
 
-1. Select the label you want to sync to. Usually this is the latest QA-approved build, designated by a label in the form of _QA_APPROVED_BUILD_[Month]_[Year]_. You may need to scroll down in the list a few times for all the labels to be loaded.
+1. Select the label you want to sync to. Usually this is the latest QA-approved build, designated by a label in the form of **QA_APPROVED_BUILD_[Month]_[Year]**. You may need to scroll down in the list a few times for all the labels to be loaded.
 
 	![p4v_label_build.jpg](p4v_label_build.jpg)(w:800 h:526)
 
-	Press the ![p4v_ok_button.jpg](p4v_ok_button.jpg) to select the label and close the window.
+	Press the **Ok** button to select the label and close the window.
 
-1. Press the ![p4v_get_revision_button.jpg](p4v_get_revision_button.jpg) to close the _Get Revision_ window and begin the sync. The Log should show the progress of the sync.
+1. Press the **Get Revision** button to close the **Get Revision** window and begin the sync. The Log should show the progress of the sync.
 
 	![p4v_sync_progress.jpg](p4v_sync_progress.jpg)(w:800 h:117)
 
@@ -98,7 +99,7 @@ If you are having problems doing a full sync to the latest QA-Approved Build, tr
 
 While doing a sync, you may receive warnings that you  "can't clobber writable file" it is likely that some files in your client workspace were marked writeable. As a safeguard, Perforce will not overwrite any file that is writeable on your client under any circumstance. If you think you have changes you need to keep, you will need to determine what files have changed and either check them in; otherwise, if you do not need any changes and just want a fresh sync, perform a "force sync" to the QA label to force those files to be overwritten.
 
-When you have problems with the build, be sure to check the latest build notes. It is also a good idea to _Check Consistency_ in Perforce (from the _File_ menu, under the _More_ submenu in [P4Win](GettingStarted/DownloadingUnrealEngine/Perforce/P4Win); from the _Actions_ menu, under the _Reconcile Offline Work_ submenu in [P4V](GettingStarted/DownloadingUnrealEngine/Perforce/P4V)). As a last resort, use _Remove from Client_; and then physically move (or delete) the entire root folder and its contents. From there, sync as normal. The _Remove from Client_ on the depot is necessary to tell the server you have nothing, which will force all files to be updated.
+When you have problems with the build, be sure to check the latest build notes. It is also a good idea to **Check Consistency** in Perforce (from the **File** menu, under the **More** submenu in [P4Win](GettingStarted/DownloadingUnrealEngine/Perforce/P4Win); from the **Actions** menu, under the **Reconcile Offline Work** submenu in [P4V](GettingStarted/DownloadingUnrealEngine/Perforce/P4V)). As a last resort, use **Remove from Client**; and then physically move (or delete) the entire root folder and its contents. From there, sync as normal. The **Remove from Client** on the depot is necessary to tell the server you have nothing, which will force all files to be updated.
 
  
 
