#UE4doc #docs:
Added engine version field
Removed a couple images as they did not match the guidelines

diff --git a/Engine/Documentation/Source/GettingStarted/DownloadingUnrealEngine/VPNSetup/VPNSetup.INT.udn b/Engine/Documentation/Source/GettingStarted/DownloadingUnrealEngine/VPNSetup/VPNSetup.INT.udn
--- a/Engine/Documentation/Source/GettingStarted/DownloadingUnrealEngine/VPNSetup/VPNSetup.INT.udn
+++ b/Engine/Documentation/Source/GettingStarted/DownloadingUnrealEngine/VPNSetup/VPNSetup.INT.udn
@@ -2,6 +2,7 @@ Availability: Licensee
 Title:Setting up a VPN Connection
 Crumbs:%ROOT%, Engine, Programming, Programming/Basics
 Description:Guide to connecting to Epic's virtual private network (VPN) in order to access the Perforce server and obtain the engine.
+version: 4.9
 
 [TOC(start:2)]
 
@@ -65,7 +66,7 @@ NOTE: This is a web-based VPN client. It has been reported to be more stable whe
 Be sure to disconnect from any existing VPN sessions you have open!
 [/REGION]
 
-1. To begin the web installation of the Cisco AnyConnect VPN Client, go to **[https://vpn.epicgames.com](https://vpn.epicgames.com)**.
+1. To begin the web installation of the Cisco AnyConnect VPN Client, go to [https://vpn.epicgames.com](https://vpn.epicgames.com).
 
 	**Regional VPN Mirror Sites:**
 
@@ -172,7 +173,7 @@ Connecting to the VPN with AnyConnect is a straightforward process.
 
 	![username and password](usernamepassword.png)
 
-1. Press the ![vpn_anyconnect_connect_button.jpg](vpn_anyconnect_connect_button.jpg) to connect to the VPN. When the connection is made, the ![vpn_anyconnect_connected.jpg](vpn_anyconnect_connected.jpg) icon will show in the taskbar.
+1. Press the **Connect** button to connect to the VPN. When the connection is made, the ![vpn_anyconnect_connected.jpg](vpn_anyconnect_connected.jpg) icon will show in the taskbar.
 [/EXCERPT:AnyConnect]
 
 Once you set up the VPN connection, you can continue on to [set up Perforce](GettingStarted/DownloadingUnrealEngine/Perforce/Setup).
