#UE4doc #docs:
Replaced a typo

diff --git a/Engine/Documentation/Source/GettingStarted/DownloadingUnrealEngine/Perforce/Integration/Integration.INT.udn b/Engine/Documentation/Source/GettingStarted/DownloadingUnrealEngine/Perforce/Integration/Integration.INT.udn
--- a/Engine/Documentation/Source/GettingStarted/DownloadingUnrealEngine/Perforce/Integration/Integration.INT.udn
+++ b/Engine/Documentation/Source/GettingStarted/DownloadingUnrealEngine/Perforce/Integration/Integration.INT.udn
@@ -97,7 +97,7 @@ To integrate a single change in P4Win:
 
 ### What is the safest way to resolve files quickly?
 
-Each file integrated MUST be resolved before checking in. Difficult conflicts will sometimes require help from the original [TOC(start:2)]. To ease the process, Perforce has a few options for auto-resolving simple merges to pare the list down. P4V and P4Win both provide access to these options. The safest way to proceed is in the following order:
+Each file integrated must be resolved before checking in. Difficult conflicts will sometimes require help from the original submitter. To ease the process, Perforce has a few options for auto-resolving simple merges to pare the list down. P4V and P4Win both provide access to these options. The safest way to proceed is in the following order:
 1. **Safe automatic resolve** - auto-resolve changes if the file has only changed in one branch since the last merge. This is always safe because it can just copy the changed file over.
 1. **Automatic resolve** - resolves changes if the file has changed in both branches but no conflicts exist.
 1. **Interactive resolve** - any files left will be the hard ones where changes in both branches conflicted.
