#UE4doc #docs:
Added engine version field
Changed font settings for multiple terms
Fixed a typo
Fixed captalization errors

diff --git a/Engine/Documentation/Source/GettingStarted/FromUnity/UnrealEngineforUnityDevelopers.INT.udn b/Engine/Documentation/Source/GettingStarted/FromUnity/UnrealEngineforUnityDevelopers.INT.udn
--- a/Engine/Documentation/Source/GettingStarted/FromUnity/UnrealEngineforUnityDevelopers.INT.udn
+++ b/Engine/Documentation/Source/GettingStarted/FromUnity/UnrealEngineforUnityDevelopers.INT.udn
@@ -2,6 +2,7 @@ Availability: Public
 Title:Unreal Engine 4 For Unity Developers
 Crumbs: %ROOT%, GettingStarted
 Description: Translate your Unity knowledge into UE4 so you can get up to speed quickly.
+version: 4.9
 
 
 ![](image_0.png)
@@ -12,13 +13,13 @@ This guide provides an overview of UE4 from the viewpoint of a Unity user, and a
 
 ## The Editor
 
-Below are pictures of the Unity Editor and Unreal Editor, **color-coded** to indicate common functionality. Each block is labeled to show the equivalent UE4 terminology.  Unreal Editor's layout is fully customizable by dragging and dropping tabs around.
+Below are pictures of the Unity Editor and Unreal Editor, color-coded to indicate common functionality. Each block is labeled to show the equivalent UE4 terminology.  Unreal Editor's layout is fully customizable by dragging and dropping tabs around.
 
 ![](image_1.png)
 
 ### Editing Assets
 
-In Unity, the Inspector tab is used to edit selected assets in your Project.  In UE4, the **Details** panel exposes properties of selected objects, while larger editing tasks lead to a **dedicated window or tab**.  New tabbed windows will open for each edited asset, **similar to a web browser**.  Of course, you can freely drag the tabs around or float them as standalone windows.
+In Unity, the Inspector tab is used to edit selected assets in your Project.  In UE4, the **Details** panel exposes properties of selected objects, while larger editing tasks lead to a dedicated window or tab.  New tabbed windows will open for each edited asset, similar to a web browser.  Of course, you can freely drag the tabs around or float them as standalone windows.
 
 ![](image_2.png)
 
@@ -56,7 +57,7 @@ The following section contains common Unity terms on the left and their UE4 equi
 
 ### So what are all these directories and files?
 
-Just like Unity projects, Unreal projects always exist in their own directory and have their own project file.  You can **double-click** on **.uproject** files to load your game into the Unreal Editor, or **right-click** for additional options.  Project folders have various sub-folders that contain your game's content and source as well as various configuration files and binaries.  The most important are the **Content **and **Source **sub-folders.
+Just like Unity projects, Unreal projects always exist in their own directory and have their own project file.  You can **Double-click** on **.uproject** files to load your game into the Unreal Editor, or **Right-click** for additional options.  Project folders have various sub-folders that contain your game's content and source as well as various configuration files and binaries.  The most important are the **Content **and **Source** sub-folders.
 
 ### Where do I put my assets? 
 
@@ -89,9 +90,9 @@ All project settings can be found from the main menu under **Edit** / **Project
 
 In Unity, you are accustomed to placing C# source files in your assets folder.  
 
-UE4 works differently.  For projects that have C++ code, you will find a "**Source**" sub-folder under the project directory that has various files, including C++ source (.cpp) and header (.h) files, as well as some build scripts (.Build.cs, .Target.cs)  However, Blueprint-only projects will not have a Source folder.
+UE4 works differently.  For projects that have C++ code, you will find a **Source** sub-folder under the project directory that has various files, including C++ source (.cpp) and header (.h) files, as well as some build scripts (.Build.cs, .Target.cs)  However, Blueprint-only projects will not have a Source folder.
 
-The easiest way to get started with C++ in UE4 is to use the editor to **Add Code to Project** (in the main File menu), or to simply create a new C++ project from scratch from one of the many templates.  You can find C++ classes right in the **Content Browser** and can open the files in Visual Studio or Xcode by **double-clicking** on their icons.
+The easiest way to get started with C++ in UE4 is to use the editor to **Add Code to Project** (in the main File menu), or to simply create a new C++ project from scratch from one of the many templates.  You can find C++ classes right in the **Content Browser** and can open the files in Visual Studio or Xcode by **Double-clicking** on their icons.
 
 ## From GameObjects to Actors
 
@@ -133,14 +134,14 @@ SknellY4I4g
 [/PARAMLITERAL]
 [/OBJECT]
 
-In Unity, a GameObject hold a flat list of components, but in UE4 an Actor actually contains a **hierarchy **of components attached to one another.  You can see this in the example above, where the Light and Particle are attached to the Mesh.  This has some important implications discussed later in [Complex Actors and GameObjects](#compoundobjects).
+In Unity, a GameObject hold a flat list of components, but in UE4 an Actor actually contains a **Hierarchy **of components attached to one another.  You can see this in the example above, where the Light and Particle are attached to the Mesh.  This has some important implications discussed later in [Complex Actors and GameObjects](#compoundobjects).
 
 ### From Unity prefabs to UE4 Blueprint Classes
 
 Unity's workflow is based on prefabs.  In Unity you build a set of GameObjects with components, then create a prefab from them. You can then place instances of the prefab in your world, or instantiate them at runtime.
 
 
-UE4's corresponding workflow is based on Blueprint Classes.  In UE4, you build an Actor with components, select it, and click the "**Blueprint / Add Script**" button (in the **Details** panel).  Then, choose a place to save your Blueprint Class, and click "**Create Blueprint**" to save your new Blueprint Class!
+UE4's corresponding workflow is based on Blueprint Classes.  In UE4, you build an Actor with components, select it, and click the **Blueprint / Add Script** button (in the **Details** panel).  Then, choose a place to save your Blueprint Class, and click **Create Blueprint** to save your new Blueprint Class!
 
 [OBJECT:EmbeddedVideo]
 [PARAMLITERAL:width]
@@ -154,7 +155,7 @@ n35cDTFBl3s
 [/PARAMLITERAL]
 [/OBJECT]
 
-Your new Blueprint Classes can be found in the **Content Browser**.  You can **double-click** to edit them directly, or you can drag and drop them into any level. 
+Your new Blueprint Classes can be found in the **Content Browser**.  You can **Double-click** to edit them directly, or you can drag and drop them into any level. 
 
 ### Where is Script Component and MonoBehaviour?
 
@@ -318,7 +319,7 @@ There are several sub-classed Actor types you will quickly encounter in the engi
 
 Not everything.  **Actors** are most common class used for gameplay in UE4 and the only type that can be **Spawned** into the **World**.  So everything you place in your level will be an Actor.
 
-The other important type to know about is **Object**.  Object is actually the base class of all Unreal classes, including **Actor** and many others. This is a much lower-level construct than Actor, but still has features that you will expect from a Unreal class, such as **reflection** and **serialization**. Object is a very basic class we use when we need to define a new type that just does not fit the mold of an **Actor**.  For example, **Actor Component** is the base class for all Components, and derives from **Object** instead of **Actor**.
+The other important type to know about is **Object**.  Object is actually the base class of all Unreal classes, including **Actor** and many others. This is a much lower-level construct than Actor, but still has features that you will expect from a Unreal class, such as **Reflection** and **Serialization**. Object is a very basic class we use when we need to define a new type that just does not fit the mold of an **Actor**.  For example, **Actor Component** is the base class for all Components, and derives from **Object** instead of **Actor**.
 
 ### What is this Gameplay Framework thing in UE4?
 
@@ -334,11 +335,11 @@ To use the Gameplay Framework, you really just need to learn the customized buil
 
 ### So, I'm used to programming in MonoDevelop
 
-For Blueprint scripting, you will only need Unreal Editor -- everything is built in!  To write code in C++, download the [free version of Visual Studio](http://go.microsoft.com/fwlink/?LinkId=525621) on Windows, or [install Xcode](https://itunes.apple.com/us/app/xcode/id497799835?mt=12) on Mac.  When you first create a new project (or add code to an existing project), UE4 will **automatically create Visual Studio project files **for you.  You can open Visual Studio by **double-clicking** on a C++ class inside **Content Browser**, or by clicking the "Open Visual Studio" button in the main File menu.
+For Blueprint scripting, you will only need Unreal Editor -- everything is built in!  To write code in C++, download the [free version of Visual Studio](http://go.microsoft.com/fwlink/?LinkId=525621) on Windows, or [install Xcode](https://itunes.apple.com/us/app/xcode/id497799835?mt=12) on Mac.  When you first create a new project (or add code to an existing project), UE4 will automatically create Visual Studio project files for you.  You can open Visual Studio by **Double-clicking** on a C++ class inside **Content Browser**, or by clicking the **Open Visual Studio** button in the main File menu.
 
 ![image alt text](image_12.png)
 
-One important difference in UE4:  You will sometimes have to manually refresh your Visual Studio **project files** (for example, after downloading a new version of UE4, or when manually making changes to source file locations on disk.)  You can do this by clicking "Refresh Visual Studio Project" in the main menu, or by **right-clicking** on the **.uproject file** in your project's directory and selecting "Generate Visual Studio project files"
+One important difference in UE4:  You will sometimes have to manually refresh your Visual Studio project files (for example, after downloading a new version of UE4, or when manually making changes to source file locations on disk.)  You can do this by clicking **Refresh Visual Studio Project** in the main menu, or by **Right-clicking** on the **.uproject file** in your project's directory and selecting **Generate Visual Studio project files**.
 
 ![image alt text](image_13.png)
 
@@ -356,7 +357,7 @@ In Unity, we might have a simple component that looks like this:
 	}
 
 
-But remember, in UE4 you can write code **right on the Actor itself** rather than only coding new component types.  This is actually very common and useful.
+But remember, in UE4 you can write code right on the Actor itself rather than only coding new component types.  This is actually very common and useful.
 
 Similar to Unity's Start, OnDestroy, and Update functions, we have a similar set of methods for Actors in UE4:
 
@@ -1036,7 +1037,7 @@ If you want to move something 2 Feet in Unity that would be 0.61units (meter) an
 
 ### What is with this coordinate system? Which way is up?
 
-Both Unity and UE4 use a left-handed coordinate system, but the axes are swapped around.  In UE4, positive X is "forward", positive Y is "right" and positive Z is "up".  This cannot be changed, so you will have to get used to it.
+Both Unity and UE4 use a left-handed coordinate system, but the axes are swapped around.  In UE4, positive X is "forward", positive Y is "right" and positive Z is "up".
 
 ### How do I see Log Output from my game?
 
