#UE4doc #docs:
Added engine version field
Unbolded multiple terms

diff --git a/Engine/Documentation/Source/GettingStarted/HowTo/IntroHowToOverview.INT.udn b/Engine/Documentation/Source/GettingStarted/HowTo/IntroHowToOverview.INT.udn
--- a/Engine/Documentation/Source/GettingStarted/HowTo/IntroHowToOverview.INT.udn
+++ b/Engine/Documentation/Source/GettingStarted/HowTo/IntroHowToOverview.INT.udn
@@ -2,6 +2,7 @@ Availability: Public
 Title:Basic How To's
 Crumbs: %ROOT%, GettingStarted
 Description: Several short How To guides provide step-by-step instructions to working with Unreal Engine 4. 
+version: 4.9
 
 [VAR:Topic]
 [OBJECT:Topic]
@@ -44,10 +45,10 @@ Refer to the table below for more information and click each link to view its co
 
 |How To | Description |
 | ------- | ------------ |
-| **[Create/Open Projects](Engine/Basics/Projects/Browser)** | %Engine/Basics/Projects/Browser:description% |
-| **[](Engine\Actors\Placement)** | %Engine\Actors\Placement:description% |
-| **[](GettingStarted\HowTo\ManipulatingActors)** | %GettingStarted\HowTo\ManipulatingActors:description% |
-| **[](GettingStarted\HowTo\DetailsPanel)** | %GettingStarted\HowTo\DetailsPanel:description% |
-| **[](GettingStarted\HowTo\PIE)** | %GettingStarted\HowTo\PIE:description% |
-| **[](GettingStarted\HowTo\FindingHelp)** | %GettingStarted\HowTo\FindingHelp:description% |
+| [Create/Open Projects](Engine/Basics/Projects/Browser) | %Engine/Basics/Projects/Browser:description% |
+| [](Engine\Actors\Placement) | %Engine\Actors\Placement:description% |
+| [](GettingStarted\HowTo\ManipulatingActors) | %GettingStarted\HowTo\ManipulatingActors:description% |
+| [](GettingStarted\HowTo\DetailsPanel) | %GettingStarted\HowTo\DetailsPanel:description% |
+| [](GettingStarted\HowTo\PIE) | %GettingStarted\HowTo\PIE:description% |
+| [](GettingStarted\HowTo\FindingHelp) | %GettingStarted\HowTo\FindingHelp:description% |
 
