#UE4doc #docs:
Replaced multiple images as they were out of date

diff --git a/Engine/Documentation/Source/GettingStarted/HowTo/DetailsPanel/DetailsPanel.INT.udn b/Engine/Documentation/Source/GettingStarted/HowTo/DetailsPanel/DetailsPanel.INT.udn
--- a/Engine/Documentation/Source/GettingStarted/HowTo/DetailsPanel/DetailsPanel.INT.udn
+++ b/Engine/Documentation/Source/GettingStarted/HowTo/DetailsPanel/DetailsPanel.INT.udn
@@ -31,11 +31,7 @@ First let us pick our Actor, place it in the level, and scale it up a bit.
 
 1.	When the Actor is placed, the **Details** panel will become populated. 
 
-	[REGION:lightbox]
-	[![](Details2.png)(w:800)](Details2.png)
-	[/REGION]
-
-	_Click for full image._
+	![](Details2.png)
 
 1.	Inside the **Details** panel, under _Transform_, click the lock icon for **Scale** and set the **XYZ** values to **3.0**. 
 
