#UE4doc #docs - Reviewed UEDOC-2307. Corrected several issues with images to be clearly seen, and couple of capitalization errors.

diff --git a/Engine/Documentation/Source/Engine/Rendering/Materials/ExpressionReference/Utility/UtilityExpressions.INT.udn b/Engine/Documentation/Source/Engine/Rendering/Materials/ExpressionReference/Utility/UtilityExpressions.INT.udn
--- a/Engine/Documentation/Source/Engine/Rendering/Materials/ExpressionReference/Utility/UtilityExpressions.INT.udn
+++ b/Engine/Documentation/Source/Engine/Rendering/Materials/ExpressionReference/Utility/UtilityExpressions.INT.udn
@@ -349,7 +349,7 @@ The **Distance To Nearest Surface** Material Expression node allows Materials to
 This Material Expression works by outputting the signed distance in world space units from the distance field to the nearest occluders in the scene.
 
 [region:note]
-Please note that this Material Expression node requires that your project have the **Generate Mesh Distance Fields** enabled in order to work correctly. If you have not done so already, enabled this property in your projects settings under the Lighting category before going any further. 
+Please note that this Material Expression node requires that your project have the **Generate Mesh Distance Fields** enabled in order to work correctly. If you have not done so already, enabled this property in your Projects Settings under the Rendering category before going any further. 
 [/region]
 
 | Item | Description |
@@ -358,7 +358,12 @@ Please note that this Material Expression node requires that your project have t
 
 Here is an example of the **Distance To Nearest Surface** Material Expression in action. 
 
-![](DistanceNearSurface_00.png)
+
+[Region:lightbox]
+[![](DistanceNearSurface_01.png)(w:960)](DistanceNearSurface_01.png)
+[/Region]
+
+![](DistanceNearSurface_02.png)(w:550)
 
 In this example the Distance To Nearest Surface was fed into the Opacity input on a Material and that Material was applied to a Static Mesh plane that was placed just above the levels floor. 
 What the Distance To Nearest Surface is doing is telling the Material to only color areas red were the Static Meshes plane will start to intersect other Static Meshes placed in the scene.
@@ -369,7 +374,7 @@ The **DistanceFieldGradient** Material Expression node, when normalized, outputs
 This makes the Distance Field Gradient Material Expression node perfect for Materials that need to simulate the flow of liquids.
 
 [region:note]
-Please note that this Material Expression node requires that your project have the **Generate Mesh Distance Fields** enabled in order to work correctly. If you have not done so already, enabled this property in your projects settings under the Lighting category before going any further. 
+Please note that this Material Expression node requires that your project have the **Generate Mesh Distance Fields** enabled in order to work correctly. If you have not done so already, enabled this property in your Projects Settings under the Rendering category before going any further. 
 [/region]
 
 | Item | Description |
@@ -379,15 +384,22 @@ Please note that this Material Expression node requires that your project have t
 Here is an example of how to use the **DistanceFieldGradient** Material Expression in your Materials.
 In this example below make sure to note that the DistanceFieldGradient was first normalized and then input into a Mask Channel node.
 The reason for this is because without normalizing the DistanceFieldGradient first you can not get directional data. 
-The Mask Channel Parameters was added to allow for easier RGB channel switching with in the Material Instance.
+The Mask Channel Parameter was added to allow for easier RGB channel switching with in the Material Instance.
 
-![](DistanceFieldGradient_Material.png)
+[Region:lightbox]
+[![](DistanceFieldGradient_Material.png)](DistanceFieldGradient_Material.png)
+[/Region]
 
 Here is an example of the DistanceFieldGradient in action. 
 The image below shows what data the DistanceFieldGradient will use when the various RGB are enabled.
 
-![](DistanceFieldGradient_Example.png)
+[Region:lightbox]
+[![](DistanceFieldGradient_Example.png)(w:960)](DistanceFieldGradient_Example.png)
+[/Region]
 
+<!---
+![](DistanceFieldGradient_Example.png)
+--->
 
 | Number | Description |
 | -------- | ----------- |
