#UE4doc #docs - Added engine version field.

diff --git a/Engine/Documentation/Source/Resources/ContentExamples/MaterialProperties/1_7/1_7.INT.udn b/Engine/Documentation/Source/Resources/ContentExamples/MaterialProperties/1_7/1_7.INT.udn
--- a/Engine/Documentation/Source/Resources/ContentExamples/MaterialProperties/1_7/1_7.INT.udn
+++ b/Engine/Documentation/Source/Resources/ContentExamples/MaterialProperties/1_7/1_7.INT.udn
@@ -13,6 +13,7 @@ Related: Engine/Rendering/Materials/LayeredMaterials
 Related: Engine/Rendering/Materials/MaterialInstances
 Related: Engine/Actors/DecalActor
 Related: Engine/Rendering/Materials/MaterialProperties
+Version: 4.9
 
 [VAR:Steps]
 [OBJECT:Navigation]
diff --git a/Engine/Documentation/Source/Resources/ContentExamples/MaterialProperties/1_8/1_8.INT.udn b/Engine/Documentation/Source/Resources/ContentExamples/MaterialProperties/1_8/1_8.INT.udn
--- a/Engine/Documentation/Source/Resources/ContentExamples/MaterialProperties/1_8/1_8.INT.udn
+++ b/Engine/Documentation/Source/Resources/ContentExamples/MaterialProperties/1_8/1_8.INT.udn
@@ -12,6 +12,7 @@ Related: Engine/Rendering/Materials/LayeredMaterials
 Related: Engine/Rendering/Materials/MaterialInstances
 Related: Engine/Actors/DecalActor
 Related: Engine/Rendering/Materials/MaterialProperties
+Version: 4.9
 
 [VAR:Steps]
 [OBJECT:Navigation]
