#UE4doc #docs - added engine version field. removed overview header. removed extraneous bolding.

diff --git a/Engine/Documentation/Source/Resources/ContentExamples/MathHall/MathHall.INT.udn b/Engine/Documentation/Source/Resources/ContentExamples/MathHall/MathHall.INT.udn
--- a/Engine/Documentation/Source/Resources/ContentExamples/MathHall/MathHall.INT.udn
+++ b/Engine/Documentation/Source/Resources/ContentExamples/MathHall/MathHall.INT.udn
@@ -5,6 +5,7 @@ Description:Overview of the samples provided in the Math Hall Content Examples.
 Related: Engine/Blueprints/UserGuide/MathNode
 Related: Engine/Rendering/Materials/Functions
 Related: Engine/Rendering/Materials/ExpressionReference
+Version: 4.9
 
 [REGION:banner]
 ![](MathHall_Header.png)(convert:false)
@@ -44,8 +45,8 @@ Related: Engine/Rendering/Materials/ExpressionReference
 [TOC(start:2 end:3)]
 
 
-## Overview
-The **Math_Hall** map inside the **Content Examples** project aims to give you more of an understanding of how math is used within Unreal Engine 4 inside of **Blueprints** and **Materials**. 
+
+The **Math_Hall** map inside the Content Examples project aims to give you more of an understanding of how math is used within Unreal Engine 4 inside of **Blueprints** and **Materials**. 
 
 The first section entitled **Functions**, covers Math and Material Functions (left side of the wall) and example usage cases (right side of the wall). This section is excellent for learning more about blending Materials as well as how logic can be applied to Materials to create events like rippling water or bouncing balls entirely inside the **Material Editor**.  
 
@@ -60,21 +61,21 @@ Listed below are the examples provided in the **Functions** section, example use
 | Example | What is Demonstrated |
 | ------- | ------------ |
 | **1.1 Introduction** | Describes how to read each chart in this section where each chart plots the value of X for various math functions. |
-| **1.2 ABS** | Obtaining the **Absolute (ABS)** value of X where negative values are read as positive values (e.g. By taking the Sine of Time and getting the ABS value, a bouncing motion is created for a ball).  |
-| **1.3 Frac** | Using the **Frac** function to remove the integer part of any number which is useful for pseudo random number generation (e.g. the value of Raw Time is shown counting upwards extending past 1, Frac is used to express the 0 to 1 time that increments the Raw Time value). |
-| **1.4 Floor** | **Floor** removes the fraction from any number keeping only the integer (e.g. where 0.9 becomes 0 or 1.9 becomes 1, etc.). |
-| **1.5 Ceil** | **Ceil** works just like floor but rounds up instead of down. The result of ceil is always 1 greater than floor (e.g. 0.1 becomes 1, 1.1 becomes 2, etc. An example of creating a Mosaic is also provided). |
-| **1.6 FMod** | **FMod** stands for Floating Modulus and has both A and B inputs. It divides A by B and returns the remainder (e.g. an example is provided for generating a stripe pattern on a material). |
-| **1.7 Power** | **Power** takes a base input and applies an exponent. Negative values become positive just like **ABS** and a power of 1 will give the same result as ABS (e.g. an example is given to adjust the contrast of a texture).  |
-| **1.8 Lerp** | Linear Interpolation (**Lerp**) returns a blend between the A and B inputs based on the alpha input. In the given example, when the alpha is 0, 100% of A will be returned. When the alpha is 1, 100% of B is returned (e.g. an example of using Lerp to blend between two textures is given). |
-| **1.9 HeightLerp** | **HeightLerp** is a **Lerp** that modulates the blend using a heightmap (e.g. an example is given where the heightmap of a texture is used to blend between two textures where a grass texture blends into a cobblestone texture). |
-| **1.10 Clamp** | **Clamp** is used to clamp the input value to the specified Min and Max Range (e.g. using clamp to specify a min/max view angle on a security camera). |
-| **1.11 Sine** | **Sine** is a function of an angle and represents a ratio between the sides of a right triangle (e.g. Sine has a variety of applications and can create swaying motions inside Blueprints or Materials; it can also be used to make "waves" in water). |
-| **1.12 Cosine** | **Cosine** is another trigonometric function like **Sine**. It represents the ratio between the lengths of the sides in a right triangle. Cosine is the ratio between the side adjacent an angle and the hypotenuse (e.g. Sine and Cosine are very similar when used to make waves as seen in the example provided).  |
-| **1.13 OneMinus** | **OneMinus** can be used to get the inverse value of something like a texture or image (e.g. an example is shown where a black-and-white logo is inverted to a white-and-black logo).  |
-| **1.14 Exponential Density** | **Exponential Density** creates a gradient with a smooth curve at both ends (e.g. it is most commonly used for creating the density in fog but can have other uses). |
-| **1.15 Sphere Mask** | **Sphere Mask** has A and B inputs, as well as Radius and Hardness inputs and gets the distance between A and B (e.g. an example is shown where the hardness of a blend between textures is performed). |
-| **1.16 Sphere Gradient - 2D** | **Sphere Gradient - 2D** is a Material function that returns the sphere shaped falloff and takes V2 input coordinates, CenterPosition and scalar radius values (e.g. an example is shown how a material with a bevel shape is constructed).  |
+| **1.2 ABS** | Obtaining the Absolute (ABS value of X where negative values are read as positive values (e.g. By taking the Sine of Time and getting the ABS value, a bouncing motion is created for a ball).  |
+| **1.3 Frac** | Using the Frac function to remove the integer part of any number which is useful for pseudo random number generation (e.g. the value of Raw Time is shown counting upwards extending past 1, Frac is used to express the 0 to 1 time that increments the Raw Time value). |
+| **1.4 Floor** | Floor removes the fraction from any number keeping only the integer (e.g. where 0.9 becomes 0 or 1.9 becomes 1, etc.). |
+| **1.5 Ceil** | Ceil works just like floor but rounds up instead of down. The result of ceil is always 1 greater than floor (e.g. 0.1 becomes 1, 1.1 becomes 2, etc. An example of creating a Mosaic is also provided). |
+| **1.6 FMod** | FMod stands for Floating Modulus and has both A and B inputs. It divides A by B and returns the remainder (e.g. an example is provided for generating a stripe pattern on a material). |
+| **1.7 Power** | Power takes a base input and applies an exponent. Negative values become positive just like ABS and a power of 1 will give the same result as ABS (e.g. an example is given to adjust the contrast of a texture).  |
+| **1.8 Lerp** | Linear Interpolation (Lerp) returns a blend between the A and B inputs based on the alpha input. In the given example, when the alpha is 0, 100% of A will be returned. When the alpha is 1, 100% of B is returned (e.g. an example of using Lerp to blend between two textures is given). |
+| **1.9 HeightLerp** | HeightLerp is a Lerp that modulates the blend using a heightmap (e.g. an example is given where the heightmap of a texture is used to blend between two textures where a grass texture blends into a cobblestone texture). |
+| **1.10 Clamp** | Clamp is used to clamp the input value to the specified Min and Max Range (e.g. using clamp to specify a min/max view angle on a security camera). |
+| **1.11 Sine** | Sine is a function of an angle and represents a ratio between the sides of a right triangle (e.g. Sine has a variety of applications and can create swaying motions inside Blueprints or Materials; it can also be used to make "waves" in water). |
+| **1.12 Cosine** | Cosine is another trigonometric function like Sine. It represents the ratio between the lengths of the sides in a right triangle. Cosine is the ratio between the side adjacent an angle and the hypotenuse (e.g. Sine and Cosine are very similar when used to make waves as seen in the example provided).  |
+| **1.13 OneMinus** | OneMinus can be used to get the inverse value of something like a texture or image (e.g. an example is shown where a black-and-white logo is inverted to a white-and-black logo).  |
+| **1.14 Exponential Density** | Exponential Density creates a gradient with a smooth curve at both ends (e.g. it is most commonly used for creating the density in fog but can have other uses). |
+| **1.15 Sphere Mask** | Sphere Mask has A and B inputs, as well as Radius and Hardness inputs and gets the distance between A and B (e.g. an example is shown where the hardness of a blend between textures is performed). |
+| **1.16 Sphere Gradient - 2D** | Sphere Gradient - 2D is a Material function that returns the sphere shaped falloff and takes V2 input coordinates, CenterPosition and scalar radius values (e.g. an example is shown how a material with a bevel shape is constructed).  |
 
 
 
@@ -86,26 +87,26 @@ Listed below are the examples provided in the **Basic Vector Math** section.
 
 | Example | What is Demonstrated |
 | ------- | ------------ |
-| **2.1 Introduction to Vector Math** | Describes the importance of **Vector Math** in Materials and Blueprints. |
-| **2.2 What is a vector** | Explains what a **Vector** is and the types of vectors.  |
-| **2.3 1D Vectors** | Basic **1D Vector** using a roadmap analogy. |
-| **2.4 2D Vectors** | Explanation of a **2D Vector** using a 2D graph. |
-| **2.5 Distance** | How **Distance** is calculated using the Pythagorean Theorem (Note: Distance is automatically calculated in Materials and Blueprints). |
-| **2.6 Distance in practice** | **Distance** being calculated in a Blueprint and Material. |
-| **2.7 3D Vectors** | Visual demonstration of a **3D Vector** using a 3D grid. |
-| **2.8 Normalized Vectors** | Using **Normalized Vectors** with **Normalize** nodes that can be found in Blueprints or Materials. |
+| **2.1 Introduction to Vector Math** | Describes the importance of Vector Math in Materials and Blueprints. |
+| **2.2 What is a vector** | Explains what a Vector is and the types of vectors.  |
+| **2.3 1D Vectors** | Basic 1D Vector using a roadmap analogy. |
+| **2.4 2D Vectors** | Explanation of a 2D Vector using a 2D graph. |
+| **2.5 Distance** | How Distance is calculated using the Pythagorean Theorem (Note: Distance is automatically calculated in Materials and Blueprints). |
+| **2.6 Distance in practice** | Distance being calculated in a Blueprint and Material. |
+| **2.7 3D Vectors** | Visual demonstration of a 3D Vector using a 3D grid. |
+| **2.8 Normalized Vectors** | Using Normalized Vectors with Normalize nodes that can be found in Blueprints or Materials. |
 | **2.9 Vector as rotation** | Converting a Vector into a Rotation value (e.g. a character that rotates towards the player's position is shown). |
-| **2.10 WorldPosition** | **World Position** and how to use the data (e.g. an object's world position is updated as it is dragged around in the editor). |
-| **2.11 World Aligned Texture** | Using **World Aligned Texture** to blend a texture around a sphere (e.g. this is useful for creating more natural looking rock materials).  |
-| **2.12 CameraVector** | A **CameraVector** is shown with a diagram illustrating vector paths. |
-| **2.13 Camera Forward Vector** | **CameraForwardVector** is shown with a diagram illustrating forward vector paths. |
-| **2.14 Dot Product** | An explanation of **Dot Product** is given which can be used to check if the player is looking at a specific point. |
+| **2.10 WorldPosition** | World Position and how to use the data (e.g. an object's world position is updated as it is dragged around in the editor). |
+| **2.11 World Aligned Texture** | Using World Aligned Texture to blend a texture around a sphere (e.g. this is useful for creating more natural looking rock materials).  |
+| **2.12 CameraVector** | A CameraVector is shown with a diagram illustrating vector paths. |
+| **2.13 Camera Forward Vector** | CameraForwardVector is shown with a diagram illustrating forward vector paths. |
+| **2.14 Dot Product** | An explanation of Dot Product is given which can be used to check if the player is looking at a specific point. |
 | **2.15 Deriving Dot Product Algebraically** | How to derive the dot product (Note: In Unreal Engine 4, a dot product is almost always solved using a node). |
-| **2.16 Dot Product Example** | A **CameraVector** and **Dot Product** where a Material blends between colors based on player distance to an object as well as tracks player movement (Note: This is done entirely in the Material without a Blueprint). |
-| **2.17 Dot Product Example 2** | A **CameraForwardVector** and **Dot Product** is used to determine if the player is looking at an object (Note: This is done entirely in the Material without a Blueprint). |
-| **2.18 Cross Product** | Usage of **Cross Product** with a diagram illustrating how a cross product can be derived by reading vector values. |
-| **2.19 CreateThirdOrthogonalVector** | **CreateThirdOrthogonalVector** Material function and how it is useful for creating a full set of orthogonal (perpendicular) vectors is shown. |
-| **2.20 Camera Facing math** | **CameraPosition - ObjectPosition** is shown to create the effect of objects rotating to face the player (e.g. in this example when standing in front of several objects, the objects directly in front of the player do not rotate, while objects to the left or right of the player rotate inward). |
-| **2.21 Camera Facing math 2** | **Camera Left Vector** and **Camera Up Vector** are used to create the effect of objects rotating to face the player (e.g. in this example, unlike 2.20, the objects on the outside do not rotate inward as they all rotate in unison). |
+| **2.16 Dot Product Example** | A CameraVector and Dot Product where a Material blends between colors based on player distance to an object as well as tracks player movement (Note: This is done entirely in the Material without a Blueprint). |
+| **2.17 Dot Product Example 2** | A CameraForwardVector and Dot Product is used to determine if the player is looking at an object (Note: This is done entirely in the Material without a Blueprint). |
+| **2.18 Cross Product** | Usage of Cross Product with a diagram illustrating how a cross product can be derived by reading vector values. |
+| **2.19 CreateThirdOrthogonalVector** | CreateThirdOrthogonalVector Material function and how it is useful for creating a full set of orthogonal (perpendicular) vectors is shown. |
+| **2.20 Camera Facing math** | CameraPosition - ObjectPosition is shown to create the effect of objects rotating to face the player (e.g. in this example when standing in front of several objects, the objects directly in front of the player do not rotate, while objects to the left or right of the player rotate inward). |
+| **2.21 Camera Facing math 2** | Camera Left Vector and Camera Up Vector are used to create the effect of objects rotating to face the player (e.g. in this example, unlike 2.20, the objects on the outside do not rotate inward as they all rotate in unison). |
 | **2.22 Trig tips** | Sine, Cosine and Tangent as it relates to solving a right triangle where knowing any two sides can tell you the Angle of X (e.g. Our House is A House Of Animals mnemonic is explained as each capital letter stands for one of the triangle sides).  |
 
