#UE4doc #docs - added engine version field. Added/Removed bolded terms. Corrected mistake with number of rooms associated with this map.

diff --git a/Engine/Documentation/Source/Resources/ContentExamples/Matinee/Matinee.INT.udn b/Engine/Documentation/Source/Resources/ContentExamples/Matinee/Matinee.INT.udn
--- a/Engine/Documentation/Source/Resources/ContentExamples/Matinee/Matinee.INT.udn
+++ b/Engine/Documentation/Source/Resources/ContentExamples/Matinee/Matinee.INT.udn
@@ -6,6 +6,7 @@ Related: Engine/Matinee
 Related: Engine/Matinee/UserGuide
 Related: Engine/Matinee/TrackReference
 Related: Engine/UI/CurveEditor
+Version: 4.9
 
 
 
@@ -48,13 +49,13 @@ Related: Engine/UI/CurveEditor
 
 
 
-The **Matinee** Content Example map provides several examples on how to generate cutscenes or create gameplay sequences by placing **Keyframes** along a timeline. Each **Keyframe** tells the targeted **Actor** to "do something" at that point in the timeline - be it move to and arrive at a position, change color, fire off an event in a **Blueprint**, etc. 
+The **Matinee** map in Content Example provides several examples on how to generate cutscenes or create gameplay sequences by placing Keyframes along a timeline. Each Keyframe tells the targeted Actor to "do something" at that point in the timeline - be it move to and arrive at a position, change color, fire off an event in a Blueprint, etc. 
 
-With the use of a **Director Group** (see 2.5 below), it is possible to set up sequences which take over the in-game camera and allow you to switch to different cameras or add other special effects to your scene (useful for creating cutscenes). 
+With the use of a Director Group (see 2.5 below), it is possible to set up sequences which take over the in-game camera and allow you to switch to different cameras or add other special effects to your scene (useful for creating cutscenes). 
 
 
 ## Matinee Map
-The **Matinee** Content Example map is broken up into 5 sections of which are outlined below. 
+The **Matinee** map is broken up into 4 sections of which are outlined below. 
 
 ### Matinee Basics
 
@@ -64,10 +65,10 @@ Listed below are the examples provided inside the **Matinee Basics** section:
 
 | Example | What is Demonstrated |
 | ------- | ------------ |
-| **1.1 Matinee Actor** | A placed Matinee Actor in a level which also describes how to open the **Matinee Editor**. |
-| **1.2 Tracks and Keys** | Shows how a **Movement Track** and **Keys** can be used to move an Actor around in the scene. |
-| **1.3 Track Types** | An example Matinee that uses multiple tracks that affect an Actor (e.g. the Actor is scaled, moved, its **Material** is changed, is set invisible, and a sound is played through a Matinee scene). |
-| **1.4 Key Types** | How different **Key Types** can be used to affect the properties of a track that changes over time (e.g. multiple Actors move from point A to point B, however their transition is modified by Key Types). |
+| **1.1 Matinee Actor** | A placed Matinee Actor in a level which also describes how to open the Matinee Editor. |
+| **1.2 Tracks and Keys** | Shows how a Movement Track and Keys can be used to move an Actor around in the scene. |
+| **1.3 Track Types** | An example Matinee that uses multiple tracks that affect an Actor (e.g. the Actor is scaled, moved, its Material is changed, is set invisible, and a sound is played through a Matinee scene). |
+| **1.4 Key Types** | How different Key Types can be used to affect the properties of a track that changes over time (e.g. multiple Actors move from point A to point B, however their transition is modified by Key Types). |
 
 
 ### Matinee Groups
@@ -78,11 +79,11 @@ Listed below are the examples provided inside the **Matinee Groups** section:
 
 | Example | What is Demonstrated |
 | ------- | ------------ |
-| **2.1 Lighting Group** | **Lighting Group** is used to affect the movement, intensity, brightness, light color, and light radius of a light over time. |
-| **2.2 Particle Group** | The **Particle Group** is used to toggle on/off particles and can be used to determine the length a particle is displayed.  |
-| **2.3 Skeletal Group** | **Skeletal Groups** can be used to assign movement or animations to a **Skeletal Mesh** which can be assigned and played at specified keyframes. |
-| **2.4 Camera Group** | The **Camera Group** is used to manipulate a camera that is placed in a level by altering its movement and (or) FOV. |
-| **2.5 Director Group** | The **Director Group** allows you to take control over the in-game camera and switch between multiple cameras, add fade or slomo effects, adjust the master audio track, or apply new color scales to a scene. |
+| **2.1 Lighting Group** | Lighting Group is used to affect the movement, intensity, brightness, light color, and light radius of a light over time. |
+| **2.2 Particle Group** | The Particle Group is used to toggle on/off particles and can be used to determine the length a particle is displayed.  |
+| **2.3 Skeletal Group** | Skeletal Groups can be used to assign movement or animations to a **Skeletal Mesh** which can be assigned and played at specified keyframes. |
+| **2.4 Camera Group** | The Camera Group is used to manipulate a camera that is placed in a level by altering its movement and (or) FOV. |
+| **2.5 Director Group** | The Director Group allows you to take control over the in-game camera and switch between multiple cameras, add fade or slomo effects, adjust the master audio track, or apply new color scales to a scene. |
 
 
 ### Matinee Blueprint Interaction
@@ -93,7 +94,7 @@ Listed below are the examples provided inside the **Matinee Blueprint Interactio
 
 | Example | What is Demonstrated |
 | ------- | ------------ |
-| **3.1 Matinee Blueprint Functions** | Demonstrates the standard **Matinee Blueprint Functions** that can be called inside a Blueprint (e.g. play, pause, stop, rewind, reverse, loop, and set position). |
+| **3.1 Matinee Blueprint Functions** | Demonstrates the standard Matinee Blueprint Functions that can be called inside a Blueprint (e.g. play, pause, stop, rewind, reverse, loop, and set position). |
 | **3.2 Matinee Events / Event Tracks** | Shows how events can be fired inside a Blueprint using an Event Track inside Matinee (events that are fired can be used to execute and start any section of script). |
 | **3.3 Modifying Matinee Properties** | An example of some of the properties that can be set on the Matinee Actor itself (e.g. adjust the play rate, play on level load, check looping or skippable, manipulate the player by disabling movement, look input, or hiding the player character). |
 
@@ -106,8 +107,8 @@ Listed below are the examples provided inside the **Matinee Advanced** section:
 
 | Example | What is Demonstrated |
 | ------- | ------------ |
-| **4.1 Curve Editor** | How to get finer control over parameters by using the **Curve Editor** inside Matinee (e.g. fine tune the movement of an Actor by adjusting the curve manually). |
-| **4.2 Split Translation and Rotation** | Additional control over movement by splitting the automatic movement coordinates into **Translation** and **Rotation** that can be adjusted manually at each key. |
+| **4.1 Curve Editor** | How to get finer control over parameters by using the Curve Editor inside Matinee (e.g. fine tune the movement of an Actor by adjusting the curve manually). |
+| **4.2 Split Translation and Rotation** | Additional control over movement by splitting the automatic movement coordinates into Translation and Rotation that can be adjusted manually at each key. |
 
 
 
