#UE4doc #docs - Added engine version field. Adjusted image sizes to fix within the 960 px wide for the image table.

diff --git a/Engine/Documentation/Source/Resources/ContentExamples/MouseInterface/BlueprintActors/BlueprintActors.INT.udn b/Engine/Documentation/Source/Resources/ContentExamples/MouseInterface/BlueprintActors/BlueprintActors.INT.udn
--- a/Engine/Documentation/Source/Resources/ContentExamples/MouseInterface/BlueprintActors/BlueprintActors.INT.udn
+++ b/Engine/Documentation/Source/Resources/ContentExamples/MouseInterface/BlueprintActors/BlueprintActors.INT.udn
@@ -2,6 +2,7 @@ Availability: Public
 Title:Blueprint Actor Setup
 Crumbs:%ROOT%, Resources,  Resources/ContentExamples/MouseInterface
 Description:An overview of the Blueprint Actor setup used in the Mouse Interface Demolet
+Version: 4.9
 
 
 [TOC(start:2 end:2)]
@@ -12,14 +13,14 @@ Description:An overview of the Blueprint Actor setup used in the Mouse Interface
 The Blueprint Actors (MI_Sphere) arrayed within the scene can be moved during gameplay by clicking and dragging with the mouse. 
 
 [REGION:imagetable]
-|![](MiSpheres.png)(h:307)|![](OnMouseOver.png)(h:307)|![](MovingSpheres.png)(h:307)|
+|![](MiSpheres.png)(h:235)|![](OnMouseOver.png)(h:235)|![](MovingSpheres.png)(h:235)|
 |:---:|:---:|:---: |
 |In Editor|During Gameplay: Mouse Over|During Gameplay: Dragging|
 [/REGION]
 
 ## Components
 
-The Components List for the MI_Sphere Blueprint is exceedingly simple; really, it is just a single Static Mesh sphere.
+The Components List for the **MI_Sphere** Blueprint is exceedingly simple; really, it is just a single Static Mesh sphere.
 
 ![](MI_Sphere_Components.png)(w:700)
 
diff --git a/Engine/Documentation/Source/Resources/ContentExamples/MouseInterface/MouseInterface.INT.udn b/Engine/Documentation/Source/Resources/ContentExamples/MouseInterface/MouseInterface.INT.udn
--- a/Engine/Documentation/Source/Resources/ContentExamples/MouseInterface/MouseInterface.INT.udn
+++ b/Engine/Documentation/Source/Resources/ContentExamples/MouseInterface/MouseInterface.INT.udn
@@ -7,6 +7,7 @@ Related: Engine/Blueprints\Editor
 Related: Engine/Blueprints\Editor\Modes\Components
 Related: Engine/Blueprints\Editor\Modes\Defaults
 Related: Engine/Blueprints\Editor\Modes\GraphPanel
+Version: 4.9
 
 
 [VAR:Topic]
@@ -55,11 +56,9 @@ Related: Engine/Blueprints\Editor\Modes\GraphPanel
 ![](MouseInterfaceTitle.png)
 [/REGION]
 
-[TOC(start:2 end:2)]
 
 
-
-The Mouse Interface example demonstrates how to utilize mouse or touch inputs in your games. The level, which is pretty simple in terms of play and execution, consists of an array of spheres with a cube that tracks back and forth behind them.
+The **Mouse Interface** map in Content Examples demonstrates how to utilize mouse or touch inputs in your games. The level, which is pretty simple in terms of play and execution, consists of an array of spheres with a cube that tracks back and forth behind them.
 
 
 ![](MouseInterfaceGame.png)(w:400)
