#UE4doc #docs -  Added engine version field. Updated image call out box color.

diff --git a/Engine/Documentation/Source/Resources/ContentExamples/MouseInterface/LevelBlueprint/LevelBlueprint.INT.udn b/Engine/Documentation/Source/Resources/ContentExamples/MouseInterface/LevelBlueprint/LevelBlueprint.INT.udn
--- a/Engine/Documentation/Source/Resources/ContentExamples/MouseInterface/LevelBlueprint/LevelBlueprint.INT.udn
+++ b/Engine/Documentation/Source/Resources/ContentExamples/MouseInterface/LevelBlueprint/LevelBlueprint.INT.udn
@@ -2,6 +2,7 @@ Availability: Public
 Title:Level Blueprint
 Crumbs:%ROOT%, Resources,  Resources/ContentExamples/MouseInterface
 Description:An overview of the Level Blueprint setup used in the Mouse Interface Demolet
+Version: 4.9
 
 
 [TOC(start:2 end:2)]
diff --git a/Engine/Documentation/Source/Resources/ContentExamples/MouseInterface/MouseControlSetup/MouseControlSetup.INT.udn b/Engine/Documentation/Source/Resources/ContentExamples/MouseInterface/MouseControlSetup/MouseControlSetup.INT.udn
--- a/Engine/Documentation/Source/Resources/ContentExamples/MouseInterface/MouseControlSetup/MouseControlSetup.INT.udn
+++ b/Engine/Documentation/Source/Resources/ContentExamples/MouseInterface/MouseControlSetup/MouseControlSetup.INT.udn
@@ -2,6 +2,7 @@ Availability: Public
 Title:Mouse Control Setup
 Crumbs:%ROOT%, Resources,  Resources/ContentExamples/MouseInterface
 Description:How to set up a game for mouse control
+Version: 4.9
 
 [TOC(start:2 end:2)]
 
@@ -9,7 +10,7 @@ Description:How to set up a game for mouse control
 To make a game where you use the mouse cursor possible, there are a few things you must do when getting started. At a high level, the steps in Blueprint are as follows. We will cover them in greater detail below:
 
 * Create a new **PlayerController** Blueprint Class that enables the mouse cursor and any other desired mouse behavior. 
-* Create a new **GameMode** Blueprint Class that utilizes your new **PlayerController** Class.
+* Create a new **GameMode** Blueprint Class that utilizes your new PlayerController Class.
 * Set the **World Settings** to use your new GameMode in the Default GameType.
 * Use Blueprint scripting (either within the Level Blueprint or any Actor Blueprints) to process mouse interaction.
 
@@ -36,7 +37,7 @@ By default, the PlayerController class does not have mouse interfacing enabled.
 
 	![](MousePlayerController_CB.png)
 
-1. **Right-click** on your new Blueprint in the **Content Browser** and choose **Edit...**. This will open up the **Class Defaults** window.
+1. **Right-click** on your new Blueprint in the Content Browser and choose **Edit...**. This will open up the **Class Defaults** window.
 
 	![](EditDefaultsBlueprint.png)
 
@@ -63,7 +64,7 @@ Now that we have a PlayerController that has Mouse Input enabled, our next step
 
 	![](Mouse_GameInfoBlueprint.png)
 
-1. **Right-click** on your new Blueprint in the **Content Browser** and choose **Edit...**. This will open up the **Class Defaults** window.
+1. **Right-click** on your new Blueprint in the Content Browser and choose **Edit...**. This will open up the **Class Defaults** window.
 
 	![](EditDefaultsBlueprint2.png)
 
@@ -78,7 +79,7 @@ Now that we have a PlayerController that has Mouse Input enabled, our next step
 
 Now that we have our custom PlayerController and GameMode Blueprints completed, we just need to tell our World Settings to use them.
 
-1. **Click** the **Settings** button located in the toolbar. Then from the drop down **select** **World Settings**. This opens up the World Settings in the **Details** panel.
+1. Click the **Settings** button located in the toolbar. Then from the drop down select **World Settings**. This opens up the World Settings in the **Details** panel.
 
 	![](WorldSettings.png)
 
