#UE4doc #docs - Added engine version field

diff --git a/Engine/Documentation/Source/Resources/ContentExamples/NavMesh/NavMesh.INT.udn b/Engine/Documentation/Source/Resources/ContentExamples/NavMesh/NavMesh.INT.udn
--- a/Engine/Documentation/Source/Resources/ContentExamples/NavMesh/NavMesh.INT.udn
+++ b/Engine/Documentation/Source/Resources/ContentExamples/NavMesh/NavMesh.INT.udn
@@ -6,6 +6,7 @@ Related: Engine/AI/BehaviorTrees/QuickStart
 Related: Gameplay/AI
 Related: Engine/AI/BehaviorTrees/NodeReference
 Related: Engine/Actors/Volumes
+Version: 4.9
 
 [REGION:banner]
 ![](NavMesh_Header.png)(convert:false)
@@ -46,7 +47,7 @@ Related: Engine/Actors/Volumes
 
 
 
-The **Navmesh** Content Example map provides three examples of how to generate and use a **NavMesh** to enable **Pawns** the ability to pathfind their way through obstacles, over ramps or to jump off ledges. While in the editor, and with a **NavMeshBoundsVolume** placed in the level, pressing the **P** key will show/hide the area that the NavMesh covers. The final example in this map shows how the coverage area of a NavMesh is automatically updated when moving **Actors** around in the level (this can be disabled via the **Editor Settings** under **Level Editor/Miscellaneous** and unchecking the **Update Navigation Automatically** checkbox). 
+The **Navmesh** map in Content Example provides three examples of how to generate and use a **NavMesh** to enable **Pawns** the ability to pathfind their way through obstacles, over ramps or to jump off ledges. While in the editor, and with a **NavMeshBoundsVolume** placed in the level, pressing the **P** key will show/hide the area that the NavMesh covers. The final example in this map shows how the coverage area of a NavMesh is automatically updated when moving **Actors** around in the level (this can be disabled via the **Editor Settings** under **Level Editor/Miscellaneous** and unchecking the **Update Navigation Automatically** checkbox). 
 
 [region:note]
 You can find additional resources for each example by clicking on the example's name.
