#UE4doc #docs:
Added engine version field
Changed font settings for multiple terms
Removed lightbox region from multiple images

diff --git a/Engine/Documentation/Source/GettingStarted/HowTo/FindingHelp/FindingHelp.INT.udn b/Engine/Documentation/Source/GettingStarted/HowTo/FindingHelp/FindingHelp.INT.udn
--- a/Engine/Documentation/Source/GettingStarted/HowTo/FindingHelp/FindingHelp.INT.udn
+++ b/Engine/Documentation/Source/GettingStarted/HowTo/FindingHelp/FindingHelp.INT.udn
@@ -4,6 +4,7 @@ Crumbs: %ROOT%, Gameplay, GettingStarted\HowTo
 Description:Where to find in-editor help and tutorials or external help and answers to your Unreal Engine 4 questions. 
 Related: Resources
 Related: Resources/ContentExamples
+version: 4.9
 
 When working with **Unreal Engine 4**, if you come across something that you do not understand or would like more clarification on, do not become discouraged as there are several avenues you can use to learn more about the Editor or an aspect of the Editor you are working in. This page will highlight where you can go to learn more about the Editor or where to go to get answers to any questions you may have. 
 
@@ -37,11 +38,7 @@ Some aspects of the Editor may have a **?** icon next to them.
 
 This **Help Icon** when clicked will take you to relevant external documentation related to object it is linked to. 
 
-[REGION:lightbox]
-[![](HelpIcon_2.png)](HelpIcon_2.png)
-[/REGION]
-
-_Click for full image._
+![](HelpIcon_2.png)
 
 Above we have clicked on a **Brush** Actor in our level, then in the **Details** panel, clicked the **?** Help Icon which automatically opened a documentation page on **Geometry Brush Actors**. This is useful when you are not sure what an object does or want to learn more about the object, but not sure where to look on the official documentation pages. 
 
@@ -93,25 +90,15 @@ Inside the **Help Menu** noted above, next to **Documentation...** you may have
 
 Pressing **F1** is a keyboard shortcut to bringing up an Official Documentation page. The **F1 Help** button is context sensitive, meaning that depending on where you are in the Editor, a different help page may be presented to you. For example, below we are in the Main Editor View and pressing **F1** for help displays the home documentation page. 
 
-[REGION:lightbox]
-[![](F1_1.png)](F1_1.png)
-[/REGION]
-
-_Click for full image._
+![](F1_1.png)
 
 With a **Blueprint** window open, if we were to press **F1** here a different documentation page is displayed for help regarding Blueprints. 
 
-[REGION:lightbox]
-[![](F1_2.png)](F1_2.png)
-[/REGION]
-
-_Click for full image._
+![](F1_2.png)
 
 Or if we were in the **Material Editor**, a page on **Materials** is displayed. 
 
-[REGION:lightbox]
-[![](F1_3.png)](F1_3.png)
-[/REGION]
+![](F1_3.png)
 
 Whenever you are in a new Editor window and not sure what it is or how to work with it, try pressing **F1** for help. 
 
@@ -131,17 +118,13 @@ Not quite sure what it is you are looking for but have an idea of what you want
 
 ![](HelpSearch_3.png)
 
-Above we are searching for **How to display Health**. Depending upon your entry, there may be Official Documentation related to that very topic. There may also be [Wiki](#externalhelp:wiki) help or tutorials or even questions on Answer Hub that have been posed and answered. 
+Above we are searching for "How to display health". Depending upon your entry, there may be Official Documentation related to that very topic. There may also be [Wiki](#externalhelp:wiki) help or tutorials or even questions on Answer Hub that have been posed and answered. 
 
 ## In-Editor Tutorials
 
 **In-Editor** Tutorials provide you with step-by-step instructions on how to navigate and use particular aspects of Unreal Engine 4. 
 
-[REGION:lightbox]
-[![](Tutorial_1.png)](Tutorial_1.png)
-[/REGION]
-
-_Click for full image._
+![](Tutorial_1.png)
 
 While proceeding through each In-Editor Tutorial, you will be guided across the interface as well as given instructions on how to use the feature the tutorial covers. You will also be given links to external Official Documentation on the feature which will help you expand upon the topics that are covered in the tutorial itself. 
 
@@ -175,7 +158,7 @@ Click the blinking icon to start an In-Editor Tutorial covering the area you are
 
 ![](Tutorial_6.png)
 
-Above we opened a **Static Mesh** asset, then clicked the In-Editor Tutorial icon and a **Static Mesh Editor** tutorial was started. 
+Above we opened a Static Mesh asset, then clicked the In-Editor Tutorial icon and a **Static Mesh Editor** tutorial was started. 
 
 
 ## Content Examples
@@ -184,7 +167,7 @@ Another learning resource you can use to help answer any questions you may have
 
 [INCLUDE:Resources/ContentExamples#contentexamples]
 
-You can download the **Content Examples** project from the **Learn** tab of the Unreal Engine Launcher. 
+You can download the Content Examples project from the **Learn** tab of the Unreal Engine Launcher. 
 
 [REGION:note]
 For more information, see the [Content Examples](Resources/ContentExamples) documentation.
