#UE4doc #docs:
Removed italics from the Overview excerpt region

diff --git a/Engine/Documentation/Source/Engine/UI/LevelEditor/InEditorTesting/InEditorTesting.INT.udn b/Engine/Documentation/Source/Engine/UI/LevelEditor/InEditorTesting/InEditorTesting.INT.udn
--- a/Engine/Documentation/Source/Engine/UI/LevelEditor/InEditorTesting/InEditorTesting.INT.udn
+++ b/Engine/Documentation/Source/Engine/UI/LevelEditor/InEditorTesting/InEditorTesting.INT.udn
@@ -64,7 +64,7 @@ Version: 4.9
 
 [EXCERPT:overview]
 Unreal Editor allows you to spawn a player anywhere and instantly preview your game from within the editor without waiting for files to save. There are two preview types available: **Play In Editor** (PIE) which is directly accessed via the **Play** button located on the Main Toolbar
-and **Simulate In Editor** (SIE) which is accessed via the **Play** button's drop-down menu (or by pressing **Alt+S**). The in-editor preview system supports toggling between _Play In Editor_ and _Simulate In Editor_ sessions, so that you can quickly iterate on gameplay and asset tweaks
+and **Simulate In Editor** (SIE) which is accessed via the **Play** button's drop-down menu (or by pressing **Alt+S**). The in-editor preview system supports toggling between Play In Editor and Simulate In Editor sessions, so that you can quickly iterate on gameplay and asset tweaks
 and see how your game changes as a result.
 [/EXCERPT:overview]
 
