#DocFix

Updating images and paragraphs on this page via forum request

diff --git a/Engine/Documentation/Source/Engine/MediaFramework/MediaFramework.INT.udn b/Engine/Documentation/Source/Engine/MediaFramework/MediaFramework.INT.udn
--- a/Engine/Documentation/Source/Engine/MediaFramework/MediaFramework.INT.udn
+++ b/Engine/Documentation/Source/Engine/MediaFramework/MediaFramework.INT.udn
@@ -120,7 +120,7 @@ Instead of clicking the **Create** button, you can also **Right-click** in the *
 
 1. Inside the MediaPlayer asset, in the Details panel, click the **Browse** button. 
 
-	![](AddMediaButton.png)
+	![](BrowseMediaPlayer.jpg)
 
 1. Browse to your media and select your file. 
 
@@ -132,11 +132,11 @@ Instead of clicking the **Create** button, you can also **Right-click** in the *
 
 1. Inside the Details panel, you may notice that a warning icon is present. 
 
-	![](WarningIcon.png)
+	![](EpicMovieMediaFrameWorkFile.jpg)
 
 1. In order to package correctly, local media files MUST be placed in the **Content/Movies/** directory of your project. 
 
-	![](NoWarning.png)
+	![](EpicMovieMediaFrameWorkWMV.jpg)
 
 	Once you have placed your media in your projects Content/Movies/ folder, the warning icon will no longer be shown. 
 	
@@ -159,13 +159,11 @@ the media window below the playback controls.
 
 Inside the Details panel, you can inspect the properties of the media source. 
 
-![](MediaPlayerDetails.png)
+![](EpicMovieMediaFrameWork.jpg)
 
-Under the **Playback** options, you can set the media to **Auto Play** (which will automatically play the media when loaded) and **Auto Play Rate** (to define the playback 
-rate when Auto Play is enabled), as well as the ability to enable looping of the media when playback reaches the end. 
+Under the **Playback** option, you have the opyion to enable looping of the media when playback reaches the end by checking the **Looping** checkbox. 
 
-For the **Source** section, you can specify the path to the media locally or through the Internet via URL as well as set the Streaming Mode. The last section, displays 
-information about the media itself. 
+For the **Source** section, you can specify the path to the media locally or through the Internet via URL, as well as set the Streaming Mode.  
 
 ### MediaTexture Asset
 
