#UE4doc #docs:
Added engine version field
Replaced multiple images as they were out of date

diff --git a/Engine/Documentation/Source/GettingStarted/HowTo/PIE/PIE.INT.udn b/Engine/Documentation/Source/GettingStarted/HowTo/PIE/PIE.INT.udn
--- a/Engine/Documentation/Source/GettingStarted/HowTo/PIE/PIE.INT.udn
+++ b/Engine/Documentation/Source/GettingStarted/HowTo/PIE/PIE.INT.udn
@@ -6,6 +6,7 @@ Related: Engine/UI/LevelEditor/InEditorTesting
 Related: Engine/UI/LevelEditor/InEditorTesting/Settings
 Related: Engine/UI/LevelEditor/InEditorTesting/KeepSimulationChanges
 Related: Gameplay/HowTo/Networking/TestMultiplayer
+version: 4.9
 
 
 [INCLUDE:Engine\UI\LevelEditor\InEditorTesting#overview]
