#UE4doc #docs:
Added engine version field
Changed font settings for a term

diff --git a/Engine/Documentation/Source/GettingStarted/HowTo/PlacingActors/PlacingActors.INT.udn b/Engine/Documentation/Source/GettingStarted/HowTo/PlacingActors/PlacingActors.INT.udn
--- a/Engine/Documentation/Source/GettingStarted/HowTo/PlacingActors/PlacingActors.INT.udn
+++ b/Engine/Documentation/Source/GettingStarted/HowTo/PlacingActors/PlacingActors.INT.udn
@@ -6,6 +6,7 @@ Related: Engine/Actors/Placement
 Related: Engine/UI/LevelEditor/Modes/PlaceMode
 Related: Engine/UI/ClassViewer
 Related: GettingStarted\HowTo\ManipulatingActors
+version: 4.9
 
 At the most fundamental level, an **Actor** is any Object that you can place in a level and this page will show you the various methods in which you can place those Actors inside your levels. 
 
@@ -72,7 +73,7 @@ This method of placing an Actor involves using the **Content Browser** and dragg
 
 A more advanced method of placing Actors is performed by using the [](Engine/UI/ClassViewer), which is a hierarchical list of classes used by the Editor.
 
-1.  To open the **Class Viewer**, click **Window** then **Developer Tools** and under _Miscellaneous_, select **Class Viewer**. 
+1.  To open the **Class Viewer**, click **Window** then **Developer Tools** and under **Miscellaneous**, select **Class Viewer**. 
 
 	![](Place6.png)
 
