#UE4doc #docs:
Updated engine version field
Changed font settings for multiple terms

diff --git a/Engine/Documentation/Source/GettingStarted/Installation/InstallingUnrealEngine.INT.udn b/Engine/Documentation/Source/GettingStarted/Installation/InstallingUnrealEngine.INT.udn
--- a/Engine/Documentation/Source/GettingStarted/Installation/InstallingUnrealEngine.INT.udn
+++ b/Engine/Documentation/Source/GettingStarted/Installation/InstallingUnrealEngine.INT.udn
@@ -21,7 +21,7 @@ By the end of this tutorial, you'll be able to:
 
 ![](InstallationGoal.png)
 
-_The Epic Games Launcher enables our subscribers to conveniently download and install updates for UE4._
+The Epic Games Launcher enables our subscribers to conveniently download and install updates for UE4.
 
 [REGION:warning]
 If you already have an Epic Games account, skip ahead to learn about [Downloading and Running the Installer](#Bookmark2).
@@ -126,15 +126,15 @@ Space Requirements
 Unreal Engine uses approximately eight gigabytes of disk space for every version of the engine that is installed on a machine. Make sure you have plenty of disk space prior to installing UE4.
 [/REGION]
 
-1.	**Click** the **Unreal Engine** tab inside the Epic Games Launcher Loading Screen.
+1.	Click the **Unreal Engine** tab inside the Epic Games Launcher Loading Screen.
 
 	![](LauncherLoadScreen.png)
 
-1.	**Click** the **Install Engine** button to download and install the latest version of Unreal Engine.
+1.	Click the **Install Engine** button to download and install the latest version of Unreal Engine.
 
 	![](InstallEngine.png)
 
-1.	**Read** our [Unreal Engine End User License Agreement](https://www.unrealengine.com/eula) (EULA) prior to **clicking** the **Accept** button.
+1.	Read our [Unreal Engine End User License Agreement](https://www.unrealengine.com/eula) (EULA) prior to clicking the **Accept** button.
 
 	![](AcceptEULALauncher.png)
 
@@ -148,7 +148,7 @@ Unreal Engine uses approximately eight gigabytes of disk space for every version
 
 (#Bookmark5)
 
-After the Epic Games Launcher successfully downloads and installs Unreal Engine, you're ready to **click** the **Launch** button.
+After the Epic Games Launcher successfully downloads and installs Unreal Engine, you're ready to click the **Launch** button.
 
 ![](LaunchUE4.png)
 
@@ -162,6 +162,7 @@ Congratulations! You've learned how to:
 
 Now, you're ready to start using Unreal Engine.
 
+<!--Deprecated as it is redundant with the Related Topics section directly below it.
 ## Getting Started with Unreal Engine
 
 Learn how to use Unreal Engine by referring to the [](GettingStarted) documentation!
@@ -172,9 +173,9 @@ If you're looking to quickly get started with UE4, refer to the following tutori
 *	[](Engine/QuickStart)
 
 
-[REGION:tip]
+[REGION:note]
 UE4's in-editor help features are a great way to get your questions answered.
-[/REGION]
+[/REGION]-->
 <!--Deprecated on 21 July 2015 due to outdated screenshots and outdated procedures
 %Globals:OSSelect%
 
