#UE4doc #docs:
Added engine version field

diff --git a/Engine/Documentation/Source/GettingStarted/QuickStartDirectory/QuickStartDirectory.INT.udn b/Engine/Documentation/Source/GettingStarted/QuickStartDirectory/QuickStartDirectory.INT.udn
--- a/Engine/Documentation/Source/GettingStarted/QuickStartDirectory/QuickStartDirectory.INT.udn
+++ b/Engine/Documentation/Source/GettingStarted/QuickStartDirectory/QuickStartDirectory.INT.udn
@@ -2,6 +2,7 @@ Availability:Public
 Title:Unreal Engine 4 Quick Start Directory
 Crumbs: %ROOT%, GettingStarted
 Description:All of the quick start guides available to help you get started with Unreal Engine, at a glance.
+version: 4.9
 
 [VAR:Topic]
 [OBJECT:Topic]
