#UE4doc #docs:
Added engine version field

diff --git a/Engine/Documentation/Source/GettingStarted/RecommendedSpecifications/Specifications.INT.udn b/Engine/Documentation/Source/GettingStarted/RecommendedSpecifications/Specifications.INT.udn
--- a/Engine/Documentation/Source/GettingStarted/RecommendedSpecifications/Specifications.INT.udn
+++ b/Engine/Documentation/Source/GettingStarted/RecommendedSpecifications/Specifications.INT.udn
@@ -3,6 +3,7 @@ Title:Hardware & Software Specifications
 SEO-Title:Hardware & Software Specifications
 Crumbs: %ROOT%, GettingStarted
 Description:Minimum and recommended hardware specifications and necessary software for developing with Unreal Engine.
+version: 4.9
 
 [PUBLISH:Rocket]
 %Globals:OSSelect_Linux%
