#UE4doc #docs:
Resubmitting as not all changes were saved

diff --git a/Engine/Documentation/Source/GettingStarted/Installation/InstallingUnrealEngine.INT.udn b/Engine/Documentation/Source/GettingStarted/Installation/InstallingUnrealEngine.INT.udn
--- a/Engine/Documentation/Source/GettingStarted/Installation/InstallingUnrealEngine.INT.udn
+++ b/Engine/Documentation/Source/GettingStarted/Installation/InstallingUnrealEngine.INT.udn
@@ -8,7 +8,7 @@ Related: Engine/Content/QuickStart
 Related: Programming/QuickStart
 Related: Engine/QuickStart
 Related: Support/Builds
-Version: 4.8
+Version: 4.9
 SkillLevel: Beginner
 
 By the end of this tutorial, you'll be able to:
