#UE4doc #docs:
Added engine version field
Changed font settings for multiple terms

diff --git a/Engine/Documentation/Source/GettingStarted/RunningUnrealEngine/RunningUnrealEngine.INT.udn b/Engine/Documentation/Source/GettingStarted/RunningUnrealEngine/RunningUnrealEngine.INT.udn
--- a/Engine/Documentation/Source/GettingStarted/RunningUnrealEngine/RunningUnrealEngine.INT.udn
+++ b/Engine/Documentation/Source/GettingStarted/RunningUnrealEngine/RunningUnrealEngine.INT.udn
@@ -2,6 +2,7 @@ Availability:Public
 Title:Running Unreal Engine
 Crumbs: %ROOT%, GettingStarted
 Description:Running the engine as a game or editor, including how to launch a specific map and useful in-game commands.
+version: 4.9
 
 %Globals:OSSelect%
 
@@ -25,11 +26,11 @@ This process will launch a copy of Unreal Editor.
 
 The process of running the editor simply requires passing the name of the project to run as an argument to the executable.
 
-[REGION:tip]
+[REGION:note]
 You can add the `-debug` switch to force the executable to load the debug version of the modules for your project, which 
 contain all of the debugging symbols. This is necessary even when debugging through Visual Studio with the configuration 
 set to **Debug** because the main executable is always compiled using the **Development** configuration. Of course, you 
-must first compile your modules using the **Debug** configuration so they exist for the executable to load.
+must first compile your modules using the Debug configuration so they exist for the executable to load.
 [/REGION]
 
 [/PUBLISH:Rocket]
@@ -171,14 +172,15 @@ or open [content examples and sample games](Resources).
 ## Running an Uncooked Game
 
 Once a project is loaded in Unreal Editor, you can [test your gameplay](Engine/UI/LevelEditor/InEditorTesting#modes) in Uncooked Game mode from the **Play In** menu. 
-To play the uncooked game in its own window, select the [**New Window At** > **Default Player Start** mode](Engine/UI/LevelEditor/InEditorTesting#modes) using 
-the [**Play In** dropdown menu](Engine/UI/LevelEditor/InEditorTesting#toolbar) in the Level Editor toolbar.
+To play the uncooked game in its own window, select the [New Window At > Default Player Start mode](Engine/UI/LevelEditor/InEditorTesting#modes) using 
+the [Play In dropdown menu](Engine/UI/LevelEditor/InEditorTesting#toolbar) in the Level Editor toolbar.
+
 
-**Click to see full-size image:**
 
 [REGION:lightbox]
 [![](Engine/UI/LevelEditor/InEditorTesting/PIE_newwindow.png)(w:300)](Engine/UI/LevelEditor/InEditorTesting/PIE_newwindow.png)
 [/REGION:lightbox]
+_Click to see full-size image:_
 
 Running the uncooked game version of the engine executable using either of the below methods will result in the same behavior.
 
