#UE4doc #docs:
Added engine version field

diff --git a/Engine/Documentation/Source/GettingStarted/SubEditors/SubEditors.INT.udn b/Engine/Documentation/Source/GettingStarted/SubEditors/SubEditors.INT.udn
--- a/Engine/Documentation/Source/GettingStarted/SubEditors/SubEditors.INT.udn
+++ b/Engine/Documentation/Source/GettingStarted/SubEditors/SubEditors.INT.udn
@@ -2,6 +2,7 @@ Availability: Public
 Title:Tools and Editors
 Crumbs: %ROOT%, Gameplay, GettingStarted
 Description:There are several different types of Editors contained within Unreal Engine 4, this page highlights each of them. 
+version: 4.9
 
 [VAR:Topic]
 [OBJECT:Topic]
