#UE4doc #docs:
Added engine version field

diff --git a/Engine/Documentation/Source/GettingStarted/Terminology/Terminology.INT.udn b/Engine/Documentation/Source/GettingStarted/Terminology/Terminology.INT.udn
--- a/Engine/Documentation/Source/GettingStarted/Terminology/Terminology.INT.udn
+++ b/Engine/Documentation/Source/GettingStarted/Terminology/Terminology.INT.udn
@@ -2,6 +2,7 @@ Availability: Public
 Title:Unreal Engine 4 Terminology
 Crumbs: %ROOT%, Gameplay, GettingStarted
 Description: Covers the most commonly used terms when working with Unreal Engine 4. 
+version: 4.9
 
 [VAR:Topic]
 [OBJECT:Topic]
