#UE4doc #docs:
Added engine version field

diff --git a/Engine/Documentation/Source/Platforms/Platforms.INT.udn b/Engine/Documentation/Source/Platforms/Platforms.INT.udn
--- a/Engine/Documentation/Source/Platforms/Platforms.INT.udn
+++ b/Engine/Documentation/Source/Platforms/Platforms.INT.udn
@@ -3,6 +3,7 @@ Title:Platform Development
 Crumbs:%ROOT%
 Description:Information over developing for platforms other than PC.
 Navigation:platform
+version: 4.9
 
 [VAR:Params]
 [PARAM:image]
