#UE4doc #docs:
Added engine version field

diff --git a/Engine/Documentation/Source/Platforms/Mobile/Mobile.INT.udn b/Engine/Documentation/Source/Platforms/Mobile/Mobile.INT.udn
--- a/Engine/Documentation/Source/Platforms/Mobile/Mobile.INT.udn
+++ b/Engine/Documentation/Source/Platforms/Mobile/Mobile.INT.udn
@@ -4,6 +4,7 @@ Crumbs:%ROOT%, Platforms
 Description:Information for developers creating mobile games.
 Platform:Mobile
 Navigation:platform
+version: 4.9
 
 [REGION:banner]
 ![](mobile_banner.png)
