#UE4doc #docs:
Added engine version field

diff --git a/Engine/Documentation/Source/Platforms/Mobile/Lighting/MobileLighting.INT.udn b/Engine/Documentation/Source/Platforms/Mobile/Lighting/MobileLighting.INT.udn
--- a/Engine/Documentation/Source/Platforms/Mobile/Lighting/MobileLighting.INT.udn
+++ b/Engine/Documentation/Source/Platforms/Mobile/Lighting/MobileLighting.INT.udn
@@ -3,6 +3,7 @@ Title:Lighting for Mobile Platforms
 Crumbs:%ROOT%, Platforms, Platforms/Mobile
 Description:Setting up lighting for mobile games.
 Platform:Mobile
+version: 4.9
 
 ## Supported Features
 
