#UE4doc #docs:
Added engine version field
Changed font settings for multiple terms

diff --git a/Engine/Documentation/Source/Platforms/Mobile/Materials/MobileMaterials.INT.udn b/Engine/Documentation/Source/Platforms/Mobile/Materials/MobileMaterials.INT.udn
--- a/Engine/Documentation/Source/Platforms/Mobile/Materials/MobileMaterials.INT.udn
+++ b/Engine/Documentation/Source/Platforms/Mobile/Materials/MobileMaterials.INT.udn
@@ -9,11 +9,10 @@ Related: Engine/Rendering/Materials/CustomizedUVs
 Related: Engine/Rendering/Materials/MaterialProperties/BlendModes
 Related: Engine/Rendering/Materials/ExpressionReference/Utility
 Related: Platforms/Mobile/Lighting
+version: 4.9
 
 [TOC(start:2)]
 
-## Overview
-
 Materials on mobile platforms (Feature Level ES 2) use the same node-based creation process used by other platforms. Almost all of the Material inputs are available to mobile Materials, 
 and Materials will generally look and behave the same on mobile platforms as they do on other platforms. For instance, the following Material inputs all work exactly the same on mobile platforms 
 as they do on PC.
@@ -31,7 +30,7 @@ There are some considerations you must take into account when creating Materials
 * Only Default and Unlit shading models are available.
 * [](Engine/Rendering/Materials/CustomizedUVs) should be used to avoid dependent texture fetches (no math on texture UVs).
 * Translucent and Masked Materials are extremely expensive. It is recommended that you use [Opaque Materials](Engine/Rendering/Materials/MaterialProperties/BlendModes#opaque) wherever possible.
-* Depth fade can be used in translucent Materials for **iOS platforms**. It is _**NOT**_ supported on platforms where the hardware does not support fetching from the depth buffer and 
+* Depth fade can be used in translucent Materials for **iOS platforms**. It is not supported on platforms where the hardware does not support fetching from the depth buffer and 
 	emulating this would incur an unacceptable performance cost.
 
 
@@ -60,5 +59,5 @@ To maintain compatibility with older devices, use the [Feature Level](Engine/Ren
 
 | Issue | Solution |
 | ----- | -------- |
-| Default Material is displayed in the Mobile Previewer. | Open the [log file](Engine/Basics/DirectoryStructure) in **Engine\Saved** and search for **error messages** to know how to fix the Material. |
+| Default Material is displayed in the Mobile Previewer. | Open the [log file](Engine/Basics/DirectoryStructure) in **Engine\Saved** and search for error messages to know how to fix the Material. |
 | Material Animation with Time Expressions looks choppy but framerate is good elsewhere. | Make sure the Period property is enabled in the **Time** expression along with a Period. This will move the computation to the CPU and run at full-precision instead of half. |
