#UE4doc #docs:
Added engine version field
Removed some redundant text
Changed a tip region to a note due to the content

diff --git a/Engine/Documentation/Source/Platforms/Mobile/Previewer/MobilePreviewer.INT.udn b/Engine/Documentation/Source/Platforms/Mobile/Previewer/MobilePreviewer.INT.udn
--- a/Engine/Documentation/Source/Platforms/Mobile/Previewer/MobilePreviewer.INT.udn
+++ b/Engine/Documentation/Source/Platforms/Mobile/Previewer/MobilePreviewer.INT.udn
@@ -3,6 +3,7 @@ Title:Mobile Previewer
 Crumbs:%ROOT%, Platforms, Platforms/Mobile
 Description: Previewing games on mobile platforms will look in the UE4 editor.
 Platform:Mobile
+version: 4.9
 
 ![](T_Rendering_Level_Examples.png)
 
@@ -17,7 +18,7 @@ When enabling different preview levels the Materials in your scene will be re-co
 You can also switch back and forth between rendering levels seamlessly inside of UE4 without having to re-start the editor.
 To change the which rendering level is used in the current UE4 Editor session, you will need to do the following. 
 
-[region:tip]
+[region:note]
 While you do not need to re-start the Editor for the changes to take effect the first time you change the Preview Rendering Level the editor will need time to re-compile shaders.
 [/region]
 
@@ -31,7 +32,7 @@ While you do not need to re-start the Editor for the changes to take effect the
 ![](T_PRL.png)
 
 [region:note]
-Please note that the Mobile previewer is intended to match mobile devices as closely as possible but is not indicative of what your project is going to look like on the target device.
+The Mobile previewer is intended to match mobile devices as closely as possible but is not indicative of what your project is going to look like on the target device.
 You should always make sure to fully test your project on your target device and only use the mobile previewer to see if your work is headed in the right direction.
 [/region]
 
