#UE4doc #docs:
Added engine version field

diff --git a/Engine/Documentation/Source/Programming/ProgrammingGuide.INT.udn b/Engine/Documentation/Source/Programming/ProgrammingGuide.INT.udn
--- a/Engine/Documentation/Source/Programming/ProgrammingGuide.INT.udn
+++ b/Engine/Documentation/Source/Programming/ProgrammingGuide.INT.udn
@@ -2,6 +2,7 @@ Availability:Public
 Title:Programming Guide
 Crumbs: %ROOT%
 Description:Information for programmers developing with Unreal Engine.
+version: 4.9
 
 [VAR:Params]
 [PARAM:image]
