#UE4doc #docs:
Added engine version field
Changed font settings for a few terms
Removed a note region

diff --git a/Engine/Documentation/Source/Programming/Animation/AnimNodes/CreatingAnimNodes.INT.udn b/Engine/Documentation/Source/Programming/Animation/AnimNodes/CreatingAnimNodes.INT.udn
--- a/Engine/Documentation/Source/Programming/Animation/AnimNodes/CreatingAnimNodes.INT.udn
+++ b/Engine/Documentation/Source/Programming/Animation/AnimNodes/CreatingAnimNodes.INT.udn
@@ -2,7 +2,7 @@ Availability:Public
 Title: Animation Node Technical Guide
 Crumbs: %ROOT%, Engine, Engine/Animation
 Description: Guide to creating new nodes for use within graphs in Anim Blueprints.
-Version: 4.5
+Version: 4.9
 
 [REGION:banner]
 ![Anim Nodes](animnodes_banner.png)(convert:false)
@@ -10,7 +10,7 @@ Version: 4.5
 
 [TOC(start:2)]
 
-Animation nodes are used inside of [_Anim Blueprints_](Engine/Animation/AnimBlueprints) 
+Animation nodes are used inside of [Anim Blueprints](Engine/Animation/AnimBlueprints) 
 to perform operations, such as blending of poses or direct manipulation of bones. Several animation 
 nodes are provided, but custom nodes can be created to suit the needs of any game.
 
@@ -100,16 +100,13 @@ The following metadata keys are available:
 | `AlwaysAsPin` | The property is always exposed as a data pin. |
 
 (#OptionalPins)
-[REGION:note]
-	**Optional Pins**
-	[REGION:none]
-		For properties with `PinHiddenByDefault` or `PinShownByDefault`, a check box will show up next 
-		to the property in the **Details** panel, allowing you to show or hide it.
-	[/REGION]	
-	[REGION:none]
-		![Property Details](node_property_details.png)
-	[/REGION]
-
+**Optional Pins**
+[REGION:none]
+	For properties with `PinHiddenByDefault` or `PinShownByDefault`, a check box will show up next 
+	to the property in the **Details** panel, allowing you to show or hide it.
+[/REGION]	
+[REGION:none]
+	![Property Details](node_property_details.png)
 [/REGION]
 
 ## Editor Node
@@ -126,7 +123,7 @@ The editor-time class should contain an instance of your runtime node exposed as
 
 ![Node Title](node_title.png)
 
-The background color and text of the title of the animation node displayed in the graph of the _Animation Blueprint_ in 
+The background color and text of the title of the animation node displayed in the graph of the Animation Blueprint in 
 Persona are defined by overriding the `GetNodeTitle()` and `GetNodeTitleColor()` functions.
 
 For example, the `UAnimGraphNode_ApplyAdditive` node uses a gray background and displays "Apply Additive":
@@ -154,7 +151,7 @@ The tooltip displayed when hovering over the node in Persona is defined by overr
 
 ### Context Menu
 
-Each animation node can add node-specific options to the context menu displayed when **right-clicking** on the 
+Each animation node can add node-specific options to the context menu displayed when **Right-clicking** on the 
 node in a graph in Persona. Options are added using the `GetContextMenuActions()` function which is a 
 member of all Blueprint nodes.
 
