#UE4doc #docs - Added engine version field.

diff --git a/Engine/Documentation/Source/Resources/ContentExamples/NavMesh/1_3/1_3.INT.udn b/Engine/Documentation/Source/Resources/ContentExamples/NavMesh/1_3/1_3.INT.udn
--- a/Engine/Documentation/Source/Resources/ContentExamples/NavMesh/1_3/1_3.INT.udn
+++ b/Engine/Documentation/Source/Resources/ContentExamples/NavMesh/1_3/1_3.INT.udn
@@ -6,6 +6,7 @@ Related: Engine/AI/BehaviorTrees/QuickStart
 Related: Gameplay/AI
 Related: Engine/AI/BehaviorTrees/NodeReference
 Related: Engine/Actors/Volumes
+Version: 4.9
 
 [VAR:Steps]
 [OBJECT:Navigation]
@@ -69,6 +70,7 @@ Related: Engine/Actors/Volumes
 
 The **NavMesh** automatically updates when Actors like **Static Meshes** are moved in the editor. In this example, you can move the block Static Mesh and position it so that it bridges the gap between the ramps. The NavMesh will automatically update and build across the newly-formed bridge, allowing the Pawn to navigate across it when you **Play** or **Simulate**.
 
+
 [REGION:lightbox]
 [![](NavMesh_1_3_001.png)(w:225)](NavMesh_1_3_001.png)
 [![](NavMesh_1_3_002.png)(w:225)](NavMesh_1_3_002.png)
