#UE4doc #docs - Added engine version field. Added/Removed bolded terms.

diff --git a/Engine/Documentation/Source/Resources/ContentExamples/Networking/NetworkFeatures.INT.udn b/Engine/Documentation/Source/Resources/ContentExamples/Networking/NetworkFeatures.INT.udn
--- a/Engine/Documentation/Source/Resources/ContentExamples/Networking/NetworkFeatures.INT.udn
+++ b/Engine/Documentation/Source/Resources/ContentExamples/Networking/NetworkFeatures.INT.udn
@@ -7,6 +7,7 @@ Related: Gameplay/Networking/Actors/Components
 Related: Gameplay/Networking/Actors
 Related: Gameplay/Networking/Actors/RPCs
 Related: Engine/UI/LevelEditor/InEditorTesting/NetworkedPlayInEditor
+Version: 4.9
 
 [REGION:banner]
 ![](NetworkFeatures_Header.png)(convert:false)
@@ -47,7 +48,7 @@ Related: Engine/UI/LevelEditor/InEditorTesting/NetworkedPlayInEditor
 
 
 
-The **Network_Features** map inside the **Content Examples** project illustrates how to implement network functionality entirely through Blueprints. It provides key examples in **Replication** including the **Replicates** boolean which can be set on Actors, detecting **Network Authority** and executing different logic based on that authority, usage of **Variable Replication** and how it is used during gameplay, variable replication using **RepNotify** to execute logic based on updates made to a variable that is replicated, as well as the use of **Function Replication** to handle non-critical gameplay elements or "one-off" events.
+The **Network_Features** map inside the Content Examples project illustrates how to implement network functionality entirely through Blueprints. It provides key examples in **Replication** including the **Replicates** boolean which can be set on Actors, detecting **Network Authority** and executing different logic based on that authority, usage of **Variable Replication** and how it is used during gameplay, variable replication using **RepNotify** to execute logic based on updates made to a variable that is replicated, as well as the use of **Function Replication** to handle non-critical gameplay elements or "one-off" events.
 
 The other section in the map provides insight on **Relevancy** and how to properly handle communicating information to and from the Server and Clients. Often times, it is not critical that every bit of information be passed from the server to every client (e.g. a player that is on one side of the map does not need to know about another player opening a chest on the other side of the map). There are situations when this would become relevant to that player however (e.g. if that player moved from the other side of the map to the chest that was opened), the Content Examples in this map will show how to handle situations like this as well as how to handle a player that joins an in progress session. 
 
@@ -76,10 +77,10 @@ The examples below provide possible solutions to the problem stated above and wh
 
 | Example | What is Demonstrated |
 | ------- | ------------ |
-| **[2.1 Network Relevancy - Not Replicated At All](Resources\ContentExamples\Networking\2_1)** | An example of relevancy where **No Replication** is used (e.g. Server opens a chest, Client does not see the chest as opened).  |
-| **[2.2 Network Relevancy - Function Replication Solution Attempt](Resources\ContentExamples\Networking\2_2)** | An example of relevancy where **Function Replication** is used (e.g. Server opens a chest, Client does not see the chest open and does not see it open when moving close to it). |
-| **[2.3 Network Relevancy - Variable Replication Solution Attempt](Resources\ContentExamples\Networking\2_3)** | An example of relevancy where **Variable Replication** is used. e.g. Server opens a chest, Client does not see the chest open but sees it open when moving close to it. Particle effects also play, which is not desired for this example as the chest was previously opened by another player. |
-| **[2.4 Network Relevancy - Combination Solution](Resources\ContentExamples\Networking\2_4)** | An example of relevancy where a combination of **Function Replication** is used to play the effects associated with opening the chest and a **Replicated Variable** is used to preserve the state of the chest and if it has already been opened or not. |
+| **[2.1 Network Relevancy - Not Replicated At All](Resources\ContentExamples\Networking\2_1)** | An example of relevancy where No Replication is used (e.g. Server opens a chest, Client does not see the chest as opened).  |
+| **[2.2 Network Relevancy - Function Replication Solution Attempt](Resources\ContentExamples\Networking\2_2)** | An example of relevancy where Function Replication is used (e.g. Server opens a chest, Client does not see the chest open and does not see it open when moving close to it). |
+| **[2.3 Network Relevancy - Variable Replication Solution Attempt](Resources\ContentExamples\Networking\2_3)** | An example of relevancy where Variable Replication is used. e.g. Server opens a chest, Client does not see the chest open but sees it open when moving close to it. Particle effects also play, which is not desired for this example as the chest was previously opened by another player. |
+| **[2.4 Network Relevancy - Combination Solution](Resources\ContentExamples\Networking\2_4)** | An example of relevancy where a combination of Function Replication is used to play the effects associated with opening the chest and a Replicated Variable is used to preserve the state of the chest and if it has already been opened or not. |
 
 
 
