#UE4doc #docs - Added engine version field.

diff --git a/Engine/Documentation/Source/Resources/ContentExamples/Networking/1_5/1_5.INT.udn b/Engine/Documentation/Source/Resources/ContentExamples/Networking/1_5/1_5.INT.udn
--- a/Engine/Documentation/Source/Resources/ContentExamples/Networking/1_5/1_5.INT.udn
+++ b/Engine/Documentation/Source/Resources/ContentExamples/Networking/1_5/1_5.INT.udn
@@ -7,6 +7,7 @@ Related: Gameplay/Networking/Actors/Components
 Related: Gameplay/Networking/Actors
 Related: Gameplay/Networking/Actors/RPCs
 Related: Engine/UI/LevelEditor/InEditorTesting/NetworkedPlayInEditor
+Version: 4.9
 
 [VAR:Steps]
 [OBJECT:Navigation]
