#UE4doc #docs - added engine version field.

diff --git a/Engine/Documentation/Source/Resources/ContentExamples/Networking/2_1/2_1.INT.udn b/Engine/Documentation/Source/Resources/ContentExamples/Networking/2_1/2_1.INT.udn
--- a/Engine/Documentation/Source/Resources/ContentExamples/Networking/2_1/2_1.INT.udn
+++ b/Engine/Documentation/Source/Resources/ContentExamples/Networking/2_1/2_1.INT.udn
@@ -7,6 +7,7 @@ Related: Gameplay/Networking/Actors/Components
 Related: Gameplay/Networking/Actors
 Related: Gameplay/Networking/Actors/RPCs
 Related: Engine/UI/LevelEditor/InEditorTesting/NetworkedPlayInEditor
+Version:4.9
 
 [VAR:Steps]
 [OBJECT:Navigation]
diff --git a/Engine/Documentation/Source/Resources/ContentExamples/Networking/2_2/2_2.INT.udn b/Engine/Documentation/Source/Resources/ContentExamples/Networking/2_2/2_2.INT.udn
--- a/Engine/Documentation/Source/Resources/ContentExamples/Networking/2_2/2_2.INT.udn
+++ b/Engine/Documentation/Source/Resources/ContentExamples/Networking/2_2/2_2.INT.udn
@@ -7,6 +7,7 @@ Related: Gameplay/Networking/Actors/Components
 Related: Gameplay/Networking/Actors
 Related: Gameplay/Networking/Actors/RPCs
 Related: Engine/UI/LevelEditor/InEditorTesting/NetworkedPlayInEditor
+Version: 4.9
 
 [VAR:Steps]
 [OBJECT:Navigation]
