#UE4doc #docs - Added engine version field. Removed Overview header. Removed ToC bc of single line entry. Removed extraneous bolding of terms.

diff --git a/Engine/Documentation/Source/Resources/ContentExamples/Paper2D/Paper2D.INT.udn b/Engine/Documentation/Source/Resources/ContentExamples/Paper2D/Paper2D.INT.udn
--- a/Engine/Documentation/Source/Resources/ContentExamples/Paper2D/Paper2D.INT.udn
+++ b/Engine/Documentation/Source/Resources/ContentExamples/Paper2D/Paper2D.INT.udn
@@ -8,6 +8,7 @@ Related: Engine/Content/Tools/RenderToTextureTools/5
 Related: Engine/Content/Tools/RenderToTextureTools/3
 Related: Engine/Rendering/ParticleSystems/Reference/Modules/SubUV
 Related: Engine/Rendering/Materials/Functions/Reference/Misc
+Version: 4.9
 
 [REGION:banner]
 ![](paper2d_header.png)(convert:false)
@@ -44,12 +45,10 @@ Related: Engine/Rendering/Materials/Functions/Reference/Misc
 [/OBJECT]
 [/VAR]
 
-[TOC(start:2 end:2)]
 
 
-## Overview
 
-The **Paper 2D Content Examples** cover several topics related to 2D game creation including **Sprite** creation and editing via the **Sprite Editor**, **Flipbooks** which are in essence a method in which you can create animated 2D assets using **Sprites** and keyframing, how to control **Flipbooks** with **Blueprints**, adding physics to 2D assets as well as how to create the classic low resolution pixel look for your **Sprites**. 
+The **Paper 2D** map in Content Examples cover several topics related to 2D game creation including **Sprite** creation and editing via the **Sprite Editor**, **Flipbooks** which are in essence a method in which you can create animated 2D assets using Sprites and keyframing, how to control Flipbooks with Blueprints, adding physics to 2D assets as well as how to create the classic low resolution pixel look for your Sprites. 
 
 ## Examples
 
@@ -57,10 +56,10 @@ The **Paper 2D Content Examples** cover several topics related to 2D game creati
 
 | Example | What is Demonstrated |
 | ------- | ------------ |
-| **1.1 Sprites** | A quick and easy way to draw 2D images in UE4, this example shows how to create **Sprites** in the **Content Browser** (e.g. a simple static 2D sprite and an animated 2D sprite is shown). |
-| **1.2 Sprite Editor** | The **Paper 2D Sprite Editor** enables setup and editing of individual **Sprite** assets (e.g. custom collision is shown for 2D assets which can be edited in the **Sprite Editor**). |
-| **1.3 Flipbooks** | **Flipbooks** consist of a series of keyframes, each of which contains a **Sprite** to be displayed and a duration (in frames) to display it (e.g. two examples of 2D characters are shown running using different frames for their animations). |
-| **1.4 Flipbooks (Blueprint Control)** | Shows how **Blueprint Functions** can be used to control animations in a **Flipbook** or to switch **Flipbooks** (e.g. a 2D character performs different animations when pressing different buttons). |
+| **1.1 Sprites** | A quick and easy way to draw 2D images in UE4, this example shows how to create Sprites in the Content Browser (e.g. a simple static 2D sprite and an animated 2D sprite is shown). |
+| **1.2 Sprite Editor** | The Paper 2D Sprite Editor enables setup and editing of individual Sprite assets (e.g. custom collision is shown for 2D assets which can be edited in the Sprite Editor). |
+| **1.3 Flipbooks** | Flipbooks consist of a series of keyframes, each of which contains a Sprite to be displayed and a duration (in frames) to display it (e.g. two examples of 2D characters are shown running using different frames for their animations). |
+| **1.4 Flipbooks (Blueprint Control)** | Shows how Blueprint Functions can be used to control animations in a Flipbook or to switch Flipbooks (e.g. a 2D character performs different animations when pressing different buttons). |
 | **1.5 Physics (3D Physics Limited to a Plane)** | This example illustrates how physics can be applied to 2D assets by using locked 3D physics (e.g. several coins are dropped and physics are applied, but are restricted to the Y axis). |
-| **1.6 Classic Pixel Sprites** | How to achieve the classic pixelated look for **Sprites** by changing texture settings (e.g. imported default texture is shown alongside an edited version which uses a pixel look). |
+| **1.6 Classic Pixel Sprites** | How to achieve the classic pixelated look for Sprites by changing texture settings (e.g. imported default texture is shown alongside an edited version which uses a pixel look). |
 
