#UE4doc #docs - Added engine version field. Removed ToC bc of single line entry. Removed extraneous bolded terms.

diff --git a/Engine/Documentation/Source/Resources/ContentExamples/Physics/Physics.INT.udn b/Engine/Documentation/Source/Resources/ContentExamples/Physics/Physics.INT.udn
--- a/Engine/Documentation/Source/Resources/ContentExamples/Physics/Physics.INT.udn
+++ b/Engine/Documentation/Source/Resources/ContentExamples/Physics/Physics.INT.udn
@@ -8,6 +8,7 @@ Related: Engine/Physics/Constraints/ConstraintsReference
 Related: Engine/Physics/Constraints/ConstraintsUserGuide
 Related: Engine/Physics/Collision
 Related: Engine/Physics/PhysicalMaterials/Reference
+Version: 4.9
 
 [REGION:banner]
 ![](Physics_Header.png)(convert:false)
@@ -44,10 +45,9 @@ Related: Engine/Physics/PhysicalMaterials/Reference
 [/OBJECT]
 [/VAR]
 
-[TOC(start:2 end:2)]
 
 
-Inside the **Physics** map, you will find several examples of how physics can be used in Unreal Engine 4. Creating **Physics Actors** out of **Static Meshes** so that they respond to player interaction, applying a **Physics Asset** to a **Skeletal Mesh** so that it responds to physics, and several examples of how **Force**, **Impulse**, **Motors**, and **Physical Constraints** are on display in this map with details of how they can aid in your level design or gameplay scenarios. 
+Inside the **Physics** map, you will find several examples of how physics can be used in Unreal Engine 4. Creating Physics Actors out of Static Meshes so that they respond to player interaction, applying a Physics Asset to a Skeletal Mesh so that it responds to physics, and several examples of how Force, Impulse, Motors, and Physical Constraints are on display in this map with details of how they can aid in your level design or gameplay scenarios. 
 
 
 
