#UE4doc #docs -  Added engine version field. Fixed spacing issue between words.

diff --git a/Engine/Documentation/Source/Resources/ContentExamples/Physics/1_3/1_3.INT.udn b/Engine/Documentation/Source/Resources/ContentExamples/Physics/1_3/1_3.INT.udn
--- a/Engine/Documentation/Source/Resources/ContentExamples/Physics/1_3/1_3.INT.udn
+++ b/Engine/Documentation/Source/Resources/ContentExamples/Physics/1_3/1_3.INT.udn
@@ -8,6 +8,7 @@ Related: Engine/Physics/Constraints/ConstraintsReference
 Related: Engine/Physics/Constraints/ConstraintsUserGuide
 Related: Engine/Physics/Collision
 Related: Engine/Physics/PhysicalMaterials/Reference
+Version: 4.9
 
 [VAR:Steps]
 [OBJECT:Navigation]
@@ -71,11 +72,11 @@ Related: Engine/Physics/PhysicalMaterials/Reference
 [/REGION]
 
 
-Using a **Radial Force Actor**, you can apply a **Constant Force** (listed as _Force Strength_ in the editor) or an **Impulse**. Impulses need to be fired from a **Blueprint**, and apply their entire force in only one frame, as compared to every frame from a Constant Force. In the content example, only a Constant Forceis being applied.
+Using a **Radial Force Actor**, you can apply a **Constant Force** (listed as _Force Strength_ in the editor) or an **Impulse**. Impulses need to be fired from a **Blueprint**, and apply their entire force in only one frame, as compared to every frame from a Constant Force. In the content example, only a Constant Force is being applied.
 
 ![](physics_1_3_002.png)(w:450) ![](physics_1_3_003.png)(w:450)
 
-This force is applied to everything that simulates physics.
+_This force is applied to everything that simulates physics._
 
 
 
