#UE4doc #docs - Added engine version field.

diff --git a/Engine/Documentation/Source/Resources/ContentExamples/Physics/1_6/1_6.INT.udn b/Engine/Documentation/Source/Resources/ContentExamples/Physics/1_6/1_6.INT.udn
--- a/Engine/Documentation/Source/Resources/ContentExamples/Physics/1_6/1_6.INT.udn
+++ b/Engine/Documentation/Source/Resources/ContentExamples/Physics/1_6/1_6.INT.udn
@@ -8,6 +8,7 @@ Related: Engine/Physics/Constraints/ConstraintsReference
 Related: Engine/Physics/Constraints/ConstraintsUserGuide
 Related: Engine/Physics/Collision
 Related: Engine/Physics/PhysicalMaterials/Reference
+Version: 4.9
 
 [VAR:Steps]
 [OBJECT:Navigation]
@@ -71,7 +72,7 @@ Related: Engine/Physics/PhysicalMaterials/Reference
 [/REGION]
 
 
-A constraint with an angular motor will essentially spin, but the detail is in **how** it spins. An angular motor **Driven by Velocity** will apply a constant force to essentially "spin," where as a **Driven by Position** will attempt to obtain an angle based on the size of the spring and dampening applied to said spring.
+A Physics Constraint with an **Angular Motor** will essentially spin, but the detail is in **how** it spins. An angular motor **Driven by Velocity** will apply a constant force to essentially "spin," where as a **Driven by Position** will attempt to obtain an angle based on the size of the spring and dampening applied to said spring.
 
 <div align="center">
 ![](physics_1_6_002.png)(h:300) ![](physics_1_6_003.png)(h:300)
diff --git a/Engine/Documentation/Source/Resources/ContentExamples/Physics/1_7/1_7.INT.udn b/Engine/Documentation/Source/Resources/ContentExamples/Physics/1_7/1_7.INT.udn
--- a/Engine/Documentation/Source/Resources/ContentExamples/Physics/1_7/1_7.INT.udn
+++ b/Engine/Documentation/Source/Resources/ContentExamples/Physics/1_7/1_7.INT.udn
@@ -8,6 +8,7 @@ Related: Engine/Physics/Constraints/ConstraintsReference
 Related: Engine/Physics/Constraints/ConstraintsUserGuide
 Related: Engine/Physics/Collision
 Related: Engine/Physics/PhysicalMaterials/Reference
+Version: 4.9
 
 [VAR:Steps]
 [OBJECT:Navigation]
