#UE4doc #docs - Added engine version field.

diff --git a/Engine/Documentation/Source/Resources/ContentExamples/Physics/1_8/1_8.INT.udn b/Engine/Documentation/Source/Resources/ContentExamples/Physics/1_8/1_8.INT.udn
--- a/Engine/Documentation/Source/Resources/ContentExamples/Physics/1_8/1_8.INT.udn
+++ b/Engine/Documentation/Source/Resources/ContentExamples/Physics/1_8/1_8.INT.udn
@@ -8,6 +8,7 @@ Related: Engine/Physics/Constraints/ConstraintsReference
 Related: Engine/Physics/Constraints/ConstraintsUserGuide
 Related: Engine/Physics/Collision
 Related: Engine/Physics/PhysicalMaterials/Reference
+Version: 4.9
 
 [VAR:Steps]
 [OBJECT:Navigation]
@@ -66,15 +67,16 @@ Related: Engine/Physics/PhysicalMaterials/Reference
 [/VAR]
 
 
+
 <div align="center">
-	![](physics_1_8_002.png)(w:450) ![](physics_1_8_001.png)(w:450)
+	![](physics_1_8_002.png)(w:960) <!---![](physics_1_8_001.png)(w:450)--->
 </div>
 
 
 A constraint can be made to break once a certain amount of angular or linear force has been applied to the constraint through its constrained objects. An **Angular Force** would be considered like snapping a twig or a torsion force. A **Linear Force** is a "pulling" force along the "length" of the constraint, such as applying too much weight to a crane's lift line.
 
-[REGION:fullwidth]
-	![](physics_1_8_003.png)
-[/REGION]
+<div align="center">
+	![](physics_1_8_003.png)(w:500)
+</div>
 
 
