#UE4doc #docs - Added engine version field. Removed Overview header. Remove ToC for single line entry. Added/Remove bolding.

diff --git a/Engine/Documentation/Source/Resources/ContentExamples/PivotPainter/PivotPainter.INT.udn b/Engine/Documentation/Source/Resources/ContentExamples/PivotPainter/PivotPainter.INT.udn
--- a/Engine/Documentation/Source/Resources/ContentExamples/PivotPainter/PivotPainter.INT.udn
+++ b/Engine/Documentation/Source/Resources/ContentExamples/PivotPainter/PivotPainter.INT.udn
@@ -6,6 +6,7 @@ Related: Engine/Rendering/Materials
 Related: Engine/Rendering/Materials/ExpressionReference
 Related: Engine/Rendering/Materials/Functions/Reference/PivotPainter
 Related: Engine/Rendering/Materials/Functions/Creating
+Version: 4.9
 
 [REGION:banner]
 ![](PivotPainter_Header.png)(convert:false)
@@ -42,13 +43,11 @@ Related: Engine/Rendering/Materials/Functions/Creating
 [/OBJECT]
 [/VAR]
 
-[TOC(start:2 end:2)]
 
 
+The [Pivot Painter Tool](Engine\Content\Tools\PivotPainter) is a MAXScript that stores model pivot and rotation information in the model's vertex data. That information can then be referenced inside of Unreal Engine 4's shader system to create interactive effects. 
 
-The [**Pivot Painter Tool**](Engine\Content\Tools\PivotPainter) is a MAXScript that stores model pivot and rotation information in the model's vertex data. That information can then be referenced inside of Unreal Engine 4's shader system to create interactive effects. 
-
-The **PivotPainter** map inside the **Content Examples** project provides several examples which use **Pivot Painter Material Functions** to produce models that achieve this effect and demonstrates how the models can be manipulated through several forces such as wind or player interaction. 
+The **PivotPainter** map inside the Content Examples project provides several examples which use **Pivot Painter Material Functions** to produce models that achieve this effect and demonstrates how the models can be manipulated through several forces such as wind or player interaction. 
 
 
  
diff --git a/Engine/Documentation/Source/Resources/ContentExamples/PostProcessing/PostProcessing.INT.udn b/Engine/Documentation/Source/Resources/ContentExamples/PostProcessing/PostProcessing.INT.udn
--- a/Engine/Documentation/Source/Resources/ContentExamples/PostProcessing/PostProcessing.INT.udn
+++ b/Engine/Documentation/Source/Resources/ContentExamples/PostProcessing/PostProcessing.INT.udn
@@ -6,6 +6,7 @@ Related: Engine/Rendering/Overview
 Related: Engine/Rendering/PostProcessEffects
 Related: Engine/Rendering/PostProcessEffects/PostProcessMaterials
 Related: Resources/Showcases/Stylized
+Version: 4.9
 
 [REGION:banner]
 ![](PostProcessing_Header.png)(convert:false)
@@ -42,12 +43,9 @@ Related: Resources/Showcases/Stylized
 [/OBJECT]
 [/VAR]
 
-[TOC(start:2 end:2)]
 
 
-## Overview
-
-The Content Examples provided in the **PostProcessing** map illustrate how to add screen effects inside and outside of a **Post Processing Volume**. By adjusting the properties of a **Post Processing Volume** you can achieve effects like lens flares, screen tinting or blurring and many others as well as define how volumes interact with one another or a player. 
+The Content Examples provided in the **PostProcessing** map illustrate how to add screen effects inside and outside of a Post Processing Volume. By adjusting the properties of a Post Processing Volume you can achieve effects like lens flares, screen tinting or blurring and many others as well as define how volumes interact with one another or a player. 
 
 
 [REGION:note]
@@ -63,25 +61,25 @@ Listed below are the examples provided inside the **PostProcessing** map:
 
 | Example | What is Demonstrated |
 | ------- | ------------ |
-| **[1.1 Default Settings](Resources\ContentExamples\PostProcessing\1_1)** | The default settings of a **Post Process Volume** (also the settings used when a camera is not within a **Post Process Volume**). |
-| **[1.2 Film](Resources\ContentExamples\PostProcessing\1_2)** | **Film** allows you to adjust the color of your scenes. |
-| **[1.3 SceneColor](Resources\ContentExamples\PostProcessing\1_3)** | **Scene Color** applies effects, instead of the color shifts seen in the **Film** section, to your scene (e.g. **Vignette**, **Scene Fringe**, and **Tone Mapping** effects are shown). |
-| **[1.4 Bloom](Resources\ContentExamples\PostProcessing\1_4)** | **Bloom** simulates the effects the eye perceives when viewing objects that appear very bright. |
-| **[1.5 Ambient Cubemap](Resources\ContentExamples\PostProcessing\1_5)** | The **Ambient Cubemap** lights the scene from a provided image. |
-| **[1.6 AutoExposure](Resources\ContentExamples\PostProcessing\1_6)** | Shows how to simulate the effect of the human eye adjusting to a bright environment when coming from a dark environment or vice versa with **Auto Exposure**.|
-| **[1.7 Lens Flares](Resources\ContentExamples\PostProcessing\1_7)** | **Lens Flare** simulates the scattering of light when viewing bright objects through imperfections found in camera lenses. |
+| **[1.1 Default Settings](Resources\ContentExamples\PostProcessing\1_1)** | The default settings of a Post Process Volume (also the settings used when a camera is not within a Post Process Volume). |
+| **[1.2 Film](Resources\ContentExamples\PostProcessing\1_2)** | Film allows you to adjust the color of your scenes. |
+| **[1.3 SceneColor](Resources\ContentExamples\PostProcessing\1_3)** | Scene Color applies effects, instead of the color shifts seen in the Film section, to your scene (e.g. Vignette, Scene Fringe, and Tone Mapping effects are shown). |
+| **[1.4 Bloom](Resources\ContentExamples\PostProcessing\1_4)** | Bloom simulates the effects the eye perceives when viewing objects that appear very bright. |
+| **[1.5 Ambient Cubemap](Resources\ContentExamples\PostProcessing\1_5)** | The Ambient Cubemap lights the scene from a provided image. |
+| **[1.6 AutoExposure](Resources\ContentExamples\PostProcessing\1_6)** | Shows how to simulate the effect of the human eye adjusting to a bright environment when coming from a dark environment or vice versa with Auto Exposure.|
+| **[1.7 Lens Flares](Resources\ContentExamples\PostProcessing\1_7)** | Lens Flare simulates the scattering of light when viewing bright objects through imperfections found in camera lenses. |
 | **[1.8 Ambient Occlusion](Resources\ContentExamples\PostProcessing\1_8)** | An effect that approximates the attenuation of light due to occlusion. Best used as a subtle effect that darkens corners or crevices to make them appear more natural.|
-| **[1.9 Screenspace Reflection](Resources\ContentExamples\PostProcessing\1_9)** | Alters the appearance of objects in reflections where the **Intensity**, **Quality**, or **Max Roughness** can be set.  |
-| **[1.10 Global Illumination](Resources\ContentExamples\PostProcessing\1_10)** | How to affect the indirect lighting contribution coming from **Lightmass** in order to alter a scene's brightness, tint, or color. |
-| **[1.11 Depth Of Field](Resources\ContentExamples\PostProcessing\1_11)** | Applying a blur effect to a scene based on distance in front or behind a focal point with **Depth of Field**. |
-| **[1.12 Motion Blur](Resources\ContentExamples\PostProcessing\1_12)** | Generating a **Motion Blur** which blurs objects based on its motion. |
-| **[1.13 Screen Percentage](Resources\ContentExamples\PostProcessing\1_13)** | Through **Screen Percentage**, a scene is rendered at a lower resolution and then scaled up. The Screen Percentage value represents a percentage of the overall resolution of the scene. |
-| **[1.14 AA Method](Resources\ContentExamples\PostProcessing\1_14)** | Demonstrates **Anti-Aliasing** (or AA) which improves image quality by smoothing out the edges that appear jagged on screen. |
-| **[1.15 Blur GBuffer intensity](Resources\ContentExamples\PostProcessing\1_15)** | Although hard to see, this effect helps with anti-aliasing by increasing the **Material Roughness** on the hard edges of surfaces, smoothing out the reflections on those hard edges. |
-| **[1.16 Blendables](Resources\ContentExamples\PostProcessing\1_16)** | **Blendables** are, in short, overlays that can be applied to the screen (e.g. screen warping or coloring). |
-| **[1.17 Priorities](Resources\ContentExamples\PostProcessing\1_17)** | Allows you to control which **Post Process Volumes** will override one another if they are overlapping. |
-| **[1.18 Blend Radius and Weight](Resources\ContentExamples\PostProcessing\1_18)** | Used to change the transition from one **Post Process Volume** to another over a distance (e.g. blend scene tint or color based on distance to an object). |
-| **[1.19 Unbound](Resources\ContentExamples\PostProcessing\1_19)** | This allows you to create a global **Post Process** effect without being constrained to being inside an actual volume (e.g. the effects are applied where ever you are in the world). |
+| **[1.9 Screenspace Reflection](Resources\ContentExamples\PostProcessing\1_9)** | Alters the appearance of objects in reflections where the Intensity, Quality, or Max Roughness can be set.  |
+| **[1.10 Global Illumination](Resources\ContentExamples\PostProcessing\1_10)** | How to affect the indirect lighting contribution coming from Lightmass in order to alter a scene's brightness, tint, or color. |
+| **[1.11 Depth Of Field](Resources\ContentExamples\PostProcessing\1_11)** | Applying a blur effect to a scene based on distance in front or behind a focal point with Depth of Field. |
+| **[1.12 Motion Blur](Resources\ContentExamples\PostProcessing\1_12)** | Generating a Motion Blur which blurs objects based on its motion. |
+| **[1.13 Screen Percentage](Resources\ContentExamples\PostProcessing\1_13)** | Through Screen Percentage, a scene is rendered at a lower resolution and then scaled up. The Screen Percentage value represents a percentage of the overall resolution of the scene. |
+| **[1.14 AA Method](Resources\ContentExamples\PostProcessing\1_14)** | Demonstrates Anti-Aliasing (or AA) which improves image quality by smoothing out the edges that appear jagged on screen. |
+| **[1.15 Blur GBuffer intensity](Resources\ContentExamples\PostProcessing\1_15)** | Although hard to see, this effect helps with anti-aliasing by increasing the Material Roughness on the hard edges of surfaces, smoothing out the reflections on those hard edges. |
+| **[1.16 Blendables](Resources\ContentExamples\PostProcessing\1_16)** | Blendables are, in short, overlays that can be applied to the screen (e.g. screen warping or coloring). |
+| **[1.17 Priorities](Resources\ContentExamples\PostProcessing\1_17)** | Allows you to control which Post Process Volumes will override one another if they are overlapping. |
+| **[1.18 Blend Radius and Weight](Resources\ContentExamples\PostProcessing\1_18)** | Used to change the transition from one Post Process Volume to another over a distance (e.g. blend scene tint or color based on distance to an object). |
+| **[1.19 Unbound](Resources\ContentExamples\PostProcessing\1_19)** | This allows you to create a global Post Process effect without being constrained to being inside an actual volume (e.g. the effects are applied where ever you are in the world). |
 
 
 
