#UE4doc #docs - Added engine version field.

diff --git a/Engine/Documentation/Source/Resources/ContentExamples/PostProcessing/1_10/1_10.INT.udn b/Engine/Documentation/Source/Resources/ContentExamples/PostProcessing/1_10/1_10.INT.udn
--- a/Engine/Documentation/Source/Resources/ContentExamples/PostProcessing/1_10/1_10.INT.udn
+++ b/Engine/Documentation/Source/Resources/ContentExamples/PostProcessing/1_10/1_10.INT.udn
@@ -5,6 +5,7 @@ Description:An overview of the Post Processing example level, example 1.10: Glob
 Related: Engine/Rendering/Overview
 Related: Engine/Rendering/PostProcessEffects
 Related: Engine/Rendering/PostProcessEffects/PostProcessMaterials
+Version: 4.9
 
 [VAR:Steps]
 [OBJECT:Navigation]
