#UE4doc #docs - Added engine version field.

diff --git a/Engine/Documentation/Source/Resources/ContentExamples/PostProcessing/1_11/1_11.INT.udn b/Engine/Documentation/Source/Resources/ContentExamples/PostProcessing/1_11/1_11.INT.udn
--- a/Engine/Documentation/Source/Resources/ContentExamples/PostProcessing/1_11/1_11.INT.udn
+++ b/Engine/Documentation/Source/Resources/ContentExamples/PostProcessing/1_11/1_11.INT.udn
@@ -6,6 +6,7 @@ Related: Engine/Rendering/PostProcessEffects/DepthOfField
 Related: Engine/Rendering/Overview
 Related: Engine/Rendering/PostProcessEffects
 Related: Engine/Rendering/PostProcessEffects/PostProcessMaterials
+Version: 4.9
 
 [VAR:Steps]
 [OBJECT:Navigation]
