#UE4 Docs: Jira UEDOC-2359 (Anim Montage / Anim Notifies - Pages updated).

diff --git a/Engine/Documentation/Source/Engine/Animation/AnimMontage/AnimMontage.INT.udn b/Engine/Documentation/Source/Engine/Animation/AnimMontage/AnimMontage.INT.udn
--- a/Engine/Documentation/Source/Engine/Animation/AnimMontage/AnimMontage.INT.udn
+++ b/Engine/Documentation/Source/Engine/Animation/AnimMontage/AnimMontage.INT.udn
@@ -11,34 +11,28 @@ Related: Engine/Content/Tools/MayaRiggingTool
 Related: Gameplay/Input
 version: 4.9
 
-
-[TOC(start:2 end:3)]
-
-
-
 [EXCERPT:Intro]
-**Animation Montages** (or just Montages for short) are a multipurpose tool that allows for a wide variety of animation effects, primarily related to exposing animation controls within code or Blueprint. It can also be used to create a wide variety of animation effects including intelligent loops of animation, logic-based animation switching, root motion handling, and much more.
+**Animation Montages** (or just Montages for short) are a multipurpose tool that allows for a wide variety of animation effects, primarily related to exposing animation controls within code or [Blueprint Visual Scripting](Engine/Blueprints). It can also be used to create a wide variety of animation effects including intelligent loops of animation, logic-based animation switching and much more.
 [/EXCERPT:Intro]
 
-AnimMontages are animation assets that can be created and exist within the **Content Browser**. These assets can then be plugged directly into the AnimGraph in run-time and you can modify any state of it. For example, you can jump to different sections or you can re-link different sections. This is mostly for code driven animations or one-off animations such as melee attacks, allowing you to control triggers, stop when you want, or change state (looping or non-looping) between (see the image below). 
+[TOC(start:2 end:2)]
+
+AnimMontages are animation assets that can be created and exist within the **Content Browser**. These assets can then be plugged directly into the AnimGraph in run-time and you can modify any state of it. For example, you can jump to different sections or you can re-link different sections. This is mostly for code driven animations or one-off animations such as melee attacks, allowing you to control triggers, stop when you want, or change between states (looping or non-looping). 
 
 ![](Montage_Screen2.png)
 
-The image above is a melee attack with 3 sections [start, loop, and end]. When the player **Left-clicks**, it triggers the start section by default when you ask to play the montage. Now the middle section is set to loop. When the start is done, it will transition to loop, and will loop there forever. If the player lets the mouse button go, it will stop, but you do not want to stop right away since the animation will pop in the middle of the loop. You would like to relink the loop to the end, where it will then transition out to the end and finish the animation. 
+The image above is a melee attack with 3 sections [start, loop, and end]. When the player **Left-clicks**, it triggers the _Start_ section by default when you ask to play the montage. The middle section called _Loop_ (which is set to loop), immediately follows the _Start_ section and will loop there forever. If the player lets the mouse button go, it will stop, but you do not want to stop right away since the animation will pop in the middle of the loop. You would like to relink the _Loop_ to the _End_ section, where it will then transition out to the end and finish the animation. 
 
 Some additional uses for Montages include:
 * The ability to play an animation from within an AnimBlueprint's **EventGraph**.
 * Chaining together a complex sequence of animations that you want to think of as a single animation.
 * Looping only a specific portion of animation or animations based on code or Blueprint script.
 * Handling event-based switching of multiple animations based on code or Blueprint script.
-* Proper handling of [Root Motion](Engine/Animation/RootMotion) for your characters.
 * The ability to assign complex animation sequences to named slots that can be switched between in code or Blueprints.
 * Precise switching between various AnimSequences based on code or Blueprint script.
 
-So, as you can see, there are a lot of different things you can do with Montages. This can easily make Montages seem overwhelming as a system. However, if you really boil it down, the key point of Montages is to expose animation controls to code or Blueprints. Nearly everything else the system can do follows from that single point.
-
-[REGION:tip]
-An example of an Animation Montage applied to a character can also be seen on the [Animation Content Examples](Resources\ContentExamples\Animation) page under section 1.5.
+[REGION:note]
+An Animation Montage applied to a character can also be seen on the [Animation Content Examples](Resources\ContentExamples\Animation) page under section 1.5.
 [/REGION]
 
 
@@ -50,28 +44,29 @@ There are a few details to keep in mind regarding AnimMontages:
 * [INCLUDE:#MontageReplication]
 
 
-
 ## Montage Properties
 
 The following is a breakdown of the AnimMontage asset properties. These are available in **Persona** when looking at a Montage, and can also be accessed by **Right-clicking** on a Montage within the **Content Browser** and choosing Properties from the context menu.
 
 ![](MontageProperties.png)
 
-
-
 |Montage Properties||
 |---|---|
 |[REGION:tablesection]Montage[/REGION]||
-|**Blend In Time**|This is an amount of time at the beginning of Montage playback during which the character will blend in from its current pose.|
-|**Blend Out Time**|This is an amount of time at the end of Montage playback during which the character will blend back to its original pose.|
-|[REGION:tablesection]Root Motion[/REGION]||
-|**Enable Root Motion Translation**|Enables handling of root motion translation, cancelling out any translation applied to the root so that it can then be applied back to the character's collision capsule. Please see [](Engine/Animation/RootMotion) for more details.|
-|**Enable Root Motion Rotation**|Enables handling of root motion rotation, cancelling out any rotation applied to the root so that it can then be applied back to the character's collision capsule. Please see [](Engine/Animation/RootMotion) for more details.|
+|**Blend In Time**| The amount of time at the beginning of Montage playback during which the character will blend in from its current pose.|
+|**Blend Out Time**|The amount of time at the end of Montage playback during which the character will blend back to its original pose.|
+|**Blend Out Trigger Time**| Time from the Sequence End to trigger blend out. < 0 uses BlendOutTime, so BlendOut finishes as the Montage Ends. > = 0 uses SequenceEnd-BlendOutTriggerTime to trigger blend out. |
 |[REGION:tablesection]Additive Settings[/REGION]||
 |**Preview Base Pose**|Sets a base preview pose used for additive Blend Spaces.|
 |[REGION:tablesection]Animation[/REGION]||
 |**Rate Scale**|A multiplier value for how fast the Montage will play back. Default is 1.0.|
 |**Skeleton**|Contains the skeleton associated with this Montage. Cannot be changed in the editor.|
+|[REGION:tablesection]Meta Data[/REGION]||
+|**Meta Data**| This is Meta Data that can be saved with the asset. The meta data is a Blueprintable class derived from the Anim Meta Data class. This allows you to add custom meta data to animation assets (Anim Sequence, Anim Montage, Anim Composite, and Blendspace are all supported). You can query the data from an animation asset in C++ by calling the `GetMetaData()` method, or by using `GetSectionMetaData()` on Anim Montages. |
+|[REGION:tablesection]Thumbnail[/REGION]||
+|**Orbit Pitch**| The pitch of the orbit camera around the asset. |
+|**Orbit Yaw**| The yaw of the orbit camera around the asset. |
+|**Orbit Zoom**| The offset from the bounds sphere distance from the asset.|
 
 ## Montage UI
 
@@ -96,11 +91,10 @@ The Montage Area breaks down like so:
 [/REGION]
 
 1. **Section Track** - Shows any [Sections](#Sections) that have been defined for this Montage. Sections can be dragged to different positions along the timeline with the left mouse button.
-1. **Slot Track** - Shows the current Slot, along with the Slot name on the right. This Slot may be filled with as many animations as you desire; they will be played back in order. Notice that multiple animations will have alternating positions in the Slot Track - top, then bottom, then top again, and so on. This is to help you differentiate between different animations.
-1. **Branch Point Track** - Shows any Branch Points that have been defined for this Montage. Branch Points can be dragged to different positions along the timeline with the left mouse button.
+1. **Slot Track** - Shows the current [Slot](#Slots), along with the Slot name on the right. This Slot may be filled with as many animations as you desire; they will be played back in order. Notice that multiple animations will have alternating positions in the Slot Track - top, then bottom, then top again, and so on. This is to help you differentiate between different animations.
 
 [REGION:note]
-You may have as many Slot Tracks as you like within a single Montage, each with their own name and containing their own unique animations. However, you may only have one Section Track and one Branch Point track for each Montage.
+You may have as many Slot Tracks as you like within a single Montage, each with their own name and containing their own unique animations. However, you may only have one Section Track for each Montage.
 [/REGION]
 
 
@@ -121,7 +115,7 @@ Within a Montage, a **Slot** is just a single track that can hold any number of
 It is an important point to keep in mind that although much of your Montage control will take place in the Animation Blueprint's **Event Graph**, Slots are actually handled within the **Anim Graph**. This is done by way of the Slot node, which simply takes in the name for a Slot. By positioning this node at a strategic point along your AnimGraph's execution, you can have multiple Montages that utilize the same Slot name.
 [/EXCERPT:Slots]
 
-
+<!---------------- Branch moved to Notify page.
 ### Branch Points
 
 Branch Points allow you to create events that coincide with animation playback. These Branch Point Events can be used in code or Blueprint to cause other things to happen, but specifically, Branch Points are useful for switching to other animation Sections within a Montage. 
@@ -134,6 +128,8 @@ High precision is important when you need to jump to a specific animation at a v
 Due to their synchronous nature and the precision resulting from it, Branch Points are more performance expensive than Notifies. You should only use them when an event must be fired at a precise moment along the animation timeline, such as jumping directly to another animation that matches up frame-to-frame. If being off by a frame (or some percentage of one) is not important, you should use Notifies instead.
 [/REGION]
 
+-------------------->
+
 ## Sections Area
 
 The Sections area is where you can establish relationships between the Sections you define in the Montage Area. For instance, you may want a certain Section of animation (or group of Sections) to play in a specific sequence, or even loop. 
@@ -147,7 +143,6 @@ The Sections area is where you can establish relationships between the Sections
 1. **Section Buttons** - In this area, you will see a button for each of the Sections you define in the Montage area. By choosing an existing Section and then clicking the one of these buttons, you associate the Section for that button with the selected track. For instance, in the image above, we have associated Swing2 with Swing1. The order actually goes Swing1, Swing2, and then Swing1 again, which causes a loop. See the [Looping](#Looping) section below for details.
 1. **Section Association Tracks** - This is where you can visualize and preview the relationships between animation Sections. By clicking on the **Preview** buttons, you can see the result of each individual track, or click the **Preview All Sections** button and see all Sections play back one after the other.
 
-
 ### Looping
 
 Sections can be set up to loop indefinitely, which is extremely useful for any action that you need to repeat. By associating the same Section more than once in a Section Association Track, you cause that association to run in a loop. The section will turn blue to show this. As an example, consider an animation in which a character is reloading a shotgun, one shell at a time. You could take just the section in which the character inserts a shell, and loop it. Then, by using Notifies, you could create Notify Events in Blueprint that increment the ammo count each time the animation plays through. Once that count reaches a set number (full ammo), you could then switch to an animation of the character closing the receiver and returning to idle.
@@ -158,7 +153,7 @@ The Element Timing area takes information from the Montage area and the Notifies
 
 ![](ElementTiming1.png)
 
-Each node in the track is given a number denoting the trigger order of that object in the entire montage, and are coloured to match what they describe. The default colours are:
+Each node in the track is given a number denoting the trigger order of that object in the entire montage, and are colored to match what they describe. The default colors are:
 * Red - Notifies, Notify states and notify state ending markers
 * Light Blue - Branching point notifies
 * Green - Montage sections
@@ -191,6 +186,7 @@ For more information, see [](Engine/Animation/Sequences\Notifies).
 
 For more information, see [](Engine/Animation/Sequences\Curves).
 
+<!-------------------- Example removed. This should be a HowTo.
 
 ## Montage Practical Example
 
@@ -298,6 +294,7 @@ Dragging off the **IsAttacking** node which extends from the **Cast To** node, w
 
 And that is it! If we compile, we will now see that the character continues to swing as long as the mouse button is down, and performs an intelligent end animation when the mouse button is released!
 
+
 [OBJECT:EmbeddedVideo]
 [PARAMLITERAL:width]
 640
@@ -310,13 +307,12 @@ And that is it! If we compile, we will now see that the character continues to s
 [/PARAMLITERAL]
 [/OBJECT]
 
-<!--------------------
+-------------------->
+
 
 
-[EXCERPT:MontageReplication]At this time, root motion is replicated across the server. However the AnimMontage is not. That means that users utilizing root motion will need to make sure they are replicating the state of the AnimMontage across the network. (Position, play rate, etc.).[/EXCERPT:MontageReplication]
 
 
--------------------->
 
 
 
diff --git a/Engine/Documentation/Source/Engine/Animation/Sequences/Notifies/AnimNotifies.INT.udn b/Engine/Documentation/Source/Engine/Animation/Sequences/Notifies/AnimNotifies.INT.udn
--- a/Engine/Documentation/Source/Engine/Animation/Sequences/Notifies/AnimNotifies.INT.udn
+++ b/Engine/Documentation/Source/Engine/Animation/Sequences/Notifies/AnimNotifies.INT.udn
@@ -3,6 +3,10 @@ Title:Animation Notifications (Notifies)
 Crumbs: %ROOT%, Engine, Engine/Animation, Engine/Animation/Sequences
 Description:Notifies are a system for setting up and receiving events in Animation Sequences to perform external actions.
 version: 4.9
+Related: Engine/Animation/Sequences/Editor
+Related: Engine/Animation/AnimMontage
+Related: Engine/Animation
+Related: Engine/Animation/AnimHowTo
 
 [VAR:Topic]
 [OBJECT:Topic]
@@ -25,40 +29,38 @@ version: 4.9
 [/VAR]
 
 
-[TOC (start:2 end:2)] 
-
-
 [EXCERPT:Intro]
 **Animation Notifications** (**AnimNotifies** or just **Notifies**) provide a way for animation programmers to setup events to occur at specific points during an **Animation Sequence**. Notifies are commonly used to add effects like footstep sounds during a walk or run animation or spawning a particle system during animation. However, they can have any number of different uses, as the system can be extended with custom notification types to cater to the needs of any type of game.
 [/EXCERPT:Intro]
 
-Creating and editing Notifies is performed within the **Notify Editor** in **Persona**. For more information, refer to the [Editing Animation Sequences](Engine/Animation/Sequences/Editor) documentation. 
-
-## Animation Notify Types
+[TOC (start:2 end:2)] 
 
-Different types of Notifies cause different events to be triggered. Camera effects, particle effects, sounds, etc. can all be triggered at any point during an animation. 
+## Notify Types
 
-Below you will find the currently available Notifies and how they are used. 
+Different types of Notifies cause different events to be triggered. Camera effects, particle effects, sounds, etc. can all be triggered at any point during an animation through the use of a Notify. Adding a Notify can be done by **Right-clicking** on a **Notify Track** at the point in the animation you want the Notify to occur, and then choosing one of the type of Notify to add.
 
+[REGION:note]
+For more information on adding/editing Notifies, see the [](Engine/Animation/Sequences/Editor) page. 
+[/REGION]
 
 ### Play Particle Effect
 
-The **PlayParticleEffect** notify spawns and plays a particle system at a given socket location and at a specific point along an animation's timeline. This is useful for such effects as causing a cloud of dust to appear around a player's feet when they move, fire to spew from an exhaust pipe when the player uses a Nitro-Boost power-up, or other such effects.
+The **PlayParticleEffect** Notify spawns and plays a particle system at a given socket location and at a specific point along an animation's timeline. This is useful for such effects as causing a cloud of dust to appear around a player's feet when they move, fire to spew from an exhaust pipe when the player uses a Nitro-Boost power-up, or other such effects.
 
 ![](PlayParticleEffectNotify.png)
 
-A PlayParticleEffect notify does not use the Animation Blueprint to control behavior. Instead, it uses a few simple properties defined in the **Details** panel to determine what Particle System is used and from what Socket the particles will originate.
+A PlayParticleEffect Notify does not use the Animation Blueprint to control behavior. Instead, it uses a few simple properties defined in the **Details** panel to determine what Particle System is used and from what Socket the particles will originate.
 
 ![](PlayParticleEffectDetails.png)
 
 | **Property** | **Description** |
 | --- | --- |
 | **Socket Name** | This contains the Skeletal Mesh socket that will be used as the spawn point for the particle system. |
-| **PSTemplate** | This holds the particle system that will be spawned at the notify's set point in time. |
+| **PSTemplate** | This holds the particle system that will be spawned at the Notifies set point in time. |
 | **Attached** | Determines if the effect should attach to the bone/socket. |
 | **Location Offset** | Location offset from the socket. |
 | **Rotation Offset** | Rotation offset from the socket. |
-| **Notify Color** | Changes the color of the notify on the notify bar. |
+| **Notify Color** | Changes the color of the Notify on the Notify bar. |
 
 Once set up, the associated Particle System will play, using the Socket location as the spawn point.
 
@@ -67,11 +69,11 @@ Once set up, the associated Particle System will play, using the Socket location
 
 ### Play Sound
 
-The **PlaySound** notify provides a simple means to get your sound effects synced with animation. Much like Play Particle Effect, this notify has a few simple properties that can be set in the **Details** panel, which allow for fast and easy sound effect setup. This is useful for such things as mechanical whirrs when a robot moves, the rustle of cloth during character motion, and similar types of effects.
+The **PlaySound** Notify provides a simple means to get your sound effects synced with animation. Much like Play Particle Effect, this Notify has a few simple properties that can be set in the **Details** panel, which allow for fast and easy sound effect setup. This is useful for such things as mechanical whirrs when a robot moves, the rustle of cloth during character motion, and similar types of effects.
 
 ![](PlaySound.png)
 
-The **Play Sound** notify uses a few simple properties defined in the **Details** panel to determine what sound effect is used and from what **Socket** or **Bone** the sound will originate.
+The **Play Sound** Notify uses a few simple properties defined in the **Details** panel to determine what sound effect is used and from what **Socket** or **Bone** the sound will originate.
 
 ![](PlaySoundDetails.png)
 
@@ -81,69 +83,54 @@ The **Play Sound** notify uses a few simple properties defined in the **Details*
 | **Volume Multiplier** | This scaling multiplier is applied to the volume of the sound effect. |
 | **Pitch Multiplier** | This scaling multiplier is applied to the pitch of the sound effect. |
 | **Follow** | If checked, the sound effect source will follow the Static Mesh as it moves. If unchecked, the sound remains behind at the location it was spawned. |
-| **Notify Color** | Changes the color of the notify on the notify bar. |
+| **Notify Color** | Changes the color of the Notify on the Notify bar. |
 
-### Custom
+### Custom Notifies
 
-The **Custom** notify serves as a catch-all for anything an artist may want to do within a Animation Blueprint at a certain point in an animation. Custom notifies are created by **Right-clicking** on a notify track and choosing **New Notify** from the context menu.
+The **Custom** Notify serves as a catch-all for anything an artist may want to do within a Animation Blueprint at a certain point in an animation. Custom Notifies are created by **Right-clicking** on a Notify track and choosing **New Notify** from the context menu.
 
 ![](MySpecialNotify.png)
 
-As soon as you do this, you will be required to supply a name for the new notify.
-
-![](NotifyName.png)
-
-After the name is typed in, you will see the new notify appear in the timeline.
-
-![](MyNotifyIn.png)
-
-Custom notifies will not have any properties to edit in the **Details** panel. Instead, they simply allow you to create a notify event of the same name in the **Event Graph** or any of the **Transitions** in a state machine of the Skeletal Mesh's Animation Blueprint.
-
-Switch to the **Graph** edit mode
-<!--![](SwitchToAnimBP.png) Commented out picture as it didn't seem useful--> 
-
-**Right-click** anywhere in the **Event Graph** and there is a section for adding Animation Notifies events.
+Custom Notifies allow you to create a **Notify Event** of the same name in the Event Graph or Transition Graphs in the [State Machine](Engine/Animation/StateMachines) for a Skeletal Mesh's [Animation Blueprint](Engine\Animation\AnimBlueprints). To add your Custom Notify, **Right-click** anywhere in your Event or Transition Graphs and there is a section for adding Animation Notifies events.
 
 ![](rightClickAddNotify.png)
 
-![](MySpecialNotifyBlueprint.png)
-
-## Animation Notify States 
-
-**AnimNotifyStates** (**Notify States**) work much like the standard notifies above.  They have 3 distinct events: a start, the tick, and an end. They start straightforward, firing at the moments the notify begins and ends, and the **Event Graphs** in them fire off when their time comes up in the animation. The tick fires off every animation update until the end event is hit. The major difference between the standard Notifies is that AnimNotifyStates are self-contained Blueprints.
-
-To use an AnimNotifyState:
+After adding the Custom Notify, an Event Node will be added to the graph which will be executed when the Notify is called. 
 
-1. Create a new AnimNotifyState asset in the **Content Browser** by creating a new **Blueprint** that inherits from `AnimNotifyState`:
-
-	![](newBP.png)
-
-	![](AnimNotifyStateCreate.png)
+![](MySpecialNotifyBlueprint.png)
 
-1. On the Animation Sequence or Animation Montage, you can **Right-click** in the notify area and now add your new AnimNoitfyState.
+## Notify States 
 
-	![](addNotifyState.png)
+**AnimNotifyStates** (**Notify States**) work much like the standard Notifies above. They have 3 distinct events: a start, the tick, and an end. They start straightforward, firing at the moment the Notify begins and ends, and the **Event Graphs** in them fire off when their time comes up in the animation. The tick fires off every animation update until the end event is hit. The major difference between normal Notifies and Notify States is that Notify States are self-contained [Blueprints](Engine/Blueprints).
 
-1. You can set the length of the AnimNotifyState by **Right-clicking** on it and choosing **Set AnimNotifyState Duration**.
+Creating Notify States (which are the same as [Native Notify States](#nativenotifystates) but are custom Blueprints that you can provide any logic you need to occurring during an animation) is done the same as adding a regular Notify by **Right-clicking** in the **Notify Track** of an animation and selecting **Add Notify State...** then selecting your custom Notify.
 
-	![](AnimNotifyStateSetLength.png)
+[REGION:note]
+To use a custom Notify, you must first [create a Blueprint Class](Engine/Blueprints/UserGuide/Types/ClassBlueprint/Creation) of the AnimNotifyState class. 
+[/REGION]
 
-	![](duration.png)
+![](addNotifyState.png)
 
-1. Finally from the **Content Browser**, **Double Left-click** your AnimNotifyState to bring up its Blueprint so you can start defining what if does.
-	
+An AnimNotifyState Blueprint Class called **MyAnimNotifyState** (above) was created so that we may select it as the Notify State to use in our animation. Once you have created your Anim Notify State Blueprint, you can open it and begin implementing the functions that determine how your Notify is used (below).
 
 ![](animNotifyStateImplementFunction.png)
 
-
-At this point you can start implementing the 3 functions of the **AnimNotifyState**. It is not necessary to implement every function, but at the bare minimum you should implement **Recieved Notify Tick** or both **Recieved Notify Start** and **Recieved Notify End**. If implementing just **Recieved Notify Start** or **Recieved Notify End**, you could simply use a standard Animation Notify.
+It is not necessary to implement every function, but at the bare minimum you should implement **Received Notify Tick** or both **Received Notify Start** and **Received Notify End**. If implementing just **Received Notify Start** or **Received Notify End**, you could simply use a standard Animation Notify.
 
 ![](Recieved_NotifyStart.png)
 
+_The Notify Start Override._
+
 ![](Recieved_NotifyEnd.png)
 
+_The Notify End Override._
+
 ![](Recieved_NotifyTick.png)
 
+_The Notify Tick Override._
+
+Below is a description of how each pin type is used in the override functions:
+
 | **Pin** | **Description** |
 | --- | ----------- |
 | **Mesh Comp** | This is the Skeletal Mesh Component of the Animation Sequence or Montage that is triggering this AnimNotifyState. |
@@ -176,20 +163,21 @@ When choosing to add a Notify State, you can also assign one of two **Native Not
 
 ### Timed Particle Effects
 
-The Timed Particle Effect is very similar to the **Play Particle Effect** notify in that they both allow you to play a Particle Effect at any given point in an animation, the difference however is that with a Timed Particle Effect, you are given the ability to define the length of the effect inside the notify itself. 
+The Timed Particle Effect is very similar to the **Play Particle Effect** Notify in that they both allow you to play a Particle Effect at any given point in an animation, the difference however is that with a Timed Particle Effect, you are given the ability to define the length of the effect inside the Notify itself. 
 
 ![](TimedParticleEffect.png)
 
-In the image above, both a Play Particle Effect and Timed Particle Effect are assigned to an animation. Inside the yellow box is the Timed Particle Effect and you will notice that this notify has two handle bars; one at the start and one at the end of the event. By **Left-Clicking** on the handle bar at the end (denoted by the green arrow), you can drag to increase or decrease the length of the effect.
+In the image above, a Timed Particle Effect is assigned to an animation. By **Left-Clicking** on the handle bar at the end (or at the start), you can drag to increase or decrease the length of the effect.
 
 ### Animation Trails
 Animation Trails (or **AnimTrail** for short) are emitters that allow you to leave trails behind moving objects in a scene. 
 They do this by sampling the positions of two sockets on a Skeletal Mesh each tick and stitching triangles between the sockets, down the length of the trail. 
   
-To use an AnimTrail emitter you must do the following; 
-1.       Create a Particle System that contains an emitter with AnimTrail type data. 
-2.       In the **Animation Editor**, add an AnimNotifyState of type **Trail**. 
+To use an AnimTrail emitter you must:
+1. Create a [Particle System](Engine/Rendering/ParticleSystems) that contains an emitter with AnimTrail type data. 
+2. Create and add a Notify State of the **Trail** type to an animation.
 
+<!-- To be included as part of a HowTo.
 Setting Up the Particle System
 
 Create a Particle System as usual and add a new emitter with type data AnimTrail to it. 
@@ -220,24 +208,25 @@ However, the more tessellated a trail is, the more expensive it is.
 | **Distance Tessellation Step Size** | The (approx.) distance in world units that the trail must be to warrant another tessellation point.   |
 | **Tangent Tessellation Step Size** | The (approx.) angle difference in degrees that the tangent must change from one particle to another to warrant another tessellation point. This allows you to ensure that any sharp turns in the trail are nice and smooth.  |
 | **Width Tessellation Step Size** | The (approx.) world unit change in the width of the trail that warrants an additional tessellation point.|
- 
+
 Setting Up the AnimNotify
 
 Now that we have a particle system, we must use it with an AnimNotify. 
-Open the **Animation Editor**, add or use an existing **Timeline** under the **Notifies** section and **Right-click** on it to bring up the menu. The Trails notify can be found under the **Add Notify State** section. 
+Open the **Animation Editor**, add or use an existing **Timeline** under the **Notifies** section and **Right-click** on it to bring up the menu. The Trails Notify can be found under the **Add Notify State** section. 
 
 ![](SetUpTrail.png)
+-->
 
-Below is the **Details** panel and **Properties** for the Trail notify. 
+Once you have added an AnimTrail, you can adjust its properties from the **Details** panel: 
 
 ![](TrailDetails.png)
 
 | **Property** | **Description** |
 | --- | --- |
+| **PSTemplate** | This is the particle system containing the anim trail that was described above. |
 | **First Socket Name** | This is the name of the first Bone/Socket defining the trail. |
 | **Second Socket Name** | This is the name of the second Bone/Socket defining the trail. |
-| **PSTemplate** | This is the particle system containing the anim trail that was described above. |
-  
+
 The above properties are required but the below are optional. They allow you to alter the width of the trail during the animation. You could also attach the trail to sockets that move as you need but this may not be ideal. 
 This method is based on providing a width scale value. A width scale of 1.0 will keep the trail the same width, 0.0 will cause a trail of zero width. 
 
@@ -259,6 +248,7 @@ There are also some debug rendering options available:
 | **Render Tangents** | This toggles the debug rendering of the curve tangents. |
 | **Render Tessellation** | This toggles debug rendering of all vertices on the curve. This shows how the curve is tessellated between the spawn points. |
 
+<!-- To be included with AnimTrail HowTo.
 ### Controlling Animation Trails in Blueprint
 While the above method is the standard approach to controlling AnimTrails, you can also control them from inside a Blueprint without the need of an Anim Notify. 
 The only restriction with this is that the Particle System Component must be a child of the Skeletal Mesh Component from which you are sampling socket locations. 
@@ -271,4 +261,34 @@ The only restriction with this is that the Particle System Component must be a c
  ![](TrailBlueprint3.png)
 1. You can also alter the trail's properties while it is running (the example below alters width each tick using the game time).  
  ![](TrailBlueprint4.png)
+  -->
+
+## Additional Notify Properties
+In addition to the properties that can be modified in the **Anim Notify** section of the **Details** panel for each type of Notify (when applicable), the properties below can also be used to define how your Notify is used.
+
+![](AdditionalProperties.png)
+
+|---|---|
+|[REGION:tablesection]Trigger Settings[/REGION]||
+|**Notify Trigger Chance**| Defines the chance of triggering the Notify (0=No Chance, 1=Always Triggers).|
+|**Trigger on Dedicated Server**| If disabled, this Notify will be skipped on Dedicated Servers.|
+|**Notify Filter Type**| Defines a method for filtering notifies (stopping them from triggering). Can be set to filter by the meshes current LOD.|
+|[REGION:tablesection]Event[/REGION]||
+|**Trigger Weight Threshold**| Determines how much much weight the animation sequence needs to have to trigger the Notify. |
+|**Montage Tick Type**| The method to determine how the Notify is handled, Queued or Branching Point (See [Montage Tick Type](#montageticktype) below). |
+|[REGION:tablesection]Anim Link[/REGION]||
+|**Slot**| The slot index currently being used within the LinkedMontage. |
+|**Link Method**| The method used to calculate times, Absolute (stays at a specific time without moving), Relative (moves with its segment, but not when the segment changes size) or Proportional (moves with its segment and will stay at a certain proportion through the segment). |
+
+
+### Montage Tick Type
+
+The key difference between Queued and Branching Point settings for the Montage Tick Type is that Queued Notifies is asynchronous, while Branching Points are synchronous. What this means to the end user is that Branching Points come with a much higher degree of precision for where they will take place along the animation timeline.
+
+High precision is important when you need to jump to a specific animation at a very precise moment in time. While you could use a Queued Notify to do the same job, the asynchronous nature of Queued Notifies means that the event could be fired at the incorrect animation frame, which can lead to hitches and jumps in your motion.
+
+[REGION:warning]
+Due to their synchronous nature and the precision resulting from it, Branching Points are more performance expensive than Queued. You should only use them when an event must be fired at a precise moment along the animation timeline, such as jumping directly to another animation that matches up frame-to-frame. If being off by a frame (or some percentage of one) is not important, you should use Queued instead.
+[/REGION]
+
 
