#UE4doc #docs - Addresses UEDOC-2533.

diff --git a/Engine/Documentation/Source/Engine/Rendering/LightingAndShadows/LightFunctions/LightFunctions.INT.udn b/Engine/Documentation/Source/Engine/Rendering/LightingAndShadows/LightFunctions/LightFunctions.INT.udn
--- a/Engine/Documentation/Source/Engine/Rendering/LightingAndShadows/LightFunctions/LightFunctions.INT.udn
+++ b/Engine/Documentation/Source/Engine/Rendering/LightingAndShadows/LightFunctions/LightFunctions.INT.udn
@@ -52,6 +52,10 @@ Version: 4.9
 [EXCERPT:LightFunctions]
 **Light Functions** are - in essence - a Material that can be applied to filter a light's intensity.  
 
+[Region:note]
+It should be noted, however, that you cannot change the color of lights using a Light Function, only the Light Color setting can do that. Also, Static Lights cannot not use Light Functions, as this feature is only available for non-lightmapped lights, like Stationary and Movable lights.
+[/Region]
+
 It should be noted however that you cannot change the color of lights using a Light Function, only the Light Color setting can do that. 
 [/EXCERPT:LightFunctions]
 
