#UE4doc #docs - Added engine version field. Replaced image table with comparison slider.

diff --git a/Engine/Documentation/Source/Resources/ContentExamples/PostProcessing/1_13/1_13.INT.udn b/Engine/Documentation/Source/Resources/ContentExamples/PostProcessing/1_13/1_13.INT.udn
--- a/Engine/Documentation/Source/Resources/ContentExamples/PostProcessing/1_13/1_13.INT.udn
+++ b/Engine/Documentation/Source/Resources/ContentExamples/PostProcessing/1_13/1_13.INT.udn
@@ -5,6 +5,7 @@ Description:An overview of the Post Processing example level, example 1.13: Scre
 Related: Engine/Rendering/Overview
 Related: Engine/Rendering/PostProcessEffects
 Related: Engine/Rendering/PostProcessEffects/PostProcessMaterials
+Version: 4.9
 
 [VAR:Steps]
 [OBJECT:Navigation]
@@ -63,11 +64,24 @@ Related: Engine/Rendering/PostProcessEffects/PostProcessMaterials
 [/OBJECT]
 [/VAR]
 
+<div align="center">
+[OBJECT:ComparisonSlider]
+	[PARAM:before]
+	![Screen Percentage - 100](1_13_01.png)(w:600)
+	[/PARAM]
+	[PARAM:after]
+	![Screen Percentage - 25](1_13_02.png)(w:600)
+	[/PARAM]
+[/OBJECT]
+</div>
+
+<!---
 <div align="center">
 [REGION:lightbox]
 [![](1_13_01.png)(w:440)](1_13_01.png) [![](1_13_02.png)(w:440)](1_13_02.png)
 [/REGION:lightbox]
 </div>
+--->
 
 **Screen Percentage** renders the scene at a lower resolution and then scales it up. The value represents a percentage of the overall resolution of the scene.
 
