#UE4doc #docs:
Updated engine version field

diff --git a/Engine/Documentation/Source/Programming/Assets/AsyncLoading/AsyncLoading.INT.udn b/Engine/Documentation/Source/Programming/Assets/AsyncLoading/AsyncLoading.INT.udn
--- a/Engine/Documentation/Source/Programming/Assets/AsyncLoading/AsyncLoading.INT.udn
+++ b/Engine/Documentation/Source/Programming/Assets/AsyncLoading/AsyncLoading.INT.udn
@@ -2,7 +2,7 @@ Availability:Public
 Title:Asynchronous Asset Loading
 Crumbs:%ROOT%, Programming
 Description:Methods for loading and unloading assets during runtime.
-Version: 4.5
+Version: 4.9
 
 [TOC(start:2)]
 
