#UE4doc #docs:
Updated engine version field

diff --git a/Engine/Documentation/Source/Programming/Assets/ReferencingAssets/ReferencingAssets.INT.udn b/Engine/Documentation/Source/Programming/Assets/ReferencingAssets/ReferencingAssets.INT.udn
--- a/Engine/Documentation/Source/Programming/Assets/ReferencingAssets/ReferencingAssets.INT.udn
+++ b/Engine/Documentation/Source/Programming/Assets/ReferencingAssets/ReferencingAssets.INT.udn
@@ -2,7 +2,7 @@ Title: Referencing Assets
 Crumbs:
 Description:
 Availability: Public
-Version: 4.7
+Version: 4.9
 
 Unreal Engine 4 provides a number of mechanisms to control how an asset is referenced and by extension loaded into memory. You can think of references in two ways: a hard reference where object A refers to object B and causes object B to be loaded when object A is loaded; and a soft reference where object A refers to object B via an indirect mechanism such as the string form of the path to the object. The first two sections following cover hard references with the remaining sections exploring soft references. 
 
