#UE4doc #docs:
Updated engine version field
Changed font settings for multiple terms

diff --git a/Engine/Documentation/Source/Programming/Assets/Registry/AssetRegistry.INT.udn b/Engine/Documentation/Source/Programming/Assets/Registry/AssetRegistry.INT.udn
--- a/Engine/Documentation/Source/Programming/Assets/Registry/AssetRegistry.INT.udn
+++ b/Engine/Documentation/Source/Programming/Assets/Registry/AssetRegistry.INT.udn
@@ -2,7 +2,7 @@ Availability:Public
 Title:Asset Registry
 Crumbs: %ROOT%, Engine, Engine/UI, Programming
 Description:How assets are discovered by the editor and how to make it know more about asset types before they are loaded.
-Version: 4.5
+Version: 4.9
 
 The **Asset Registry** is an editor subsystem which gathers information about unloaded assets asynchronously as the editor loads.
 This information is stored in memory so the editor can create lists of assets without loading them.
@@ -102,10 +102,10 @@ For example (from `UTexture`):
 	UPROPERTY(Category=Texture, AssetRegistrySearchable)
 	TEnumAsByte<enum TextureFilter> Filter;
 
-Once this flag was added to the 'Filter' property of `UTexture`, all `UTextures` that were saved thereafter now have an entry in their `FAssetData`'s `TagsAndValues` map whose key is `"Filter"` and whose value is the string representation of the enum value, such as `"TF_Linear"`.
+Once this flag was added to the Filter property of `UTexture`, all `UTextures` that were saved thereafter now have an entry in their `FAssetData`'s `TagsAndValues` map whose key is `Filter` and whose value is the string representation of the enum value, such as `"TF_Linear"`.
 
 [REGION:note]
-	This **requires assets to be resaved** before their properties will be discovered by the Asset Registry.
+	This requires assets to be resaved before their properties will be discovered by the Asset Registry.
 [/REGION]
 
 If you want the Asset Registry to be able to search for information that is not directly a UProperty, 
