#UE4doc #docs:
Added engine version field
Changed font settings for multiple terms

diff --git a/Engine/Documentation/Source/Programming/Automation/Automation.INT.udn b/Engine/Documentation/Source/Programming/Automation/Automation.INT.udn
--- a/Engine/Documentation/Source/Programming/Automation/Automation.INT.udn
+++ b/Engine/Documentation/Source/Programming/Automation/Automation.INT.udn
@@ -3,6 +3,7 @@ Title:Automation System Overview
 Crumbs: %ROOT%, Programming
 Description:Overview of the automation system used for unit testing, feature testing, and content stress testing.
 Navigation:topic
+version: 4.9
 
 [VAR:Topic]
 [OBJECT:Topic]
@@ -45,27 +46,27 @@ Navigation:topic
 [/VAR]
 
 [EXCERPT:Intro]
-The **Automation System** provides the ability to perform _Unit Testing_, _Feature Testing_, and _Content 
-Stress Testing_ using the power of the **Unreal Message Bus** in order 
+The **Automation System** provides the ability to perform Unit Testing, Feature Testing, and Content 
+Stress Testing using the power of the **Unreal Message Bus** in order 
 to increase stability.
 
-The Automation system works by performing one or more **automation tests**. Automation tests can be broken down 
+The Automation system works by performing one or more automation tests. Automation tests can be broken down 
 into different categories depending on their purpose or function:
 [/EXCERPT:Intro]
 
-$ Unit Test : API level verification tests.  See `TimespanTest.cpp` or `DateTimeTest.cpp` for examples.
+$ Unit Test : API level verification tests.  See TimespanTest.cpp or DateTimeTest.cpp for examples.
 $ Feature Tests : System level verification tests such as verify PIE works, verify in-game stats work, 
-and verify changing resolution works.  See `EditorAutomationTests.cpp` or `EngineAutomationTests.cpp` for examples. 
+and verify changing resolution works.  See EditorAutomationTests.cpp or EngineAutomationTests.cpp for examples. 
 $ Smoke Tests : Smoke Tests are just considered a speed promise by the implementer.  They are intended to 
 be fast so they can run _everytime_ the editor, game, or a commandlet starts.  They are also selected by 
 default in the [UI](Programming/Automation/UserGuide#UserInterface).
 	
 	[REGION:warning]
-		All Smoke tests are intended to complete within **1 second**.  Only mark Unit Tests or fast Feature Tests as Smoke Tests.
+		All Smoke tests are intended to complete within 1 second.  Only mark Unit Tests or fast Feature Tests as Smoke Tests.
 	[/REGION]
 
 $ Content Stress Tests : More thorough testing of a particular system to avoid crashes, such as loading all maps 
-or loading and compiling all blueprints.  See `EditorAutomationTests.cpp` or `EngineAutomationTests.cpp` 
+or loading and compiling all blueprints.  See EditorAutomationTests.cpp or EngineAutomationTests.cpp 
 for examples.
 
 $ Screen Shot Comparison :The Screen Shot Comparison tool enables QA to quickly compare screenshots to identify potential 
