#UE4doc #docs:
Updated engine version field
Changed font settings for multiple terms
Fixed a grammar issue

diff --git a/Engine/Documentation/Source/Programming/Automation/ScreenShotComparison/ScreenShotComparisonGuide.INT.udn b/Engine/Documentation/Source/Programming/Automation/ScreenShotComparison/ScreenShotComparisonGuide.INT.udn
--- a/Engine/Documentation/Source/Programming/Automation/ScreenShotComparison/ScreenShotComparisonGuide.INT.udn
+++ b/Engine/Documentation/Source/Programming/Automation/ScreenShotComparison/ScreenShotComparisonGuide.INT.udn
@@ -2,7 +2,7 @@ Availability:Public
 Title:Screen Shot Comparison User Guide
 Crumbs: %ROOT%, Programming, Programming/Automation
 Description:Developers's guide to creating and comparing screen shots
-Version: 4.5
+Version: 4.9
 
 The screenshot browser tool allows QA to quickly and easily compare screenshots. The screenshots are generated via an automation test, and can be viewed in the Unreal Frontend tool. 
 A history of screenshots is maintained, so QA can compare the latest screenshots to a history, and identify any obvious rendering errors with the latest build.
@@ -33,9 +33,9 @@ Run the test by clicking the run button. This will start an automation test that
 
 As the screenshots are being captured, they will be sent to the Unreal Frontend app that requested the test. Then the UI will update automatically when it receives a new screenshot.
 
-$ Select **Screen Comparison** from the Session Frontend's *Window* menu :
+$ Select **Screen Comparison** from the Session Frontend's Window menu :
 
-This will display an UI that update as it receives images.
+This will display an UI that updates as it receives images.
 
 $ UI Navigation :
 
