#UE4doc #docs:
Fixed another grammar issue

diff --git a/Engine/Documentation/Source/Programming/Automation/ScreenShotComparison/ScreenShotComparisonGuide.INT.udn b/Engine/Documentation/Source/Programming/Automation/ScreenShotComparison/ScreenShotComparisonGuide.INT.udn
--- a/Engine/Documentation/Source/Programming/Automation/ScreenShotComparison/ScreenShotComparisonGuide.INT.udn
+++ b/Engine/Documentation/Source/Programming/Automation/ScreenShotComparison/ScreenShotComparisonGuide.INT.udn
@@ -35,7 +35,7 @@ As the screenshots are being captured, they will be sent to the Unreal Frontend
 
 $ Select **Screen Comparison** from the Session Frontend's Window menu :
 
-This will display an UI that updates as it receives images.
+This will display a UI that updates as it receives images.
 
 $ UI Navigation :
 
