#UE4doc #docs:
Changed font settings for multiple terms

diff --git a/Engine/Documentation/Source/Programming/Automation/TechnicalGuide/AutomationTechnicalGuide.INT.udn b/Engine/Documentation/Source/Programming/Automation/TechnicalGuide/AutomationTechnicalGuide.INT.udn
--- a/Engine/Documentation/Source/Programming/Automation/TechnicalGuide/AutomationTechnicalGuide.INT.udn
+++ b/Engine/Documentation/Source/Programming/Automation/TechnicalGuide/AutomationTechnicalGuide.INT.udn
@@ -18,7 +18,7 @@ Blueprints would be good fits for complex automation tests.
 ## Architecture and Execution
 
 All automation tests derive from the **FAutomationTestBase** class which defines generic functionality for 
-performing a set of commands. The important functions of the **FAutomationTestBase** class used when setting 
+performing a set of commands. The important functions of the FAutomationTestBase class used when setting 
 up a new automation test are:
 
 | Member Functions | Description |
@@ -69,7 +69,7 @@ For example, the declaration of a new simple test for the `SetRes` command funct
 
 	IMPLEMENT_SIMPLE_AUTOMATION_TEST( FSetResTest, "Windows.SetResolution", ATF_Game )
 
-Using the _SetRes_ example above, the implementation would be:
+Using the SetRes example above, the implementation would be:
 
 	bool FSetResTest::RunTest(const FString& Parameters)
 	{
@@ -141,7 +141,7 @@ Using the map loading example, the implementations would be:
 ## Latent Commands
 
 Latent Commands allow for sections of an automation test to be run over multiple frames. These are intended to 
-be queued up during _RunTest_ calls.
+be queued up during RunTest calls.
 
 The first step is declaring a Latent Command with the following syntax:
 
@@ -149,9 +149,9 @@ The first step is declaring a Latent Command with the following syntax:
 
 	bool FExecStringLatentCommand::Update()
 
-The Update call should return _true_ when it is done executing and _false_ if it would like to stop execution of the 
+The Update call should return true when it is done executing and false if it would like to stop execution of the 
 automation test and try again next frame. Commands execute in order and execution will not continue if a latent 
-command returns _false_ from _Update_.
+command returns false from Update.
 
 The second step is adding the latent command to the queue for execution. This is wrapped in a macro to avoid boiler 
 plate code as follows:
