#UE4doc #docs:
Added engine version field
Removed Topic List from end of the page and replaced it with "related:" entries
Changed font settings for multiple terms

diff --git a/Engine/Documentation/Source/Programming/Basics/Basics.INT.udn b/Engine/Documentation/Source/Programming/Basics/Basics.INT.udn
--- a/Engine/Documentation/Source/Programming/Basics/Basics.INT.udn
+++ b/Engine/Documentation/Source/Programming/Basics/Basics.INT.udn
@@ -2,6 +2,12 @@ Availability:Public
 Title:Programming Basics
 Crumbs: %ROOT%, Programming
 Description:Information for programmers developing with Unreal Engine.
+related: Programming\UnrealArchitecture\Reference\Classes
+related: Programming\UnrealArchitecture\Reference\Functions
+related: Programming\UnrealArchitecture\Reference\Properties
+related: Programming\UnrealArchitecture\Reference\Structs
+related: Programming\UnrealArchitecture\Reference\Interfaces
+version: 4.9
 
 
 An individual game is defined by a game project that contains all of the code, content, and settings specific to that particular game. Gameplay code is contained within one or more 
@@ -56,29 +62,10 @@ Different types of components can be used to control how Actors move, how they a
 UObject class, which is the base class of all gameplay Objects.  This means they cannot be directly instanced into the world; they must belong to an Actor.
 
 Each Actor or Object is a single instance of a class.  The class sets up the template for the Actor or Object. It defines the variables that can be set for that Actor or Object, and the functions that
-can be carried out within that Actor or Object.  You can create new classes, or types of Objects and Actors, with C++ code. _Blueprint Classes_ primarily allow you to create classes that set up new Actors, 
-although you can extend a few Objects with _Blueprint Classes_ as well. You can also combine the two, by creating a new C++ class and then making a _Blueprint Class_ derived from that C++ class. To learn more about 
+can be carried out within that Actor or Object.  You can create new classes, or types of Objects and Actors, with C++ code. Blueprint Classes primarily allow you to create classes that set up new Actors, 
+although you can extend a few Objects with Blueprint Classes as well. You can also combine the two, by creating a new C++ class and then making a Blueprint Class derived from that C++ class. To learn more about 
 creating classes so you can make new kinds of Actors and Objects, see the [Class Creation Basics](Gameplay\ClassCreation) page.
 
 For more on Objects, Actors, and Components, see the [](Programming/UnrealArchitecture) documentation.
 
 
-[OBJECT:TopicList]
-	[PARAM:icon]
-		![](%ROOT%/reference_icon.png)(convert:false)
-	[/PARAM]
-	[PARAM:title]
-		Related Pages
-	[/PARAM]
-	[PARAM:description]
-	[/PARAM]
-	[PARAM:links]
-	* [](Programming\UnrealArchitecture\Reference\Classes)
-	* [](Programming\UnrealArchitecture\Reference\Functions)
-	* [](Programming\UnrealArchitecture\Reference\Properties)
-	* [](Programming\UnrealArchitecture\Reference\Structs)
-	* [](Programming\UnrealArchitecture\Reference\Interfaces)
-	[/PARAM]
-[/OBJECT]
-
-
