#UE4doc #docs:
Changed font settings for a term

diff --git a/Engine/Documentation/Source/Programming/Basics/CommandLineArguments/CommandLineArguments.INT.udn b/Engine/Documentation/Source/Programming/Basics/CommandLineArguments/CommandLineArguments.INT.udn
--- a/Engine/Documentation/Source/Programming/Basics/CommandLineArguments/CommandLineArguments.INT.udn
+++ b/Engine/Documentation/Source/Programming/Basics/CommandLineArguments/CommandLineArguments.INT.udn
@@ -23,7 +23,7 @@ These arguments are used to force the main editor executable (`UE4Editor.exe`) t
 These commands are not case sensitive.
 [/REGION]
 
-*Example:*
+Example:
 
       
 	UE4Editor.exe -game
