#UE4doc #docs:
Changed font settings for multiple terms
Changed multiple tip regions to notes

diff --git a/Engine/Documentation/Source/Programming/Basics/ConfigurationFiles/ConfigurationFiles.INT.udn b/Engine/Documentation/Source/Programming/Basics/ConfigurationFiles/ConfigurationFiles.INT.udn
--- a/Engine/Documentation/Source/Programming/Basics/ConfigurationFiles/ConfigurationFiles.INT.udn
+++ b/Engine/Documentation/Source/Programming/Basics/ConfigurationFiles/ConfigurationFiles.INT.udn
@@ -14,27 +14,27 @@ Engine configuration files are used for object and variable default values. User
 are created when you create a new blank project with the **Project Wizard**. New projects that begin with templates may also generate DefaultEditor.ini and DefaultInput.ini configuration files
 if needed.  
 
-The `SaveConfig()` function can be called on a class with the  [](Programming\UnrealArchitecture\Reference\Classes\Specifiers\Config) class specifier.  This saves any properties marked with the 
-`Config` property specifier to the appropriate configuration file. Generally, variables saved by **SaveConfig()** are in a section title
-which follows the format `[(package).(classname)]`.  For instance, the `[/Script/Engine.Engine]` section in DefaultEngine.ini points to the Engine class stored within the Engine package.
+The SaveConfig() function can be called on a class with the  [](Programming\UnrealArchitecture\Reference\Classes\Specifiers\Config) class specifier.  This saves any properties marked with the 
+**Config** property specifier to the appropriate configuration file. Generally, variables saved by **SaveConfig()** are in a section title
+which follows the format [(package).(classname)].  For instance, the [/Script/Engine.Engine] section in DefaultEngine.ini points to the Engine class stored within the Engine package.
 There are some exceptions with hard-coded section names.
 
 A number of the settings previously available only through editing configuration files are available in Unreal Editor in the [](Engine\UI\ProjectSettings) editor.
 
 ## Specifying Configurable Variables
 
-In order to specify which variables should be read in from configuration files, the class that contains those variables must first be given the `Config` specifier in its `UCLASS` macro.
+In order to specify which variables should be read in from configuration files, the class that contains those variables must first be given the Config specifier in its UCLASS macro.
 
 	UCLASS(Config=Game)
 	class AExampleClass : public AActor
 
-[REGION:tip]
-	Note that a category (i.e. Game) must be supplied for the `Config` specifier. This determines which configuration files the class's variables are read from and saved to. All possible categories are 
+[REGION:note]
+	A category (i.e. Game) must be supplied for the Config specifier. This determines which configuration files the class's variables are read from and saved to. All possible categories are 
 	defined in FConfigCache.ini. For a list of all configuration file categories, see [Configuration Categories](#configurationcategories).
 [/REGION]
 
-Decorating a class with the `Config` specifier just indicates that that class can have variables read in from configuration files, and specifies which files the configurations are read from. To
-specify a particular variable as read in and saved to a configuration file, the `Config` specifier must also be supplied to a `UPROPERTY()` macro.
+Decorating a class with the Config specifier just indicates that that class can have variables read in from configuration files, and specifies which files the configurations are read from. To
+specify a particular variable as read in and saved to a configuration file, the Config specifier must also be supplied to a UPROPERTY() macro.
 
 	UCLASS(Config=Game)
 	class AExampleClass : public AActor
@@ -46,7 +46,7 @@ specify a particular variable as read in and saved to a configuration file, the
 
 	};
 
-No category is supplied to the property's `Config` specifier. The `ExampleVariable` property can now be read from any Game configuration file in the [configuration file hierarchy](#filehierarchy),
+No category is supplied to the property's Config specifier. The ExampleVariable property can now be read from any Game configuration file in the [configuration file hierarchy](#filehierarchy),
 as long as the information is specified with the following syntax:
 
 	[/Script/ModuleName.ExampleClass]
@@ -54,9 +54,9 @@ as long as the information is specified with the following syntax:
 
 ### Config Files and Inheritance
 
-The `Config` `UCLASS` and `UPROPERTY` specifiers are inherited.  This means that a child class can either read in or save out all variables specified as `Config` in the parent class, and they
-will be in the same configuration file category. The variables will all be under a section title with the child class's name. For example, the configuration file information for a `ChildExampleClass`
-which inherits from the `ExampleClass` above would look like the following lines, and be saved in the same Game configuration files.
+The Config, UCLASS, and UPROPERTY specifiers are inherited.  This means that a child class can either read in or save out all variables specified as Config in the parent class, and they
+will be in the same configuration file category. The variables will all be under a section title with the child class's name. For example, the configuration file information for a ChildExampleClass
+which inherits from the ExampleClass above would look like the following lines, and be saved in the same Game configuration files.
 
 	[/Script/ModuleName.ChildExampleClass]
 	ExampleVariable=0.0f
@@ -64,8 +64,8 @@ which inherits from the `ExampleClass` above would look like the following lines
 
 ### Per-Instance Configurations
 
-Unreal Engine 4 has the ability to save the configuration of an object to any desired configuration file. If the `PerObjectConfig` specifier is used in the `UCLASS` macro, 
-configuration information for this class will be stored per instance, where each instance has a section in the .ini file named after the object in the format `[ObjectName ClassName]`. 
+Unreal Engine 4 has the ability to save the configuration of an object to any desired configuration file. If the **PerObjectConfig** specifier is used in the UCLASS macro, 
+configuration information for this class will be stored per instance, where each instance has a section in the .ini file named after the object in the format [ObjectName ClassName]. 
 This keyword is propagated to child classes.
 
 
@@ -95,7 +95,7 @@ to all projects, while project-specific settings should be in files in the proje
 The below file hierarchy example is for the Engine category of configuration files.
 
 1. Engine/Config/Base.ini
-	[REGION:tip]
+	[REGION:note]
 	Base.ini is usually empty.
 	[/REGION]
 1. Engine/Config/BaseEngine.ini
@@ -103,7 +103,7 @@ The below file hierarchy example is for the Engine category of configuration fil
 1. [ProjectDirectory]/Config/DefaultEngine.ini
 1. [ProjectDirectory]/Config/[Platform]/[Platform]Engine.ini
 1. [ProjectDirectory]/Saved/Config/[Platform]/Engine.ini
-	[REGION:tip]
+	[REGION:note]
 	The configuration file in the Saved directory only stores the project-specific and platform-specific differences in the stack of configuration files.
 	[/REGION]
 	
@@ -130,7 +130,7 @@ Typical configuration files consist of sections of key-value pairs, arranged as
 * *!* - Removes a property; but you do not have to have an exact match, just the name of the property.
 
 [REGION:note]
-**Note**: . is like + except it will potentially add a duplicate line. This is useful for the bindings (as seen in DefaultInput.ini), for instance, where the bottom-most binding takes effect, so if you add something like:
+. is like + except it will potentially add a duplicate line. This is useful for the bindings (as seen in DefaultInput.ini), for instance, where the bottom-most binding takes effect, so if you add something like:
     
     [/Script/Engine.PlayerInput]
     Bindings=(Name="Q",Command="Foo")
@@ -142,7 +142,7 @@ It will work appropriately. Using a *+* there would fail to add the last line, a
 
 #### Comments
 
-Most people seem to be under the impression that the semicolon denotes comments in configuration files, but they are not (`FConfigFile::ProcessInputFileContents` does not actually treat them, or any other string, as a comment delimiter). This behavior is intentional. Technically any character can represent a different key-value pair. Typically, a semicolon is placed at the beginning of a new line. It works like a comment, but it is not actually.
+Most people seem to be under the impression that the semicolon denotes comments in configuration files, but they are not (FConfigFile::ProcessInputFileContents does not actually treat them, or any other string, as a comment delimiter). This behavior is intentional. Technically any character can represent a different key-value pair. Typically, a semicolon is placed at the beginning of a new line. It works like a comment, but it is not actually.
      
     ; This is a Comment
 	; So is this!
