#UE4doc #docs:
Fixed some grammatical errors

diff --git a/Engine/Documentation/Source/Programming/BlueprintFunctionLibraries/BlueprintFunctionLibraries.INT.udn b/Engine/Documentation/Source/Programming/BlueprintFunctionLibraries/BlueprintFunctionLibraries.INT.udn
--- a/Engine/Documentation/Source/Programming/BlueprintFunctionLibraries/BlueprintFunctionLibraries.INT.udn
+++ b/Engine/Documentation/Source/Programming/BlueprintFunctionLibraries/BlueprintFunctionLibraries.INT.udn
@@ -19,8 +19,7 @@ Creating a Blueprint Function Library is very similar to exposing functions to B
 		UFUNCTION(BlueprintCallable, Category="Analytics")
 		static bool StartSession();
 
-As you can see in the example above, a Blueprint Function Library is indirectly UObject derived so requires the standard UCLASS() and GENERATED_UCLASS_BODY() macros. It also decorates the functions that are to be callable from Blueprints with the UFUNCTION() macro. NOTE: functions in a Blueprint Function Library can be designated as BlueprintCallable or BlueprintPure depending on whether the calls have side effects or not. See /UE4/Engine/Plugins/Runtime/Analytics/AnalyticsBlueprintLibrary for the full source code. Here’s the implementation of the StartSession() function:
-
+As you can see in the example above, a Blueprint Function Library is indirectly UObject derived so requires the standard UCLASS() and GENERATED_UCLASS_BODY() macros. It also decorates the functions that are to be callable from Blueprints with the UFUNCTION() macro. Functions in a Blueprint Function Library can be designated as BlueprintCallable or BlueprintPure depending on whether the calls have side effects or not. See /UE4/Engine/Plugins/Runtime/Analytics/AnalyticsBlueprintLibrary for the full source code. Here’s the implementation of the StartSession() function:
 	bool UAnalyticsBlueprintLibrary::StartSession()
 	{
 		TSharedPtr<IAnalyticsProvider> Provider = FAnalytics::Get().GetDefaultConfiguredProvider();
