#UE4doc #docs:
Added engine version field

diff --git a/Engine/Documentation/Source/Programming/Development/Development.INT.udn b/Engine/Documentation/Source/Programming/Development/Development.INT.udn
--- a/Engine/Documentation/Source/Programming/Development/Development.INT.udn
+++ b/Engine/Documentation/Source/Programming/Development/Development.INT.udn
@@ -2,6 +2,7 @@ Availability:Public
 Title:Development Setup
 Crumbs: %ROOT%, Programming
 Description:Setting up your code environment, creating code projects, and working with Unreal Engine source.
+version: 4.9
 
 [OBJECT:TopicButtonList]
 	[PARAM:icon]
