#UE4doc #docs - Added engine version field. Added comparison image between None and TXAA at page top. Moved three comparison images to be in a imagetable

diff --git a/Engine/Documentation/Source/Resources/ContentExamples/PostProcessing/1_14/1_14.INT.udn b/Engine/Documentation/Source/Resources/ContentExamples/PostProcessing/1_14/1_14.INT.udn
--- a/Engine/Documentation/Source/Resources/ContentExamples/PostProcessing/1_14/1_14.INT.udn
+++ b/Engine/Documentation/Source/Resources/ContentExamples/PostProcessing/1_14/1_14.INT.udn
@@ -5,6 +5,7 @@ Description:An overview of the Post Processing example level, example 1.14: AA M
 Related: Engine/Rendering/Overview
 Related: Engine/Rendering/PostProcessEffects
 Related: Engine/Rendering/PostProcessEffects/PostProcessMaterials
+Version: 4.9
 
 [VAR:Steps]
 [OBJECT:Navigation]
@@ -63,9 +64,26 @@ Related: Engine/Rendering/PostProcessEffects/PostProcessMaterials
 [/OBJECT]
 [/VAR]
 
-[REGION:lightbox]
-[![](1_14_01.png)(w:302)](1_14_01.png) [![](1_14_02.png)(w:302)](1_14_02.png) [![](1_14_03.png)(w:302)](1_14_03.png)
-[/REGION:lightbox]
+
+[OBJECT:ComparisonSlider]
+	[PARAM:before]
+	![Anti-Aliasing: None](1_14_01.png)(w:960)
+	[/PARAM]
+	[PARAM:after]
+	![Anti-Aliasing: Temporal](1_14_03.png)(w:960)
+	[/PARAM]
+[/OBJECT]
+
+
+[Region:lightbox]
+[Region:imagetable]
+| [![](1_14_01.png)(w:302)](1_14_01.png) | [![](1_14_02.png)(w:302)](1_14_02.png) | [![](1_14_03.png)(w:302)](1_14_03.png) |
+| ------ | ------ | ------ |
+| None | FXAA | TXAA |
+[/Region]
+[/Region]
+
+
 
 Anti-Aliasing improves image quality by removing the stepping effect from the edges that appear angled on screen. **Temporal AA** is the default and gives the best result. Although it is on by default, it only shows in the editor when Real-time is activated in viewport or while in Game Mode **(G-key)**. 
 
