#UE4doc #docs - Added engine version field.

diff --git a/Engine/Documentation/Source/Engine/Rendering/PostProcessEffects/PostProcessMaterials/PostProcessMaterials.INT.udn b/Engine/Documentation/Source/Engine/Rendering/PostProcessEffects/PostProcessMaterials/PostProcessMaterials.INT.udn
--- a/Engine/Documentation/Source/Engine/Rendering/PostProcessEffects/PostProcessMaterials/PostProcessMaterials.INT.udn
+++ b/Engine/Documentation/Source/Engine/Rendering/PostProcessEffects/PostProcessMaterials/PostProcessMaterials.INT.udn
@@ -14,7 +14,7 @@ Version: 4.9
 
 [EXCERPT:PostProcessMaterials_Desc]
 The engine already features complex post processing based on a graph of post processing nodes. The **Post Processing Materials**
-can be additionally inserted in some specific position. See [**FAQ**](#FAQ) **r.CompositionGraphDebug** to get a dump of the full graph.
+can be additionally inserted in some specific position. See [FAQ](#FAQ) **r.CompositionGraphDebug** to get a dump of the full graph.
 The graph is actually not only doing post processing but also some parts of the lighting. Over time we plan to make more
 parts be customizable through the material editor.
 [/EXCERPT:PostProcessMaterials_Desc]
diff --git a/Engine/Documentation/Source/Resources/ContentExamples/PostProcessing/1_16/1_16.INT.udn b/Engine/Documentation/Source/Resources/ContentExamples/PostProcessing/1_16/1_16.INT.udn
--- a/Engine/Documentation/Source/Resources/ContentExamples/PostProcessing/1_16/1_16.INT.udn
+++ b/Engine/Documentation/Source/Resources/ContentExamples/PostProcessing/1_16/1_16.INT.udn
@@ -5,6 +5,7 @@ Description:An overview of the Post Processing example level, example 1.16: Blen
 Related: Engine/Rendering/Overview
 Related: Engine/Rendering/PostProcessEffects
 Related: Engine/Rendering/PostProcessEffects/PostProcessMaterials
+Version: 4.9
 
 [VAR:Steps]
 [OBJECT:Navigation]
