#UE4doc #docs - Added Engine version field. Added/Removed bolded terms.

diff --git a/Engine/Documentation/Source/Resources/ContentExamples/PostProcessing/1_17/1_17.INT.udn b/Engine/Documentation/Source/Resources/ContentExamples/PostProcessing/1_17/1_17.INT.udn
--- a/Engine/Documentation/Source/Resources/ContentExamples/PostProcessing/1_17/1_17.INT.udn
+++ b/Engine/Documentation/Source/Resources/ContentExamples/PostProcessing/1_17/1_17.INT.udn
@@ -5,6 +5,7 @@ Description:An overview of the Post Processing example level, example 1.17: Prio
 Related: Engine/Rendering/Overview
 Related: Engine/Rendering/PostProcessEffects
 Related: Engine/Rendering/PostProcessEffects/PostProcessMaterials
+Version: 4.9
 
 [VAR:Steps]
 [OBJECT:Navigation]
@@ -70,5 +71,5 @@ Related: Engine/Rendering/PostProcessEffects/PostProcessMaterials
 [/REGION:lightbox]
 </div>
 
-Volume Priority allows you to control which Post Process Volume will override another if one is nested inside the other. If you are inside two volumes, the volume with the higher priority will be used and lower priority volumes will be ignored.
+**Volume Priority** allows you to control which Post Process Volume will override another if one is nested inside the other. If you are inside two volumes, the volume with the higher priority will be used and lower priority volumes will be ignored.
 
diff --git a/Engine/Documentation/Source/Resources/ContentExamples/PostProcessing/1_18/1_18.INT.udn b/Engine/Documentation/Source/Resources/ContentExamples/PostProcessing/1_18/1_18.INT.udn
--- a/Engine/Documentation/Source/Resources/ContentExamples/PostProcessing/1_18/1_18.INT.udn
+++ b/Engine/Documentation/Source/Resources/ContentExamples/PostProcessing/1_18/1_18.INT.udn
@@ -5,6 +5,7 @@ Description:An overview of the Post Processing example level, example 1.18: Blen
 Related: Engine/Rendering/Overview
 Related: Engine/Rendering/PostProcessEffects
 Related: Engine/Rendering/PostProcessEffects/PostProcessMaterials
+Version: 4.9
 
 [VAR:Steps]
 [OBJECT:Navigation]
@@ -71,5 +72,5 @@ Related: Engine/Rendering/PostProcessEffects/PostProcessMaterials
 [![](1_18_01.png)(w:302)](1_18_01.png) [![](1_18_02.png)(w:302)](1_18_02.png) [![](1_18_03.png)(w:302)](1_18_03.png)
 [/REGION:lightbox]
 
-This is used to change the transition from one PostProcessVolume to another over a distance. Volumes can just take control of the post process pipeline or, using these settings, it can smoothing transition between 2 volumes over the distance set by the radius.
+**Blend Radius** and **Weight**  is used to change the transition from one PostProcessVolume to another over a distance. Volumes can just take control of the post process pipeline or, using these settings, it can smoothing transition between two volumes over the distance set by the radius.
 
