#UE4doc #docs - Added engine version field.

diff --git a/Engine/Documentation/Source/Resources/ContentExamples/PostProcessing/1_2/1_2.INT.udn b/Engine/Documentation/Source/Resources/ContentExamples/PostProcessing/1_2/1_2.INT.udn
--- a/Engine/Documentation/Source/Resources/ContentExamples/PostProcessing/1_2/1_2.INT.udn
+++ b/Engine/Documentation/Source/Resources/ContentExamples/PostProcessing/1_2/1_2.INT.udn
@@ -5,6 +5,7 @@ Description:An overview of the Post Processing example level, example 1.2: Film
 Related: Engine/Rendering/Overview
 Related: Engine/Rendering/PostProcessEffects
 Related: Engine/Rendering/PostProcessEffects/PostProcessMaterials
+Version: 4.9
 
 [VAR:Steps]
 [OBJECT:Navigation]
@@ -68,5 +69,5 @@ Related: Engine/Rendering/PostProcessEffects/PostProcessMaterials
 [![](1_2_01.png)(w:302)](1_2_01.png) [![](1_2_02.png)(w:302)](1_2_02.png) [![](1_2_03.png)(w:302)](1_2_03.png)
 [/REGION:lightbox]
 
-The Film category contains various properties for adjusting the color of your scenes. Includes tinting shadows, color mixing channels, saturation, contrast, crushing shadows, crushing highlights, and dynamic range.
+The **Film** category contains various properties for adjusting the color of your scenes. Includes tinting shadows, color mixing channels, saturation, contrast, crushing shadows, crushing highlights, and dynamic range.
 
