#UE4doc #docs - Added engine version field.

diff --git a/Engine/Documentation/Source/Resources/ContentExamples/PostProcessing/1_6/1_6.INT.udn b/Engine/Documentation/Source/Resources/ContentExamples/PostProcessing/1_6/1_6.INT.udn
--- a/Engine/Documentation/Source/Resources/ContentExamples/PostProcessing/1_6/1_6.INT.udn
+++ b/Engine/Documentation/Source/Resources/ContentExamples/PostProcessing/1_6/1_6.INT.udn
@@ -6,6 +6,7 @@ Related: Engine/Rendering/PostProcessEffects/AutomaticExposure
 Related: Engine/Rendering/Overview
 Related: Engine/Rendering/PostProcessEffects
 Related: Engine/Rendering/PostProcessEffects/PostProcessMaterials
+Version: 4.9
 
 [VAR:Steps]
 [OBJECT:Navigation]
diff --git a/Engine/Documentation/Source/Resources/ContentExamples/PostProcessing/1_7/1_7.INT.udn b/Engine/Documentation/Source/Resources/ContentExamples/PostProcessing/1_7/1_7.INT.udn
--- a/Engine/Documentation/Source/Resources/ContentExamples/PostProcessing/1_7/1_7.INT.udn
+++ b/Engine/Documentation/Source/Resources/ContentExamples/PostProcessing/1_7/1_7.INT.udn
@@ -6,6 +6,7 @@ Related: Engine/Rendering/PostProcessEffects/LensFlare
 Related: Engine/Rendering/Overview
 Related: Engine/Rendering/PostProcessEffects
 Related: Engine/Rendering/PostProcessEffects/PostProcessMaterials
+Version: 4.9
 
 [VAR:Steps]
 [OBJECT:Navigation]
diff --git a/Engine/Documentation/Source/Resources/ContentExamples/PostProcessing/1_8/1_8.INT.udn b/Engine/Documentation/Source/Resources/ContentExamples/PostProcessing/1_8/1_8.INT.udn
--- a/Engine/Documentation/Source/Resources/ContentExamples/PostProcessing/1_8/1_8.INT.udn
+++ b/Engine/Documentation/Source/Resources/ContentExamples/PostProcessing/1_8/1_8.INT.udn
@@ -6,6 +6,7 @@ Related: Engine/Rendering/LightingAndShadows/AmbientOcclusion
 Related: Engine/Rendering/Overview
 Related: Engine/Rendering/PostProcessEffects
 Related: Engine/Rendering/PostProcessEffects/PostProcessMaterials
+Version: 4.9
 
 [VAR:Steps]
 [OBJECT:Navigation]
diff --git a/Engine/Documentation/Source/Resources/ContentExamples/PostProcessing/1_9/1_9.INT.udn b/Engine/Documentation/Source/Resources/ContentExamples/PostProcessing/1_9/1_9.INT.udn
--- a/Engine/Documentation/Source/Resources/ContentExamples/PostProcessing/1_9/1_9.INT.udn
+++ b/Engine/Documentation/Source/Resources/ContentExamples/PostProcessing/1_9/1_9.INT.udn
@@ -6,6 +6,7 @@ Related: Engine/Rendering/PostProcessEffects/ScreenSpaceReflection
 Related: Engine/Rendering/Overview
 Related: Engine/Rendering/PostProcessEffects
 Related: Engine/Rendering/PostProcessEffects/PostProcessMaterials
+Version: 4.9
 
 [VAR:Steps]
 [OBJECT:Navigation]
