#UE4doc #docs - Added engine version field. Removed ToC with single line entry.

diff --git a/Engine/Documentation/Source/Resources/ContentExamples/Reflections/1_1/1_1.INT.udn b/Engine/Documentation/Source/Resources/ContentExamples/Reflections/1_1/1_1.INT.udn
--- a/Engine/Documentation/Source/Resources/ContentExamples/Reflections/1_1/1_1.INT.udn
+++ b/Engine/Documentation/Source/Resources/ContentExamples/Reflections/1_1/1_1.INT.udn
@@ -8,6 +8,7 @@ Related: Engine/Rendering/LightingAndShadows/LightTypes/SkyLight\
 Related: Engine/Rendering/Materials/PhysicallyBased
 Related: Resources/Showcases/RealisticRendering
 Related: Engine/Rendering/PostProcessEffects
+Version: 4.9
 
 [VAR:Steps]
 [OBJECT:Navigation]
diff --git a/Engine/Documentation/Source/Resources/ContentExamples/Reflections/Reflections.INT.udn b/Engine/Documentation/Source/Resources/ContentExamples/Reflections/Reflections.INT.udn
--- a/Engine/Documentation/Source/Resources/ContentExamples/Reflections/Reflections.INT.udn
+++ b/Engine/Documentation/Source/Resources/ContentExamples/Reflections/Reflections.INT.udn
@@ -8,6 +8,7 @@ Related: Engine/Rendering/LightingAndShadows/LightTypes/SkyLight\
 Related: Engine/Rendering/Materials/PhysicallyBased
 Related: Resources/Showcases/RealisticRendering
 Related: Engine/Rendering/PostProcessEffects
+Version: 4.9
 
 
 [REGION:banner]
@@ -45,12 +46,11 @@ Related: Engine/Rendering/PostProcessEffects
 [/OBJECT]
 [/VAR]
 
-[TOC(start:2 end:2)]
 
-The **Reflections** Content Example map provides examples on how to adjust the way reflections appear within your levels by using a Reflection Capture Actor (which projects a reflection within the Actor's effect volume or radius). 
+The **Reflections** Content Example map provides examples on how to adjust the way reflections appear within your levels by using a **Reflection Capture** Actor (which projects a reflection within the Actor's effect volume or radius). 
 Examples are also shown on how the Roughness, Metallic, and Specular properties of a Material can be used to produce different reflection results. 
 
-In addition to tweaking reflection settings, the use of both a Box Reflection Actor and the more commonly used*Sphere Reflection Actor are shown with sample effects that can be achieved.
+In addition to tweaking reflection settings, the use of both a **Box Reflection** Actor and the more commonly used Sphere Reflection Actor are shown with sample effects that can be achieved.
 Lastly, Scene Capture methods are demonstrated which can be used to capture a fully dynamic cube map or capture an image of the scene on every frame. 
 
 
