#UE4doc #docs - Added engine version field.

diff --git a/Engine/Documentation/Source/Resources/ContentExamples/Reflections/1_4/1_4.INT.udn b/Engine/Documentation/Source/Resources/ContentExamples/Reflections/1_4/1_4.INT.udn
--- a/Engine/Documentation/Source/Resources/ContentExamples/Reflections/1_4/1_4.INT.udn
+++ b/Engine/Documentation/Source/Resources/ContentExamples/Reflections/1_4/1_4.INT.udn
@@ -8,6 +8,7 @@ Related: Engine/Rendering/LightingAndShadows/LightTypes/SkyLight\
 Related: Engine/Rendering/Materials/PhysicallyBased
 Related: Resources/Showcases/RealisticRendering
 Related: Engine/Rendering/PostProcessEffects
+Version: 4.9
 
 [VAR:Steps]
 [OBJECT:Navigation]
diff --git a/Engine/Documentation/Source/Resources/ContentExamples/Reflections/1_5/1_5.INT.udn b/Engine/Documentation/Source/Resources/ContentExamples/Reflections/1_5/1_5.INT.udn
--- a/Engine/Documentation/Source/Resources/ContentExamples/Reflections/1_5/1_5.INT.udn
+++ b/Engine/Documentation/Source/Resources/ContentExamples/Reflections/1_5/1_5.INT.udn
@@ -8,6 +8,7 @@ Related: Engine/Rendering/LightingAndShadows/LightTypes/SkyLight\
 Related: Engine/Rendering/Materials/PhysicallyBased
 Related: Resources/Showcases/RealisticRendering
 Related: Engine/Rendering/PostProcessEffects
+Version: 4.9
 
 [VAR:Steps]
 [OBJECT:Navigation]
diff --git a/Engine/Documentation/Source/Resources/ContentExamples/Reflections/1_6/1_6.INT.udn b/Engine/Documentation/Source/Resources/ContentExamples/Reflections/1_6/1_6.INT.udn
--- a/Engine/Documentation/Source/Resources/ContentExamples/Reflections/1_6/1_6.INT.udn
+++ b/Engine/Documentation/Source/Resources/ContentExamples/Reflections/1_6/1_6.INT.udn
@@ -8,6 +8,7 @@ Related: Engine/Rendering/LightingAndShadows/LightTypes/SkyLight\
 Related: Engine/Rendering/Materials/PhysicallyBased
 Related: Resources/Showcases/RealisticRendering
 Related: Engine/Rendering/PostProcessEffects
+Version: 4.9
 
 [VAR:Steps]
 [OBJECT:Navigation]
diff --git a/Engine/Documentation/Source/Resources/ContentExamples/Reflections/1_7/1_7.INT.udn b/Engine/Documentation/Source/Resources/ContentExamples/Reflections/1_7/1_7.INT.udn
--- a/Engine/Documentation/Source/Resources/ContentExamples/Reflections/1_7/1_7.INT.udn
+++ b/Engine/Documentation/Source/Resources/ContentExamples/Reflections/1_7/1_7.INT.udn
@@ -8,6 +8,7 @@ Related: Engine/Rendering/LightingAndShadows/LightTypes/SkyLight\
 Related: Engine/Rendering/Materials/PhysicallyBased
 Related: Resources/Showcases/RealisticRendering
 Related: Engine/Rendering/PostProcessEffects
+Version: 4.9
 
 [VAR:Steps]
 [OBJECT:Navigation]
@@ -71,7 +72,7 @@ Related: Engine/Rendering/PostProcessEffects
 ![](Reflections_1_7_Header.png)
 ## Overview
 
-Much like the [**Scene Capture Cube**](Resources\ContentExamples\Reflections\1_6), a **Scene Capture 2D** makes an image of the scene _on every frame_. In fact, Scene Capture 2D Actors do not work like real reflections; they work more like a TV screen receiving an image from a camera. The Scene Capture 2D captures the scene from its view frustum, stores that view as an image, which is then used within a Material. In fact, the **Scene Capture 2D** representation in Unreal Editor is a camera. As such, they are most commonly used for security camera effects in game. 
+Much like the [Scene Capture Cube](Resources\ContentExamples\Reflections\1_6), a **Scene Capture 2D** makes an image of the scene _on every frame_. In fact, Scene Capture 2D Actors do not work like real reflections; they work more like a TV screen receiving an image from a camera. The Scene Capture 2D captures the scene from its view frustum, stores that view as an image, which is then used within a Material. In fact, the Scene Capture 2D representation in Unreal Editor is a camera. As such, they are most commonly used for security camera effects in game. 
 
 [REGION:warning]
 While only 1/6 as expensive as the Scene Capture Cube, care should still be used with the Scene Capture 2D, as they still require extra processing.
