#UE4doc #docs - Added engine version field. Remove ToC for single line entry.

diff --git a/Engine/Documentation/Source/Resources/ContentExamples/StaticMeshes/StaticMeshes.INT.udn b/Engine/Documentation/Source/Resources/ContentExamples/StaticMeshes/StaticMeshes.INT.udn
--- a/Engine/Documentation/Source/Resources/ContentExamples/StaticMeshes/StaticMeshes.INT.udn
+++ b/Engine/Documentation/Source/Resources/ContentExamples/StaticMeshes/StaticMeshes.INT.udn
@@ -7,6 +7,7 @@ Related: Engine/Physics/Collision/Reference
 Related: Engine/Content/Types/StaticMeshes/Editor
 Related: Engine/Actors/StaticMeshActor
 Related: Engine/Content/FBX/BestPractices
+Version: 4.9
 
 
 
@@ -45,8 +46,6 @@ Related: Engine/Content/FBX/BestPractices
 [/OBJECT]
 [/VAR]
 
-[TOC(start:2 end:2)]
-
 Inside the **StaticMeshes** Content Example map, you will find essential examples using Static Meshes including how to set Mobility, assign UVs, and add or adjust Collision settings. Other examples include how to customize your Static Meshes with Material Elements (which allow you to specify the Material applied to parts of a Static Mesh) as well as using the MeshPaint Tool to manually paint onto a Static Mesh. 
 Last, you will see how LODs can improve performance by scaling back the quality of a mesh based on player distance.
 
@@ -55,7 +54,7 @@ Last, you will see how LODs can improve performance by scaling back the quality
 
 ![](StaticMeshes_Map.png)(w:640)
 
-Listed below are the examples provided inside the StaticMeshes map:
+Listed below are the examples provided inside the **StaticMeshes** map:
 
 | Example | What is Demonstrated |
 | ------- | ------------ |
