#UE4doc #docs - Added engine version field. Remove ToC for single line entry.

diff --git a/Engine/Documentation/Source/Resources/ContentExamples/Volumes/Volumes.INT.udn b/Engine/Documentation/Source/Resources/ContentExamples/Volumes/Volumes.INT.udn
--- a/Engine/Documentation/Source/Resources/ContentExamples/Volumes/Volumes.INT.udn
+++ b/Engine/Documentation/Source/Resources/ContentExamples/Volumes/Volumes.INT.udn
@@ -7,6 +7,7 @@ Related: Engine/Rendering/PostProcessEffects
 Related: Engine/Rendering/LightingAndShadows/LightPropagationVolumes
 Related: Engine/AI/BehaviorTrees/QuickStart/2
 Related: Engine/Actors/Triggers
+Version: 4.9
 
 
 [REGION:banner]
@@ -44,7 +45,7 @@ Related: Engine/Actors/Triggers
 [/OBJECT]
 [/VAR]
 
-[TOC(start:2 end:2)]
+
 
 The **Volumes** Content Example map illustrates how you can use different types of Volumes to produce different gameplay effects, set boundaries, improve performance and (or) reduce lighting build times. 
 As seen in example 1.1 below, you can place a Volume by selecting any Volume from the Volumes Tab and dragging and dropping it into your level. Based on the type of Volume used, different settings may be available for editing. 
@@ -57,7 +58,7 @@ The most commonly used Volumes for gameplay purposes would be Blocking Volumes (
 
 ![](Volumes_Map.png)(w:640)
 
-Listed below are the examples provided inside the Volumes map:
+Listed below are the examples provided inside the **Volumes** map:
 
 | Example | What is Demonstrated |
 | ------- | ------------ |
