#UE4doc #docs - Added engine version field.

diff --git a/Engine/Documentation/Source/Resources/Downloads/Downloads.INT.udn b/Engine/Documentation/Source/Resources/Downloads/Downloads.INT.udn
--- a/Engine/Documentation/Source/Resources/Downloads/Downloads.INT.udn
+++ b/Engine/Documentation/Source/Resources/Downloads/Downloads.INT.udn
@@ -2,6 +2,7 @@ Availability: Licensee
 Title:Resource Downloads
 Crumbs:%ROOT%, Resources
 Description:Repository of downloadable content and example materials for use in developing games with UE4.
+Version: 4.9
 
 [TOC]
 
diff --git a/Engine/Documentation/Source/Resources/SampleGames/SampleGames.INT.udn b/Engine/Documentation/Source/Resources/SampleGames/SampleGames.INT.udn
--- a/Engine/Documentation/Source/Resources/SampleGames/SampleGames.INT.udn
+++ b/Engine/Documentation/Source/Resources/SampleGames/SampleGames.INT.udn
@@ -2,6 +2,7 @@ Availability: Public
 Title:Example Game Projects
 Crumbs:%ROOT%, Resources
 Description:Full sample games demonstrating the implementation of various genres.
+Version: 4.9
 
 
 [VAR:Topic]
