#UE4doc #docs - Added engine version field.

diff --git a/Engine/Documentation/Source/Resources/SampleGames/CardGame/CardGame.INT.udn b/Engine/Documentation/Source/Resources/SampleGames/CardGame/CardGame.INT.udn
--- a/Engine/Documentation/Source/Resources/SampleGames/CardGame/CardGame.INT.udn
+++ b/Engine/Documentation/Source/Resources/SampleGames/CardGame/CardGame.INT.udn
@@ -8,6 +8,7 @@ Related: Engine/Animation/AnimBlueprints/EventGraph
 Related: Gameplay/Input
 Related: Gameplay/Framework
 Related: Engine/Content/Types/StaticMeshes
+Version: 4.9
 
 
 [VAR:Topic]
