#UE4doc #docs:
Updated engine version field
Fixed capitalization issues
Changed a tip region into a note

diff --git a/Engine/Documentation/Source/Programming/Development/BuildingUnrealEngine/BuildingUnrealEngine.INT.udn b/Engine/Documentation/Source/Programming/Development/BuildingUnrealEngine/BuildingUnrealEngine.INT.udn
--- a/Engine/Documentation/Source/Programming/Development/BuildingUnrealEngine/BuildingUnrealEngine.INT.udn
+++ b/Engine/Documentation/Source/Programming/Development/BuildingUnrealEngine/BuildingUnrealEngine.INT.udn
@@ -8,7 +8,7 @@ Related: Engine/Content/QuickStart
 Related: Programming/QuickStart
 Related: Engine/QuickStart
 Related: Support/Builds
-Version: 4.8
+Version: 4.9
 SkillLevel: Advanced
 
 %Globals:osselect_linux%
@@ -39,7 +39,7 @@ SkillLevel: Advanced
 			All project files are intermediate (`[UE4Root]\Engine\Intermediate\ProjectFiles`). You must generate project files each time you sync a new build to ensure they are up to date. If you delete your `Intermediate` folder, you must regenerate project files using the 'GenerateProjectFiles' batch file.
 			[/REGION]
 
-		1.	Load the project into Visual Studio by **double-clicking UE4.sln**.
+		1.	Load the project into Visual Studio by **Double-clicking UE4.sln**.
 
 		1.	Set your solution configuration to **Development Editor**.
 
@@ -90,7 +90,7 @@ SkillLevel: Advanced
 		
 		1.	Inside the root directory, run **GenerateProjectFiles.command** to set-up your project files.
 		
-		1.	Load the project into XCode by **double-clicking UE4.xcodeproj**.
+		1.	Load the project into XCode by **Double-clicking UE4.xcodeproj**.
 
 		1.	To set your build target, select **UE4Editor - Mac** for **My Mac** in the title bar.
 
@@ -165,7 +165,7 @@ If you're looking to quickly get started with UE4, refer to the following tutori
 *	[](Engine/QuickStart)
 
 
-[REGION:tip]
+[REGION:note]
 UE4's in-editor help features are a great way to get your questions answered.
 [/REGION]
 
