#UE4doc #docs:
Fixed a logic error in provided code

diff --git a/Engine/Documentation/Source/Programming/Development/CodingStandard/CodingStandard.INT.udn b/Engine/Documentation/Source/Programming/Development/CodingStandard/CodingStandard.INT.udn
--- a/Engine/Documentation/Source/Programming/Development/CodingStandard/CodingStandard.INT.udn
+++ b/Engine/Documentation/Source/Programming/Development/CodingStandard/CodingStandard.INT.udn
@@ -111,7 +111,7 @@ Comments are communication; communication is vital. Some things to keep in mind
 * Write self-documenting code:
 
 		// Bad:
-		t = s + l + b;
+		t = s + l - b;
 		
 		// Good:
 		TotalLeaves = SmallLeaves + LargeLeaves - SmallAndLargeLeaves;
@@ -134,7 +134,7 @@ Comments are communication; communication is vital. Some things to keep in mind
 		// total number of leaves is sum of
 		// small and large leaves less the
 		// number of leaves that are both
-		t = s + l + b;
+		t = s + l - b;
 		
 		// Good:
 		TotalLeaves = SmallLeaves + LargeLeaves - SmallAndLargeLeaves;
