#UE4doc #docs:
Changed tip region into a note

diff --git a/Engine/Documentation/Source/Programming/Development/CompilingProjects/CompilingProjects.INT.udn b/Engine/Documentation/Source/Programming/Development/CompilingProjects/CompilingProjects.INT.udn
--- a/Engine/Documentation/Source/Programming/Development/CompilingProjects/CompilingProjects.INT.udn
+++ b/Engine/Documentation/Source/Programming/Development/CompilingProjects/CompilingProjects.INT.udn
@@ -135,7 +135,7 @@ When running Unreal, it is important to add the `-game` flag if you rebuilt your
 
 ## Building with Xcode
 
-[REGION:tip]
+[REGION:note]
 When building in Xcode, you are only compiling the game project, not the Editor.  
 [/REGION]
 
