#UE4doc #docs:
Changed font settings for multiple terms
Replaced an image as it was out of date
Removed a few images as they did not follow the documentation guidelines

diff --git a/Engine/Documentation/Source/Programming/Development/ManagingGameCode/CppClassWizard/CppClassWizard.INT.udn b/Engine/Documentation/Source/Programming/Development/ManagingGameCode/CppClassWizard/CppClassWizard.INT.udn
--- a/Engine/Documentation/Source/Programming/Development/ManagingGameCode/CppClassWizard/CppClassWizard.INT.udn
+++ b/Engine/Documentation/Source/Programming/Development/ManagingGameCode/CppClassWizard/CppClassWizard.INT.udn
@@ -9,7 +9,7 @@ Version: 4.9
 [EXCERPT:main]
 
 
-The _C++ Class Wizard_ provides a fast and easy way to add native C++ code classes into your project for you to extend 
+The **C++ Class Wizard** provides a fast and easy way to add native C++ code classes into your project for you to extend 
 with your own functionality, if you wish. This converts a content-only project into a code project. You can access the C++ Class Wizard like so:
 
 [REGION:note]
@@ -17,11 +17,11 @@ On Windows, please make sure you have Visual Studio 2013 Express, Community, or
 [/REGION:note]
 
 
-1. In the main editor, select **File > Add Code to Project...**
+1. In the main editor, select **File > New C++ Class...**
 
 	![](FileCPP.png)
 
-1. The **C++ Class Wizard** will appear, showing commonly added classes. If you do not see the class you are looking for, then you will need to view the entire Class hierarchy listing. Click the ![](button_ShowAllClasses.png)(h:25) checkbox in the upper-right corner of the window. 
+1. The **C++ Class Wizard** will appear, showing commonly added classes. If you do not see the class you are looking for, then you will need to view the entire Class hierarchy listing. Click the **Show All Classes** checkbox in the upper-right corner of the window. 
 
 	[REGION:imagetable]
 		| [REGION:lightbox][![](CppClassWizard.png)(w:430)](CppClassWizard.png)[/REGION] | [REGION:lightbox][![](AllClasses.png)(w:430)](AllClasses.png)[/REGION] |
@@ -29,9 +29,9 @@ On Windows, please make sure you have Visual Studio 2013 Express, Community, or
 		| Default Classes | All Classes |
 	[/REGION]
 
-	Choose the Class you wish to add and click the ![](button_Next.png) button.
+	Choose the Class you wish to add and click the **Next >** button.
 
-1. You will then be prompted to enter a **name** for your new Class. Do so and click the ![](button_Create.png) button. This will create the header (.h) and source (.cpp) files.
+1. You will then be prompted to enter a **Name** for your new Class. Do so and click the **Create Class** button. This will create the header (.h) and source (.cpp) files.
 	
 	[REGION:lightbox]
 		[![](NameNewActor.png)(w:600)](NameNewActor.png)
@@ -51,7 +51,7 @@ On Windows, please make sure you have Visual Studio 2013 Express, Community, or
 		active_button_content
 		[/PARAMLITERAL]
 		[PARAM:content]
-		The code will immediately open in _Visual Studio 2013_, ready for editing.
+		The code will immediately open in Visual Studio 2013, ready for editing.
 	
 		![](CodeEditing.png)
 
@@ -67,7 +67,7 @@ On Windows, please make sure you have Visual Studio 2013 Express, Community, or
 		[PARAMLITERAL:active]
 		[/PARAMLITERAL]
 		[PARAM:content]
-		The code will immediately open in _Xcode_, ready for editing.
+		The code will immediately open in Xcode, ready for editing.
 
 		![](CodeEditing_xcode.png)
 
