#UE4doc #docs:
Changed font settings for multiple terms
Updated an image as it was out of date

diff --git a/Engine/Documentation/Source/Programming/Development/ManagingGameCode/ManagingGameCode.INT.udn b/Engine/Documentation/Source/Programming/Development/ManagingGameCode/ManagingGameCode.INT.udn
--- a/Engine/Documentation/Source/Programming/Development/ManagingGameCode/ManagingGameCode.INT.udn
+++ b/Engine/Documentation/Source/Programming/Development/ManagingGameCode/ManagingGameCode.INT.udn
@@ -73,12 +73,12 @@ Version: 4.9
 	[PUBLISH:Licensee]
 	**If your project is inside the UE4 root directory:**
 
-	* Open the `UE4.sln` Visual Studio solution located in the UE4 root directory.
+	* Open the UE4.sln Visual Studio solution located in the UE4 root directory.
 
 	**If your project is outside of the UE4 root directory:**
 	[/PUBLISH:Licensee]
 
-	* Open the `[ProjectName].sln` Visual Studio solution located in the root of the project's directory.
+	* Open the [ProjectName].sln Visual Studio solution located in the root of the project's directory.
 	[/PARAM]
 [/OBJECT]
 [OBJECT:ToggleButtonContent]
@@ -97,7 +97,7 @@ Version: 4.9
 
 	You can also open the project in Xcode through Finder or Xcode's **File > Open**.
 
-	* Open the `[ProjectName].xcodeproj` Xcode project located in the root of the project's directory. 
+	* Open the [ProjectName].xcodeproj Xcode project located in the root of the project's directory. 
 	[/PARAM]
 [/OBJECT]
 
@@ -127,8 +127,8 @@ Version: 4.9
 
 **If your project is inside the UE4 root directory:**
 
-1. Run the `GenerateProjectFiles.bat` batch file located in the root `UE4` directory.
-1. Open the `UE4.sln` Visual Studio solution located in the root `UE4` directory.
+1. Run the GenerateProjectFiles.bat batch file located in the root `UE4` directory.
+1. Open the UE4.sln Visual Studio solution located in the root `UE4` directory.
 
 There are advanced options available for [](Programming/UnrealBuildSystem/ProjectFileGenerator), which can be used to generate project files for shipping, additional platforms, or Visual Studio 2013. 
 
@@ -152,9 +152,9 @@ To get the **Generate Visual Studio Files** menu item to show up in the **right-
 	active_button_content
 	[/PARAMLITERAL]
 	[PARAM:content]
-	1. Navigate to the location of `[ProjectName].uproject` in Windows Explorer.
+	1. Navigate to the location of [ProjectName].uproject in Windows Explorer.
 	
-	1. **Right-click** on the `[ProjectName].uproject` file and select **Generate Visual Studio Files**.
+	1. **Right-click** on the [ProjectName].uproject file and select **Generate Visual Studio Files**.
 	
 		![](Generate_Visual_Studio_Files.png)
 
@@ -162,7 +162,7 @@ To get the **Generate Visual Studio Files** menu item to show up in the **right-
 
 		![](generate_project_files.png)
 
-	1. Open the `[ProjectName].sln` Visual Studio solution located in the root of the project's directory to view the game project in Visual Studio.
+	1. Open the [ProjectName].sln Visual Studio solution located in the root of the project's directory to view the game project in Visual Studio.
 	[/PARAM]
 [/OBJECT]
 [OBJECT:ToggleButtonContent]
@@ -191,8 +191,8 @@ To get the **Generate Visual Studio Files** menu item to show up in the **right-
 
 [PUBLISH:Licensee]
 [REGION:note]
-For advanced project file generation options for projects outside the UE4 root directory, edit the `RegisterShellCommands.bat` batch file. The `Generate Visual Studio projects (all platforms)` entry
-shows an example of an additional menu entry, in this case with the `-AllPlatforms` flag. The [available options](Programming/UnrealBuildSystem/ProjectFileGenerator) are the same as for the main `GenerateProjectFiles.bat` file.  
+For advanced project file generation options for projects outside the UE4 root directory, edit the RegisterShellCommands.bat batch file. The Generate Visual Studio projects (all platforms) entry
+shows an example of an additional menu entry, in this case with the `-AllPlatforms` flag. The [available options](Programming/UnrealBuildSystem/ProjectFileGenerator) are the same as for the main GenerateProjectFiles.bat file.  
 [/REGION:note]
 [/PUBLISH:Licensee]
 
