#UE4doc #docs:
Changed font settings for multiple terms
Edited an image to match the documentation guidelines

diff --git a/Engine/Documentation/Source/Programming/Development/VisualStudioSetup/VisualStudioSetup.INT.udn b/Engine/Documentation/Source/Programming/Development/VisualStudioSetup/VisualStudioSetup.INT.udn
--- a/Engine/Documentation/Source/Programming/Development/VisualStudioSetup/VisualStudioSetup.INT.udn
+++ b/Engine/Documentation/Source/Programming/Development/VisualStudioSetup/VisualStudioSetup.INT.udn
@@ -46,14 +46,14 @@ The following are recommended settings for developers using Visual Studio with U
 
 Other configuration steps you may find useful are:
 
-* Turn off **Show Inactive Blocks**.  If you do not, many chunks of code may appear grayed out in the text editor.  (_Tools > Options > Text Editor > C/C++ > View)_.
-* Set **Disable External Dependencies Folders** to _True_ to hide unneeded folders in the **Solution Explorer**.  (**Disable External Dependencies Folder** in _Tools > Options > Text Editor > C/C++ > Advanced_)
-* Turn off **Edit & Continue** features, you do not need them.  (_Tools > Options > Debugging > Edit and Continue_)
+* Turn off **Show Inactive Blocks**.  If you do not, many chunks of code may appear grayed out in the text editor.  (**Tools > Options > Text Editor > C/C++ > View**).
+* Set **Disable External Dependencies Folders** to True to hide unneeded folders in the **Solution Explorer**.  (**Disable External Dependencies Folder** in **Tools > Options > Text Editor > C/C++ > Advanced**)
+* Turn off **Edit & Continue** features, you do not need them.  (**Tools > Options > Debugging > Edit and Continue**)
 
 * Turn on IntelliSense. (See [Intellisense, Live Errors, and Squiggles](#Intellisense))
 
 **If you have Visual Assist X installed:**
-* Turn off **Format After Paste**.  It does weird things sometimes.  (_VAssistX > Visual Assist X Options > Advanced > Corrections_)
+* Turn off **Format After Paste**.  It does weird things sometimes.  (**VAssistX > Visual Assist X Options > Advanced > Corrections**)
 
 
 (#Intellisense)
