#UE4doc #docs:
Resubmitting as not all changes were saved

diff --git a/Engine/Documentation/Source/GettingStarted/DownloadingUnrealEngine/Perforce/Setup/Setup.INT.udn b/Engine/Documentation/Source/GettingStarted/DownloadingUnrealEngine/Perforce/Setup/Setup.INT.udn
--- a/Engine/Documentation/Source/GettingStarted/DownloadingUnrealEngine/Perforce/Setup/Setup.INT.udn
+++ b/Engine/Documentation/Source/GettingStarted/DownloadingUnrealEngine/Perforce/Setup/Setup.INT.udn
@@ -17,8 +17,8 @@ To learn about setting up and establishing a VPN connection, see the [Setting Up
 
 [REGION:red]
 Epic Games has executed a custom licensing agreement with Perforce. To be in compliance with this license, only 
-one machine per license should sync full builds, but multiple users **may** connect with the same login name to 
-view source metadata and past revisions. Individual UDN login information will **not** work.
+one machine per license should sync full builds, but multiple users may connect with the same login name to 
+view source metadata and past revisions. Individual UDN login information will not work.
 [/REGION]
 
 In some cases, larger organizations may have an account for each team; even still, only one machine per team 
@@ -177,7 +177,7 @@ this on our proxy.
 
 
 It is a VBScript that removes files that have not been accessed within a certain time period, keeping things nice and 
-tidy. It is best to set up as a Windows _at job_ or a scheduled task. Be sure to run it as a user with the appropriate 
+tidy. It is best to set up as a Windows "at job" or a scheduled task. Be sure to run it as a user with the appropriate 
 permissions to the location on the file system.
 
 The script can be called as:
