#UE4doc #docs:
Changed font settings for multiple terms
Replaced a couple images

diff --git a/Engine/Documentation/Source/Programming/Development/VisualStudioSetup/UnrealVS/UnrealVS.INT.udn b/Engine/Documentation/Source/Programming/Development/VisualStudioSetup/UnrealVS/UnrealVS.INT.udn
--- a/Engine/Documentation/Source/Programming/Development/VisualStudioSetup/UnrealVS/UnrealVS.INT.udn
+++ b/Engine/Documentation/Source/Programming/Development/VisualStudioSetup/UnrealVS/UnrealVS.INT.udn
@@ -31,12 +31,12 @@ The UnrealVS extension does not work with the Visual Studio Express editions. It
 
 		[UE4RootLocation]/Engine/Extras/UnrealVS/<VS-Version>/UnrealVS.vsix
 
-1. Run the `UnrealVS.vsix` file by **double-clicking** it to begin the installation.
+1. Run the `UnrealVS.vsix` file by **Double-clicking** it to begin the installation.
 
-1. UnrealVS extensions target specific Visual Studio versions.  Check that the installer is targeting your version and that the box is ticked then click the ![Install](button_install.png) button to proceed.
+1. UnrealVS extensions target specific Visual Studio versions.  Check that the installer is targeting your version and that the box is ticked then click the **Install** button to proceed.
 
 	![](install_dialog.png)
-1. Click the ![Close](button_close.png) button once the installation is complete.
+1. Click the **Close** button once the installation is complete.
 
 	![Install Complete](install_complete_dialog.png)
 
@@ -51,7 +51,7 @@ The UnrealVS extension does not work with the Visual Studio Express editions. It
 	The extension is also shown in the Installed Products list in the **Help->About Microsoft Visual Studio** dialog.
 	[/REGION]
 
-1. Go to **View->Toolbars** or **right-click** the Visual Studio toolbar area and select **UnrealVS** to display the extension's toolbar.
+1. Go to **View->Toolbars** or **Right-click** the Visual Studio toolbar area and select **UnrealVS** to display the extension's toolbar.
 
 	![UnrealVS Toolbar](unrealvs_toolbar_cmd.png)
 
@@ -85,10 +85,10 @@ hotkey or other means of running custom commands.
 
 1. Go to **Tools->Options...** and select **Keyboard** under the **Environment** tree.
 
-1. Select the **_UnrealVS.BuildStartupProject_** command in the commands list. 
+1. Select the **UnrealVS.BuildStartupProject** command in the commands list. 
 
 	[REGION:tip]
-		Type '`UnrealVS`' in the search box to filter the list.
+		Type "UnrealVS" in the search box to filter the list.
 	[/REGION]
 
 	![Filtered Command List](build_startup_search.png)
@@ -103,12 +103,12 @@ hotkey or other means of running custom commands.
 		combination is already in use.
 	[/REGION]
 
-1. Press the ![Assign](button_assign.png) button to bind the keys to the command. The keys are 
+1. Press the **Assign** button to bind the keys to the command. The keys are 
 	displayed in the **Shortcuts for selected command** box.
 
 	![Assigned Keys](build_startup_assigned.png)
 
-1. Press the ![OK](button_ok.png) button to save the changes. Now, when you use the shortcut, the project set as the startup project will automatically build.
+1. Press the **OK** button to save the changes. Now, when you use the shortcut, the project set as the startup project will automatically build.
 
 
 ## Command Line Arguments
@@ -138,7 +138,9 @@ file and run it manually.
 
 **To refresh project files:**
 
-1. Press the Refresh Projects ![Refresh Projects](button_refresh.png) button on the **UnrealVS** toolbar.
+1. Press the Refresh Projects Refresh button on the **UnrealVS** toolbar.
+
+![Refresh Projects](button_refresh.png)
 
 1. The project generation process will be displayed in the **Output** window as it occurs.
 
@@ -152,9 +154,11 @@ file and run it manually.
 
 **To open the Batch Builder window:**
 
-1. Press the Batch Builder ![Batch Builder](button_batchbuilder.png) button on the **UnrealVS** toolbar or use a keyboard shortcut that you have assigned to the command _UnrealVS.BatchBuilder_ 
+1. Press the Batch Builder button on the **UnrealVS** toolbar or use a keyboard shortcut that you have assigned to the command UnrealVS.BatchBuilder 
 (see **Build Startup Project** above for instructions on setting up keyboard shortcuts with **UnrealVS** commands).
 
+![Batch Builder](button_batchbuilder.png)
+
 1. The **UnrealVS Batch Builder** window appears.
 
 	![Batch Builder Window](batchbuilder_window.png)
@@ -163,11 +167,11 @@ file and run it manually.
 	* Add/remove jobs with the arrow buttons.
 	* Move selected jobs up/down the list with the arrow buttons.
 	* The controls edit the current **Job Set** displayed in the combo.
-	* To create a new **Job Set** enter a new name into the combo.
-	* You can select existing **Job Sets** using the dropdown list.
-	* The **Delete** button deletes the selected **Job Set** from the list.
-	* Start/Stop the batch build task on the current **Job Set** with the **Start** button.
-	* **Job Sets** are stored for next time in the .suo file for the loaded solution.
+	* To create a new Job Set enter a new name into the combo.
+	* You can select existing Job Sets using the dropdown list.
+	* The **Delete** button deletes the selected Job Set from the list.
+	* Start/Stop the batch build task on the current Job Set with the **Start** button.
+	* Job Sets are stored for next time in the .suo file for the loaded solution.
 
 1. When you click **Start**, the **Output** tab appears showing the progress of the batch build.
 
@@ -181,13 +185,13 @@ file and run it manually.
 		* The current executing build job is shown in bold
 	* Busy animation and stop button are shown when a batch build is running.
 	
-1. You can view the output from individual **Build Jobs** by **double-clicking** an item in the **Output tab** or **right-clicking** and selecting **Show Output** from the context menu.
+1. You can view the output from individual **Build Jobs** by **Double-clicking** an item in the **Output tab** or **Right-clicking** and selecting **Show Output** from the context menu.
 
 	![BatchBuild Output Pane](batchbuilder_outputpane.png)
 
 	[REGION:note]
-		Showing **Batch Builder** output for individual items is shown in the Visual Studio **Output** window in a pane named _UnrealVS - BatchBuild_.
-		This is not to be confused with the _Build_ pane that shows the output of the current/last build job.
+		Showing **Batch Builder** output for individual items is shown in the Visual Studio **Output** window in a pane named UnrealVS - BatchBuild.
+		This is not to be confused with the Build pane that shows the output of the current/last build job.
 	[/REGION]	
 	
 ## Quick Build Menu
@@ -198,7 +202,7 @@ The **Quick Build** menu provides an easy way to build a project with any config
 
 1. **Right-click** a Project node in the **Solution Explorer**.
 
-1. The *UnrealVS Quick Build* menu contains sub-menus for the available platforms and build configs in the Unreal Engine solution.
+1. The **UnrealVS Quick Build** menu contains sub-menus for the available platforms and build configs in the Unreal Engine solution.
 
 1. Select an item to build the selected project, platform, and config. This is much quicker than changing the Solution Configuration and Solution Platform in the IDE, starting a build, then switching
 the config and platform back again.
