#UE4doc #docs:
Changed a few lines as they were outdated
Changed font settings for multiple terms
Replaced an image

diff --git a/Engine/Documentation/Source/Programming/Introduction/IntroToProgramming.INT.udn b/Engine/Documentation/Source/Programming/Introduction/IntroToProgramming.INT.udn
--- a/Engine/Documentation/Source/Programming/Introduction/IntroToProgramming.INT.udn
+++ b/Engine/Documentation/Source/Programming/Introduction/IntroToProgramming.INT.udn
@@ -92,7 +92,7 @@ Now let us expose that same property to Blueprints.
 	UPROPERTY(EditAnywhere, BlueprintReadWrite, Category="Damage")
 	int32 TotalDamage;
 
-As you can see, there is a Blueprint specific parameter to make a property available for reading and writing. There's a separate option, BlueprintReadOnly, you can use if you want the property to be treated as *const* in Blueprints. There are quite a few options available for controlling how a property is exposed to the engine. To see more options, follow this [link.](Programming/UnrealArchitecture/Reference/Properties/Specifiers)
+As you can see, there is a Blueprint specific parameter to make a property available for reading and writing. There's a separate option, BlueprintReadOnly, you can use if you want the property to be treated as const in Blueprints. There are quite a few options available for controlling how a property is exposed to the engine. To see more options, follow this [link.](Programming/UnrealArchitecture/Reference/Properties/Specifiers)
 
 Before continuing to the section below, let us add a couple of properties to this sample class. There is already a property to control the total amount of damage this actor will deal out, but let us take that further and make that damage happen over time. The code below adds one designer settable property and one that is visible to the designer but not changeable by them.
 
@@ -113,7 +113,7 @@ Before continuing to the section below, let us add a couple of properties to thi
 		...
 	};
 
-	DamageTimeInSeconds is a property the designer can modify. The DamagePerSecond property is a calculated value using the designer's settings (see the next section). The VisibleAnywhere flag marks that property as viewable, but not editable in the Unreal Editor. The Transient flag means that it won't be saved or loaded from disk; it is meant to be a derived, non-persistent value. The image below shows the properties as part of the class defaults.
+DamageTimeInSeconds is a property the designer can modify. The DamagePerSecond property is a calculated value using the designer's settings (see the next section). The VisibleAnywhere flag marks that property as viewable, but not editable in the Unreal Editor. The Transient flag means that it won't be saved or loaded from disk; it is meant to be a derived, non-persistent value. The image below shows the properties as part of the class defaults.
 
 ![image alt text](image_3.png)
 
@@ -155,17 +155,17 @@ Here again is the same view of the properties after we have added the **PostInit
 
 Here is a cool feature of Unreal that you might be surprised about if you are used to programming C++ in other projects.  You can compile your C++ changes without shutting down the editor!  There are two ways to do this:
 
-1. With the editor still running, go ahead and **Build from Visual Studio or Xcode** like you normally would.  The editor will detect the newly compiled DLLs and reload your changes instantly!  
+1. With the editor still running, go ahead and build from Visual Studio or Xcode like you normally would.  The editor will detect the newly compiled DLLs and reload your changes instantly!  
 
 	![image alt text](image_6.png)
-
-	(Note that if you are attached with the debugger, you'll need to detach first so that Visual Studio will allow you to Build.)
-
+	[REGION:note]
+	If you are attached with the debugger, you'll need to detach first so that Visual Studio will allow you to Build.
+	[/REGION]
 2. Or, simply click the **Compile **button on the editor's main toolbar.
 
 	![image alt text](image_7.png)
 
-You can use this feature in the sections below as we advance through the tutorial. What a time saver!
+You can use this feature in the sections below as we advance through the tutorial.
 
 ### Extending a C++ Class via Blueprints
 
@@ -175,7 +175,7 @@ First thing we are going to do is create a new Blueprint class from our AMyActor
 
 ![image alt text](image_8.png)
 
-Once you choose Select, a new, default named Blueprint class is created for you. In this case, I set the name to CustomActor1 as you can see in the snapshot of the **Content Browser** below.
+Once you choose **Select**, a new, default named Blueprint class is created for you. In this case, I set the name to CustomActor1 as you can see in the snapshot of the **Content Browser** below.
 
 ![image alt text](image_9.png)
 
@@ -244,7 +244,7 @@ This version still generates the thunking method to call into the Blueprint VM.
 		// Do something cool here
 	}
 
-Now this version of the function is called when the Blueprint in question does not override the method. One thing to note, is that in future versions of the build tools the auto generated _Implementation() declaration will go away and you'll be expected to explicitly add that to the header. As of version 4.7, the auto generation of that declaration still occurs.
+Now this version of the function is called when the Blueprint in question does not override the method. One thing to note, is that in previous versions of the build tools the _Implementation() declaration was auto generated. In any version 4.8 or higher, you'll be expected to explicitly add that to the header.
 
 Now that we have walked through the common gameplay programmer workflow and methods to work with designers to build out gameplay features, it is time for you to choose your own adventure. You can either continue with this document to read more about how we use C++ in the engine or you can jump right into one of our samples that we include in the launcher to get a more hands on experience. 
 
@@ -389,9 +389,9 @@ You can limit the scope of the search by providing a more specific type to the i
 	{
 		// ...
 	}
-
-Warning: Using object iterators in PIE (Play In Editor) can lead to unexpected results.  Since the editor is loaded, the object iterator will return all UObjects created for your game world instance, in addition to those that are just being used by the editor.
-
+[REGION:warning]
+Using object iterators in PIE (Play In Editor) can lead to unexpected results.  Since the editor is loaded, the object iterator will return all UObjects created for your game world instance, in addition to those that are just being used by the editor.
+[/REGION]
 Actor iterators work in much the same way as object iterators, but only work for objects that derive from AActor.  Actor iterators do not have the problem noted below, and will only return objects being used by the current game world instance.
 
 When creating an actor iterator, you need to give it a pointer to a **UWorld** instance.  Many UObject classes, such as **APlayerController**, provide a **GetWorld** method to help you.  If you are not sure, you can check the **ImplementsGetWorld** method on a UObject to see if it implements the GetWorld method.
@@ -426,7 +426,7 @@ UE4 uses the reflection system to implement a garbage collection system.  With g
 		GENERATED_BODY()
 	};
 
-In the garbage collector, there is this concept called the root set.  This root set is basically a list of objects that the collector knows about will never be garbage collected.  An object will not be garbage collected as long as there is a path of references from an object in the root set to the object in question.  If no such path to the root set exists for an object, it is called *unreachable *and will be collected (deleted) the next time the garbage collector is ran.  The engine runs the garbage collector at certain intervals.
+In the garbage collector, there is this concept called the root set.  This root set is basically a list of objects that the collector knows about will never be garbage collected.  An object will not be garbage collected as long as there is a path of references from an object in the root set to the object in question.  If no such path to the root set exists for an object, it is called unreachable and will be collected (deleted) the next time the garbage collector is ran.  The engine runs the garbage collector at certain intervals.
 
 What counts as a "reference"?  Any UObject pointer stored in a UPROPERTY.  Let us start with a simple example.
 
@@ -439,7 +439,7 @@ When we call the above function, we create a new UObject, but we do not store a
 
 ### Actors and Garbage collection
 
-Actors are not usually garbage collected.  Once spawned, you must manually call Destroy() on them.  They will not be deleted immediately, and instead will be cleaned up during the next garbage collection phase.
+Actors are not usually garbage collected.  Once spawned, you must manually call **Destroy()** on them.  They will not be deleted immediately, and instead will be cleaned up during the next garbage collection phase.
 
 This is a more common case, where you have actors with UObject properties.
 
