#UE4doc #docs - Added engine version and Skill level fields. Removed Overview header.

diff --git a/Engine/Documentation/Source/Engine/Paper2D/Flipbooks/Paper2DFlipbooks.INT.udn b/Engine/Documentation/Source/Engine/Paper2D/Flipbooks/Paper2DFlipbooks.INT.udn
--- a/Engine/Documentation/Source/Engine/Paper2D/Flipbooks/Paper2DFlipbooks.INT.udn
+++ b/Engine/Documentation/Source/Engine/Paper2D/Flipbooks/Paper2DFlipbooks.INT.udn
@@ -7,6 +7,8 @@ Related: Engine\Paper2D\Flipbooks\Components
 Related: Engine/Paper2D/Sprites
 Related: Engine/Rendering/Materials
 Related: Engine/Content/Types/Textures
+Version: 4.9
+SkillLevel: Beginner
 
 [VAR:Topic]
 [OBJECT:Topic]
@@ -54,7 +56,6 @@ Related: Engine/Content/Types/Textures
 
 [TOC (start:2 end:3)]
 
-## Overview
 
 [EXCERPT:intro]
 The best way to think of **Paper 2D Flipbooks** (or **Flipbooks** for short) is in the form of hand-drawn animation, where a series of images that are slightly different are "flipped" through to produce what appears to be motion. In Unreal Engine 4, Flipbooks consist of a series of key frames, each of which contains a [Sprite](Engine\Paper2D\Sprites) to be displayed and a duration (in frames) to display it. 
