#UE4doc #docs - Added engine version field.

diff --git a/Engine/Documentation/Source/Resources/SampleGames/PlatformerGame/PlatformerGame.INT.udn b/Engine/Documentation/Source/Resources/SampleGames/PlatformerGame/PlatformerGame.INT.udn
--- a/Engine/Documentation/Source/Resources/SampleGames/PlatformerGame/PlatformerGame.INT.udn
+++ b/Engine/Documentation/Source/Resources/SampleGames/PlatformerGame/PlatformerGame.INT.udn
@@ -2,6 +2,7 @@ Availability: Public
 Title:Platformer Game Sample
 Crumbs:%ROOT%, Resources
 Description:Overview of the Platformer Sample Game.
+Version: 4.9
 
 
 [VAR:Topic]
@@ -52,7 +53,7 @@ Description:Overview of the Platformer Sample Game.
 
 [TOC(start:2)]
 
-The PlatformerGame sample is an example of a PC/mobile side-scroller game. It includes basic implementations of 
+The **Platformer Game** sample is an example of a PC/mobile side-scroller game. It includes basic implementations of 
 different movement types along with a simple front end menu system.
 
 A complete list of the featured concepts:
