UEDOC-2567

Updated highlighted line and fixed typo

diff --git a/Engine/Documentation/Source/Engine/QuickStart/1/LDQS_1.INT.udn b/Engine/Documentation/Source/Engine/QuickStart/1/LDQS_1.INT.udn
--- a/Engine/Documentation/Source/Engine/QuickStart/1/LDQS_1.INT.udn
+++ b/Engine/Documentation/Source/Engine/QuickStart/1/LDQS_1.INT.udn
@@ -26,7 +26,7 @@ Related:GettingStarted/Installation
 %Steps%
 
 Before we get started, the first thing that we will need to do is create a [Project](Engine/Basics/Projects) to work from for the duration of this guide. In the steps below we will use Unreal Engine 4's [Project Browser](Engine/Basics/Projects/Browser) to create our project. 
-After completing this step, our project will be ready and you will gone through the process of creating a Project and using a Template which you can then use for your very own future projects. 
+After completing this step, you will have a Template that you can build onto for your future projects.  
 
 ## Steps
 
