#UE4doc #docs - Added engine version field. Changed font settings for some terms.

diff --git a/Engine/Documentation/Source/Resources/SampleGames/ShooterGame/ShooterGame.INT.udn b/Engine/Documentation/Source/Resources/SampleGames/ShooterGame/ShooterGame.INT.udn
--- a/Engine/Documentation/Source/Resources/SampleGames/ShooterGame/ShooterGame.INT.udn
+++ b/Engine/Documentation/Source/Resources/SampleGames/ShooterGame/ShooterGame.INT.udn
@@ -14,6 +14,7 @@ Related: Engine/Audio/Overview
 Related: Programming/Slate
 Related: Gameplay/Input
 Related: Engine/Rendering/ParticleSystems
+Version: 4.9
 
 
 
diff --git a/Engine/Documentation/Source/Resources/SampleGames/StrategyGame/StrategyGame.INT.udn b/Engine/Documentation/Source/Resources/SampleGames/StrategyGame/StrategyGame.INT.udn
--- a/Engine/Documentation/Source/Resources/SampleGames/StrategyGame/StrategyGame.INT.udn
+++ b/Engine/Documentation/Source/Resources/SampleGames/StrategyGame/StrategyGame.INT.udn
@@ -10,6 +10,7 @@ Related: Programming/Slate
 Related: Gameplay/Framework/UIAndHUD
 Related: Gameplay/Input
 Related: Engine/Rendering/ParticleSystems
+Version: 4.9
 
 
 
@@ -56,7 +57,7 @@ Related: Engine/Rendering/ParticleSystems
 
 [TOC(start:2 end:2)]
 
-The Tower Defense sample is an example of a RTS/Tower Defense game. 
+The **Tower Defense** sample is an example of a RTS/Tower Defense game. 
 
 A complete list of the featured concepts:
 
@@ -83,23 +84,23 @@ selected, and the current action can be stopped if there is an action with highe
 Both enemy and friendly Pawns operate with the AI logic, moving towards the opposing base and attacking Pawns of the other team if they encounter them.  While players cannot control the movement
 or behavior of their friendly Pawns, they are able to buy new units to spawn.  
 
-_Blueprints_ are also used to add logic to Minion Pawns.  Both friendly and enemy Pawns can be equipped with shields; friendly Pawns get shields if the armory upgrade has been purchased for the brewery,
+Blueprint are also used to add logic to Minion Pawns.  Both friendly and enemy Pawns can be equipped with shields; friendly Pawns get shields if the armory upgrade has been purchased for the brewery,
 and enemy Pawns get shields if they were spawned due to the SpawnHeavyFunction or SpawnEndBossFunction being called in the _Level Blueprint_. If a Pawn has a shield, projectiles from the auto-arbalest will be destroyed and do no damage. This logic is carried out using _Blueprint Interfaces_.  The
-**Minion** _Blueprint_ also contains a network to slow down enemy Pawns when they are hit by a charged-up fire turret.  
+**Minion** Blueprint also contains a network to slow down enemy Pawns when they are hit by a charged-up fire turret.  
 
 ## Building Construction
 
 There are two building classes in Tower Defense - `StrategyBuilding` and `StrategyBuilding_Brewery`.  All the turret types in Tower Defense, as well as the empty building slot, are designed using 
-_Blueprints_ with `StrategyBuilding` as the parent class.  Players can click on an empty building slot to show the context menu and select a new building to construct.  When a building is constructed, the empty building slot is destroyed and the new building is spawned.  
+Blueprint with `StrategyBuilding` as the parent class.  Players can click on an empty building slot to show the context menu and select a new building to construct.  When a building is constructed, the empty building slot is destroyed and the new building is spawned.  
 
 There is a mechanism in place for upgrading buildings as well. The `StrategyBuilding_Brewery` class implements this case so that upgrades are built in connected slots near the brewery base.
 
-Again, the code in Tower Defense simply creates the base building classes. All logic and design of the buildings in the Tower Defense were made by level designers in _Blueprints_.
+Again, the code in Tower Defense simply creates the base building classes. All logic and design of the buildings in the Tower Defense were made by level designers in Blueprint.
 
 ### Brewery
 
-The **Brewery** _Blueprint_ has the parent class `StrategyBuilding_Brewery`, and also includes an AIDirector component.  There are two breweries placed in the TowerDefenseMap, one for the enemy where enemy
-Pawns spawn, and one friendly brewery where an armory and a smithy upgrade can be built, and friendly Pawns can be spawned.  There is no graph logic present in the **Brewery** _Blueprint_, just Defaults set for the 
+The **Brewery** Blueprint has the parent class `StrategyBuilding_Brewery`, and also includes an AIDirector component.  There are two breweries placed in the TowerDefenseMap, one for the enemy where enemy
+Pawns spawn, and one friendly brewery where an armory and a smithy upgrade can be built, and friendly Pawns can be spawned.  There is no graph logic present in the **Brewery** Blueprint, just Defaults set for the 
 building properties and a Components list including the AIDirector, a TriggerBox, and a _Static Mesh_. 
 
 #### Upgrades
@@ -109,12 +110,12 @@ building properties and a Components list including the AIDirector, a TriggerBox
 There are two upgrade slots attached to the friendly brewery.  These slots are _Blueprint Classes_ which are also derived from the `StrategyBuilding` class.  If an upgrade is selected from the Brewery menu,
 one slot will be replaced by that upgrade. Only one smithy upgrade and one armory upgrade can be purchased. 
 
-Once a smithy upgrade is purchased, the build will be started, firing the OnBuildStarted event in the **Wall_Smithy** _Blueprint_.  This _Blueprint_ also informs the system that the upgrade has been built,
-once the build completes.  After this point, any friendly Pawns that are spawned will be equipped with a shield attachment, a _Blueprint_ derived from the `StrategyAttachment` class.  The network which assigns the 
+Once a smithy upgrade is purchased, the build will be started, firing the OnBuildStarted event in the **Wall_Smithy** Blueprint.  This Blueprint also informs the system that the upgrade has been built,
+once the build completes.  After this point, any friendly Pawns that are spawned will be equipped with a shield attachment, a Blueprint derived from the `StrategyAttachment` class.  The network which assigns the 
 "shield-attaching" behavior after the armory reports that it has been constructed is present in the PlayerBaseUpgrades collapsed graph in the TowerDefenseMap _Level Blueprint_. The `StrategyAttachment` 
-class simply contains a  `SkeletalMeshComponent`; the mesh and attachment point for the attachment are set in the Defaults for the **Attachment_Armorer** _Blueprint_.
+class simply contains a  `SkeletalMeshComponent`; the mesh and attachment point for the attachment are set in the Defaults for the **Attachment_Armorer** Blueprint.
 
-The armory _Blueprint_ contains the same logic setup with OnBuildStarted and OnBuildFinished events.  After construction of the armory, any friendly Pawns spawned will be equipped with a hammer, also derived from the 
+The armory Blueprint contains the same logic setup with OnBuildStarted and OnBuildFinished events.  After construction of the armory, any friendly Pawns spawned will be equipped with a hammer, also derived from the 
 `StrategyAttachment` class.  The network which assigns the "hammer-attaching" behavior after the smithy reports that it has been constructed is also present in the PlayerBaseUpgrades collapsed graph in the TowerDefenseMap _Level Blueprint_.
 
 ### Turrets
@@ -123,21 +124,21 @@ The armory _Blueprint_ contains the same logic setup with OnBuildStarted and OnB
 
 ![](building_empty_slot.png)
 
-The empty slot is also a _Blueprint_ with the `StrategyBuilding` parent class, Wall_EmptySlot. No logic is present in the _Blueprint_ graphs. This is a _Blueprint Class_ with Defaults set for building 
+The empty slot is also a Blueprint with the `StrategyBuilding` parent class, Wall_EmptySlot. No logic is present in the Blueprint graphs. This is a _Blueprint Class_ with Defaults set for building 
 properties and Static Meshes and a trigger box set as the Components.
 
-The possible turret upgrades are all set in the Defaults of the Wall_EmptySlot _Blueprint_, in the Upgrades section of the Building category.  
+The possible turret upgrades are all set in the Defaults of the Wall_EmptySlot Blueprint, in the Upgrades section of the Building category.  
 
 
 #### Arbalest
 
 ![](building_arbalest_shooting.png)
 
-The **Wall_arbalest** _Blueprint_ contains the logic for the arbalest, the basic turret type. The arbalest shoots the closest enemy with a medium-strength projectile, automatically firing arrows in its default mode. The player
+The **Wall_arbalest** Blueprint contains the logic for the arbalest, the basic turret type. The arbalest shoots the closest enemy with a medium-strength projectile, automatically firing arrows in its default mode. The player
 can manually fire the arbalest as well, by clicking on it and dragging in the direction they want the arrow to fire.  The longer the mouse drag is, the stronger the shot will be.  
 
 
-The arbalest projectile is stored in another _Blueprint_, Projectile_arbalest, derived from the _Blueprint_ TestProjectile which has `StrategyProjectile`.  The Wall_arbalest _Blueprint_ has a number of sub-networks, all contained within the EventGraph.  No _Blueprint_ logic is present in the ConstructionScript.   
+The arbalest projectile is stored in another Blueprint, Projectile_arbalest, derived from the Blueprint TestProjectile which has `StrategyProjectile`.  The Wall_arbalest Blueprint has a number of sub-networks, all contained within the EventGraph.  No Blueprint logic is present in the ConstructionScript.   
 
 
 
@@ -145,11 +146,11 @@ The arbalest projectile is stored in another _Blueprint_, Projectile_arbalest, d
 
 ![](building_auto-arbalest_shooting.png)
 
-The **Wall_arbalest_auto** _Blueprint_ contains the logic for the auto-arbalest.  The auto-arbalest shoots projectiles in a straight line from the wall, doing a small amount of damage to every unit the projectiles
+The **Wall_arbalest_auto** Blueprint contains the logic for the auto-arbalest.  The auto-arbalest shoots projectiles in a straight line from the wall, doing a small amount of damage to every unit the projectiles
 pass through.  The auto-arbalest arrows are not destroyed unless they hit a wall or an enemy with a shield.  It is possible to aim the auto-arbalest by clicking and dragging so that the auto-arbalest faces the desired direction;
 the auto-arbalest will continue to fire in the aimed direction while the mouse button is held down, and will return to its default firing position when the mouse button is released. 
 
-Like the arbalest, this turret shoots projectile arrows contained in a separate _Blueprint_.  The Projectile_arbalest_auto auto-arbalest arrows are not destroyed unless they hit an enemy Pawn with a shield or a wall, and this behavior
+Like the arbalest, this turret shoots projectile arrows contained in a separate Blueprint.  The Projectile_arbalest_auto auto-arbalest arrows are not destroyed unless they hit an enemy Pawn with a shield or a wall, and this behavior
 is carried out with the aid of _Blueprint Interfaces_, Interface_Auto_Arbalest and Interface_Auto_Projectile. 
 
 #### Flamethrower
@@ -239,11 +240,11 @@ The waves are constructed using three _Blueprint Macros_: spawn fast, spawn norm
 
 Each spawn macro takes two execution inputs, one for beginning the macro and one for opening the **Gate** after the OnWaveSpawned event fires, and an integer input for the number of Pawns to spawn.  
 Functions from the class `StrategyAIDirector` are called, with inputs specific to each type of Pawn wave.  The three functions are `SetDefaultWeapon`, `SetDefaultArmor`, and `SetBuffModifier`.  
-`SetDefaultWeapon` and `SetDefaultArmor` take _Blueprints_ as inputs and assign those _Blueprints_ as the new default weapon or default armor for the spawn.  For example, all enemy Pawns spawned by 
-the SpawnFastMacro have the Attachment_Smithy hammer _Blueprint_ as their default weapon, and all enemy Pawns spawned by the SpawnHeavyMacro have the Attachment_Armorer shield _Blueprint_ as their default armor. 
+`SetDefaultWeapon` and `SetDefaultArmor` take Blueprint as inputs and assign those Blueprint as the new default weapon or default armor for the spawn.  For example, all enemy Pawns spawned by 
+the SpawnFastMacro have the Attachment_Smithy hammer Blueprint as their default weapon, and all enemy Pawns spawned by the SpawnHeavyMacro have the Attachment_Armorer shield Blueprint as their default armor. 
 
-The last `StrategyAIDirector` function called by the spawning _Blueprint_ functions is `SetBuffModifier`, which has a number of data inputs including the attack abilities, health bonus, speed, and size
-of the Pawns.  These inputs are all exposed to the _Blueprint_, so it is straightforward for a level designer to create a new class of enemy Pawn to spawn.  Finally, each spawning _Blueprint_ function sets the `WaveSize`
+The last `StrategyAIDirector` function called by the spawning Blueprint functions is `SetBuffModifier`, which has a number of data inputs including the attack abilities, health bonus, speed, and size
+of the Pawns.  These inputs are all exposed to the Blueprint, so it is straightforward for a level designer to create a new class of enemy Pawn to spawn.  Finally, each spawning Blueprint function sets the `WaveSize`
 property of the enemy brewery's `StrategyAIDirector`.   
 
 There are five waves of enemies, each with different combinations of fast, normal, and heavy enemy Pawns.  At the beginning of a wave, the Show Wave Title node displays the wave number.  Then, the first enemy 
@@ -272,10 +273,10 @@ NumberOfLives of the friendly brewery, decreasing the value by one each time. Th
 
 ## Resource Node - Gold
 
-The gold nodes are _Blueprints_ with the parent class `StrategyResourceNode`.  This class contains the public functions `GetAvailableResources` and `GetInitialResources`, the protected function `OnDepleted`, 
+The gold nodes are Blueprint with the parent class `StrategyResourceNode`.  This class contains the public functions `GetAvailableResources` and `GetInitialResources`, the protected function `OnDepleted`, 
 a `BlueprintImplementableEvent` to report when the resource has been depleted, and the protected property `NumResources` to set the amount of the resource present in the node.  
 
-The _Blueprint_ for a gold node contains sub-networks to make the node appear and disappear on a timer.  The **ConstructionScript** in the _Blueprint_ sets the node to automatically be hidden when it is placed in the level.
+The Blueprint for a gold node contains sub-networks to make the node appear and disappear on a timer.  The **ConstructionScript** in the Blueprint sets the node to automatically be hidden when it is placed in the level.
 When the gold node appears, AppearFX particle effects are played along with an apparition noise. If the OnDepleted event fires because the node is successfully collected, the amount of gold present in the node
 is added to the player's total gold pool.  The CollectFX particle effects and the CoinSound is played, and then the node is hidden again.  If the player fails to collect the node in time by clicking on it,
 FadeFX particle effects and an appropriate sound is played.  
