#UE4doc #docs - Added engine version field.

diff --git a/Engine/Documentation/Source/Resources/Showcases/Showcases.INT.udn b/Engine/Documentation/Source/Resources/Showcases/Showcases.INT.udn
--- a/Engine/Documentation/Source/Resources/Showcases/Showcases.INT.udn
+++ b/Engine/Documentation/Source/Resources/Showcases/Showcases.INT.udn
@@ -2,6 +2,7 @@ Availability: Public
 Title:Engine Feature Examples
 Crumbs:%ROOT%, Resources
 Description:Full scenes demonstrating particular features or approaches to specific challenges.
+Version: 4.9
 
 [VAR:Topic]
 [OBJECT:Topic]
