#UE4doc #docs - https://docs.unrealengine.com/latest/INT/Resources/Showcases/

diff --git a/Engine/Documentation/Source/Resources/Showcases/BlueprintExamples/BlueprintExamples.INT.udn b/Engine/Documentation/Source/Resources/Showcases/BlueprintExamples/BlueprintExamples.INT.udn
--- a/Engine/Documentation/Source/Resources/Showcases/BlueprintExamples/BlueprintExamples.INT.udn
+++ b/Engine/Documentation/Source/Resources/Showcases/BlueprintExamples/BlueprintExamples.INT.udn
@@ -8,6 +8,7 @@ Related: Resources/ContentExamples/Blueprints
 Related: Gameplay/Framework/Camera
 Related: Engine/Rendering/ParticleSystems
 Related: Engine/Rendering/Materials
+Version: 4.9
 
 
 [VAR:Topic]
