#UE4doc #docs -  Added engine version field.

diff --git a/Engine/Documentation/Source/Resources/Showcases/BlueprintExamples/Butterfly/Butterfly.INT.udn b/Engine/Documentation/Source/Resources/Showcases/BlueprintExamples/Butterfly/Butterfly.INT.udn
--- a/Engine/Documentation/Source/Resources/Showcases/BlueprintExamples/Butterfly/Butterfly.INT.udn
+++ b/Engine/Documentation/Source/Resources/Showcases/BlueprintExamples/Butterfly/Butterfly.INT.udn
@@ -8,6 +8,7 @@ Related: Engine/Animation/AnimBlueprints/EventGraph
 Related: Engine/Blueprints/UserGuide/UserConstructionScript
 Related: Engine/Animation/AnimBlueprints/EventGraph
 Related: Gameplay/AI
+Version: 4.9
 
 [REGION:fullwidth]
 ![](ButterflyHeader.png)
