#UE4doc #docs - Added engine version field.

diff --git a/Engine/Documentation/Source/Resources/Showcases/BlueprintExamples/CommentaryBox/CommentaryBox.INT.udn b/Engine/Documentation/Source/Resources/Showcases/BlueprintExamples/CommentaryBox/CommentaryBox.INT.udn
--- a/Engine/Documentation/Source/Resources/Showcases/BlueprintExamples/CommentaryBox/CommentaryBox.INT.udn
+++ b/Engine/Documentation/Source/Resources/Showcases/BlueprintExamples/CommentaryBox/CommentaryBox.INT.udn
@@ -7,6 +7,7 @@ Related: Engine/Animation/AnimBlueprints/EventGraph
 Related: Engine/Actors/Triggers
 Related: Engine/Blueprints
 Related: Engine/Content/Types/SkeletalMeshes/Sockets
+Version: 4.9
 
 [REGION:fullwidth]
 ![](BlueprintCommentary.png)
diff --git a/Engine/Documentation/Source/Resources/Showcases/BlueprintExamples/FogSheet/FogSheet.INT.udn b/Engine/Documentation/Source/Resources/Showcases/BlueprintExamples/FogSheet/FogSheet.INT.udn
--- a/Engine/Documentation/Source/Resources/Showcases/BlueprintExamples/FogSheet/FogSheet.INT.udn
+++ b/Engine/Documentation/Source/Resources/Showcases/BlueprintExamples/FogSheet/FogSheet.INT.udn
@@ -6,6 +6,7 @@ Related: Engine/Blueprints/Overview
 Related: Engine/Blueprints/UserGuide/Variables
 Related: Engine/Blueprints/UserGuide/UserConstructionScript
 Related: Engine/Rendering/Materials/MaterialInstances
+Version: 4.9
 
 
 
@@ -33,7 +34,7 @@ The Fog Sheet Blueprint, shown above, consists of a Static Mesh plane with a clo
 [![](FogSheetCS.png)](FogSheetCS.png)
 [/REGION]
 [/REGION]
-Each of the 3 aspects of the Blueprint (i.e. the _Components List_, the _Construction Script_, and the _Event Graph_) are responsible for a different part of the overall effect. The following table breaks down the jobs for each aspect and how they work together to create the final result.
+Each of the three aspects of the Blueprint (i.e. the _Components List_, the _Construction Script_, and the _Event Graph_) are responsible for a different part of the overall effect. The following table breaks down the jobs for each aspect and how they work together to create the final result.
 
 
 
