#UE4doc #docs - Added engine version field. Updated image call out color.

diff --git a/Engine/Documentation/Source/Resources/Showcases/BlueprintExamples/Foliage/Foliage.INT.udn b/Engine/Documentation/Source/Resources/Showcases/BlueprintExamples/Foliage/Foliage.INT.udn
--- a/Engine/Documentation/Source/Resources/Showcases/BlueprintExamples/Foliage/Foliage.INT.udn
+++ b/Engine/Documentation/Source/Resources/Showcases/BlueprintExamples/Foliage/Foliage.INT.udn
@@ -6,6 +6,7 @@ Related: Engine/Blueprints/Overview
 Related: Engine/Blueprints/UserGuide/Variables
 Related: Engine/Animation/AnimBlueprints/EventGraph
 Related: Engine/Blueprints/UserGuide/UserConstructionScript
+Version: 4.9
 
 
 [REGION:fullwidth]
@@ -13,7 +14,7 @@ Related: Engine/Blueprints/UserGuide/UserConstructionScript
 [/REGION]
 
 [EXCERPT:Overview]
-The Foliage Blueprint shows how to use a Blueprint to create randomly-generated foliage (or potentially other structures) that are automatically placed such that they conform to a surface. The user has control over the extent of the effect and how many meshes the system will try to spawn. As the Blueprint is moved, the meshes update to re-conform to the new environment.
+The **Foliage** Blueprint shows how to use a Blueprint to create randomly-generated foliage (or potentially other structures) that are automatically placed such that they conform to a surface. The user has control over the extent of the effect and how many meshes the system will try to spawn. As the Blueprint is moved, the meshes update to re-conform to the new environment.
 [/EXCERPT:Overview]
 
 
@@ -28,7 +29,7 @@ The Components List, Construction Script, and Event Graph for this Blueprint per
 |---|---|
 | **Components List** | Adds a BillboardComponent to represent the center and a BoxComponent to help visualize the construction radius. |
 | **Construction Script** | [INCLUDE:#FoliageCS] |
-| **Event Graph** | This Blueprint has no **Event Graph**, meaning that once its properties are set, it remains the same throughout gameplay. |
+| **Event Graph** | This Blueprint has no Event Graph, meaning that once its properties are set, it remains the same throughout gameplay. |
 
 [REGION:lightbox]
 ![](FoliageCS.png)(w:920)
@@ -46,12 +47,12 @@ As you might expect, most of the work of this Blueprint is done in the **Constru
 
 | Property | Purpose |
 | --- | --- |
-| **Number of Meshes** | This integer value is the number of meshes the Blueprint will attempt to spawn. In the **Construction Script**, this value determines the number of times a _ForLoop_ node will trigger. Each execution coming out of the _ForLoop_ node passes through the section of script responsible for creating the Static Meshes.  |
+| **Number of Meshes** | This integer value is the number of meshes the Blueprint will attempt to spawn. In the Construction Script, this value determines the number of times a _ForLoop_ node will trigger. Each execution coming out of the _ForLoop_ node passes through the section of script responsible for creating the Static Meshes.  |
 | **Radius** | Controls the radius of the effect, within which the foliage will be randomly scattered. |
 | **Collision Trace Range** | This is the vertical distance (how far down the system will look) to find a valid surface on which to spawn foliage. |
-| **Meshes To Spawn** | This is an array variable, which meshes from the **Content Browser** can be added to. After the **Construction Script** spawns a mesh component, it is changed to a mesh picked at random from this array. |
+| **Meshes To Spawn** | This is an array variable, which meshes from the Content Browser can be added to. After the Construction Script spawns a mesh component, it is changed to a mesh picked at random from this array. |
 | **Vertical Offset** | Controls how far _into_ a surface a mesh will be spawned. This can help meshes appear more grounded. |
-| **Align To Surface** | This Boolean variable causes meshes spawned by the Blueprint will orient themselves such that they match the angle of the surface upon which they are resting. Internally, this variable affects a branch in the **Construction Script** that determines the method used for positioning the foliage meshes. When enabled, the script draws a line trace before attempting to spawn each mesh. When the trace hits a surface, the angle and height of the impact point is retrieved and used to determine the angle and height of the mesh being spawned. |
+| **Align To Surface** | This Boolean variable causes meshes spawned by the Blueprint will orient themselves such that they match the angle of the surface upon which they are resting. Internally, this variable affects a branch in the Construction Script that determines the method used for positioning the foliage meshes. When enabled, the script draws a line trace before attempting to spawn each mesh. When the trace hits a surface, the angle and height of the impact point is retrieved and used to determine the angle and height of the mesh being spawned. |
 | **Scale Max** | Maximum scale allowed for the spawned meshes. This is a 0.0-1.0 percentage value, with 1.0 representing 100%. |
 | **Scale Min** | Minimum scale allowed for the spawned meshes. This is a 0.0-1.0 percentage value, with 1.0 representing 100%. |
 | **Cluster Around Center** | This Boolean variable determines whether or not the randomized locations for the meshes will be primarily located around the center of the radius. When this option is enabled, the distance meshes can spawn from the center point is reduced by a random amount, causing fewer meshes to spawn the farther away the system looks from the radius. |
