#UE4doc #docs - Added engine version field.

diff --git a/Engine/Documentation/Source/Resources/Showcases/BlueprintExamples/SecurityCameraAndAlarm/SecurityCameraAndAlarm.INT.udn b/Engine/Documentation/Source/Resources/Showcases/BlueprintExamples/SecurityCameraAndAlarm/SecurityCameraAndAlarm.INT.udn
--- a/Engine/Documentation/Source/Resources/Showcases/BlueprintExamples/SecurityCameraAndAlarm/SecurityCameraAndAlarm.INT.udn
+++ b/Engine/Documentation/Source/Resources/Showcases/BlueprintExamples/SecurityCameraAndAlarm/SecurityCameraAndAlarm.INT.udn
@@ -7,6 +7,7 @@ Related: Engine/Blueprints/UserGuide/Variables
 Related: Engine/Animation/AnimBlueprints/EventGraph
 Related: Engine/Blueprints/UserGuide/UserConstructionScript
 Related: Engine/Blueprints/UserGuide/Components
+Version: 4.9
 
 
 [EXCERPT:Overview]
@@ -18,7 +19,7 @@ Tied into the existing camera Blueprints is a security control room. This room c
 ![](ScreenBlueprint.png)(w:600)
 
 
-The Screen Blueprint is - at its heart - an optimized mechanism to switch between different Scene Capture Actors. The effect of the security cameras is created because of Scene Capture Actors located within each of the [Camera Blueprints](Resources\Showcases\BlueprintExamples\SecurityDoor). The purpose of a Scene Capture Actor is to function very much like a camera, capturing an image of the pixels it can see. It then transmits those to a Render Target texture, which functions like an internal video screen. By applying this Render Target texture to a Material, you effectively create a Material that can show what the camera sees.
+The **Screen** Blueprint is - at its heart - an optimized mechanism to switch between different Scene Capture Actors. The effect of the security cameras is created because of Scene Capture Actors located within each of the [Camera Blueprints](Resources\Showcases\BlueprintExamples\SecurityDoor). The purpose of a Scene Capture Actor is to function very much like a camera, capturing an image of the pixels it can see. It then transmits those to a Render Target texture, which functions like an internal video screen. By applying this Render Target texture to a Material, you effectively create a Material that can show what the camera sees.
 
 The catch? With many cameras running at the same time, it is possible to run into performance issues, since they would all have to be capturing and transmitting their current views. To alleviate this problem, the screen Blueprint will activate only one Scene Capture at a time, and deactivate the one that was used previously.
 
