#UE4doc #docs - added engine version field. Centered images.

diff --git a/Engine/Documentation/Source/Resources/Showcases/BlueprintExamples/SecurityDoor/SecurityDoor.INT.udn b/Engine/Documentation/Source/Resources/Showcases/BlueprintExamples/SecurityDoor/SecurityDoor.INT.udn
--- a/Engine/Documentation/Source/Resources/Showcases/BlueprintExamples/SecurityDoor/SecurityDoor.INT.udn
+++ b/Engine/Documentation/Source/Resources/Showcases/BlueprintExamples/SecurityDoor/SecurityDoor.INT.udn
@@ -9,8 +9,11 @@ Related: Engine/Blueprints/UserGuide/UserConstructionScript
 Related: Engine/Blueprints/UserGuide/Components
 Related: Engine/Blueprints/UserGuide/Timelines
 Related: Gameplay/Framework/Pawn
+Version: 4.9
 
+<div align= "center">
 ![](SecurityCameraAndDoorHeader.png)(w:600)
+</div>
 
 [EXCERPT:Overview]
 The security camera Blueprints seen throughout the level are designed to interact with the player when in Play mode. During play, they will pan back and forth in a search pattern. When the green search cone overlaps a pawn, the light turns red and the camera begins to track player movement. It continues to track the player until the player moves out of range. Then, after a few seconds, it returns to searching.
@@ -79,8 +82,9 @@ The camera has the following exposed properties:
 
 
 ## Alarm Blueprint
-
+<div align= "center">
 ![](AlarmImage.png)
+</div>
 
 The Alarm Blueprint serves as a visual cue that the player has tripped the alarm by stepping in front of the security camera. When activated, it seems to change from green to red, and a beacon light begins spinning. It continues to do this until the alarm deactivates.
 
@@ -116,8 +120,9 @@ The alarm has the following exposed properties:
 
 
 ## Door Blueprint
-
+<div align= "center">
 ![](Door.png)
+</div>
 
 The Door is set up such that it remains open so long as the alarm is not engaged. The moment the alarm is triggered, however, the doors are animated such that they slam shut and will not reopen until the alarm ends. It is important to note that Timelines are used to animate the doors. In fact, there are 2 separate Timelines - one to open the doors and another to close them. This allows the doors to open at one speed and close at another. 
 
