#UE4doc #docs - Added engine version field. Updated image call out color.

diff --git a/Engine/Documentation/Source/Resources/Showcases/BlueprintMultiplayer/BlueprintMultiplayer.INT.udn b/Engine/Documentation/Source/Resources/Showcases/BlueprintMultiplayer/BlueprintMultiplayer.INT.udn
--- a/Engine/Documentation/Source/Resources/Showcases/BlueprintMultiplayer/BlueprintMultiplayer.INT.udn
+++ b/Engine/Documentation/Source/Resources/Showcases/BlueprintMultiplayer/BlueprintMultiplayer.INT.udn
@@ -8,6 +8,7 @@ Related: Programming/Online
 Related: Programming/Online/Steam
 Related: Gameplay/Networking/Actors
 Related: Resources/ContentExamples/Networking
+Version: 4.9
 
 
 [VAR:Topic]
