#UE4doc #docs - Added engine version field. Added/Removed bolded terms.

diff --git a/Engine/Documentation/Source/Resources/Showcases/Effects/Effects.INT.udn b/Engine/Documentation/Source/Resources/Showcases/Effects/Effects.INT.udn
--- a/Engine/Documentation/Source/Resources/Showcases/Effects/Effects.INT.udn
+++ b/Engine/Documentation/Source/Resources/Showcases/Effects/Effects.INT.udn
@@ -7,6 +7,7 @@ Related: Engine/Rendering/ParticleSystems
 Related: Engine/Rendering/ParticleSystems\Cascade
 Related: Engine/Rendering/ParticleSystems\Reference\TypeData\GPUSprites
 Related: Engine/Rendering/ParticleSystems/Reference
+Version: 4.9
 
 
 
