#UE4doc #docs - Added engine version field. Added/Removed bolded terms.

diff --git a/Engine/Documentation/Source/Resources/Showcases/Effects/Effects.INT.udn b/Engine/Documentation/Source/Resources/Showcases/Effects/Effects.INT.udn
--- a/Engine/Documentation/Source/Resources/Showcases/Effects/Effects.INT.udn
+++ b/Engine/Documentation/Source/Resources/Showcases/Effects/Effects.INT.udn
@@ -64,7 +64,7 @@ The Particle Effects Example shows off a variety of particle systems that demons
 
 ## Scene Breakdown
 
-The Particle Effects Example showcases a cavern level, high up in a snowy arctic mountain range. Within the cave, water cascades across moss-covered rock formations, trickling into rich, flowing pools on the cavern floor. Mysterious and seemingly magical flame sources provide an eerie light. Outside, a harsh blizzard rips through the rocky crags. 
+This example showcases a cavern level, high up in a snowy arctic mountain range. Within the cave, water cascades across moss-covered rock formations, trickling into rich, flowing pools on the cavern floor. Mysterious and seemingly magical flame sources provide an eerie light. Outside, a harsh blizzard rips through the rocky crags. 
 
 
 ## Particle System Examples
@@ -119,7 +119,7 @@ The documentation included with this example is intended to help you get off the
 
 ![](FindInCB.png)
 
-This will select the Particle System in the **Content Browser**. **Double-click** its icon in the **Content Browser** to open the effect up in **Cascade**, Unreal Engine 4's powerful particle editing system.
+This will select the Particle System in the Content Browser. **Double-click** its icon in the Content Browser to open the effect up in **Cascade**, Unreal Engine 4's powerful particle editing system.
 
 ![](CascadeEditorOpen.png)(w:600)
 
