#UE4doc #docs:
Added engine version field

diff --git a/Engine/Documentation/Source/Programming/Modules/API/ModuleAPI.INT.udn b/Engine/Documentation/Source/Programming/Modules/API/ModuleAPI.INT.udn
--- a/Engine/Documentation/Source/Programming/Modules/API/ModuleAPI.INT.udn
+++ b/Engine/Documentation/Source/Programming/Modules/API/ModuleAPI.INT.udn
@@ -2,6 +2,7 @@ Availability:Public
 Title:Module API Specifiers
 Crumbs:%ROOT%, Programming
 Description:Exposing classes and functions to other modules.
+version: 4.9
 
 The easiest way to think about these specifiers is that they are used to tag functions, classes or data as public to your module's DLL file. If you mark a function in the Engine module as ENGINE_API, then any module that imports Engine can access that function directly.
 
