#UE4doc #docs:
Changed font settings for multiple terms.

diff --git a/Engine/Documentation/Source/Programming/Modules/Gameplay/GameplayModules.INT.udn b/Engine/Documentation/Source/Programming/Modules/Gameplay/GameplayModules.INT.udn
--- a/Engine/Documentation/Source/Programming/Modules/Gameplay/GameplayModules.INT.udn
+++ b/Engine/Documentation/Source/Programming/Modules/Gameplay/GameplayModules.INT.udn
@@ -16,13 +16,13 @@ If your game project is not already created, see [Creating a Game Project](Engin
 
 ## Module Creation
 
-Gameplay modules must contain, at the very least, a header file (`.h`), C++ file (`.cpp`), and a build file (`*.Build.cs`). The header file must be located in the Public folder of the module's directory, i.e. [GameName]\Source\[ModuleName]\Public. This file includes any header files - including the module's auto-generated headers - necessary to compile the classes contained within the module.
+Gameplay modules must contain, at the very least, a header file (.h), C++ file (.cpp), and a build file (*.Build.cs). The header file must be located in the Public folder of the module's directory, i.e. [GameName]\Source\[ModuleName]\Public. This file includes any header files - including the module's auto-generated headers - necessary to compile the classes contained within the module.
 
 	#include "Engine.h"
 	#include "EnginePrivate.h"
 	#include "<ModuleName>Classes.h"
 
-The C++ file is placed in the Private folder of the module's directory, i.e. [GameName]\Source\[ModuleName]\Private, registers and implements the module.
+The C++ file is placed in the Private folder of the module's directory, i.e. [GameName] \Source\ [ModuleName]\Private, registers and implements the module.
 
 [REGION:note]At least one module in your game must be registered using `IMPLEMENT_PRIMARY_GAME_MODULE`. Additional modules can use the alternative `IMPLEMENT_GAME_MODULE` method. See the [Multiple Gameplay Modules](#MultipleGameplayModules) section for more details on using more than one gameplay module in a game project.[/REGION]
 
@@ -32,7 +32,7 @@ The C++ file is placed in the Private folder of the module's directory, i.e. [Ga
 	// Designate the module as primary
 	IMPLEMENT_PRIMARY_GAME_MODULE( <ModuleName>, "<GameName>" );
 
-The build file is placed in the root directory of the gameplay module, i.e. [GameName]\Source\[ModuleName], and it defines certain pieces of information used by the UnrealBuildTool to compile the module.
+The build file is placed in the root directory of the gameplay module, i.e. [GameName] \Source\ [ModuleName], and it defines certain pieces of information used by the UnrealBuildTool to compile the module.
 
 	using UnrealBuildTool;
 
@@ -49,7 +49,7 @@ The build file is placed in the root directory of the gameplay module, i.e. [Gam
 
 Because the new gameplay modules will have UObject code, some configuration is necessary.
 
-1. The module needs to be added in several areas of the `DefaultEngine.ini` file:
+1. The module needs to be added in several areas of the DefaultEngine.ini file:
 
 	The `EditPackages` array in the `[UnrealEd.EditorEngine]` section:
 
@@ -71,7 +71,7 @@ Because the new gameplay modules will have UObject code, some configuration is n
 
 There are philosophical choices about game DLL modularity. Splitting a game up into a lot of DLL files may be more trouble than it is worth, but this is a decision to be made by each individual team based on their needs and principles. Using multiple gameplay modules will result in better link times and faster code iteration, but with more modules you will need to deal with DLL exports and/or interface classes more often. This trade-off is the correct one for engine and editor code, but it is questionable for gameplay.
 
-You can create a primary game module, then any number of additional game-specific modules. You can create `*.Build.cs` files for these new modules, then add references to these modules to your game's [`Target.cs` file](Programming/UnrealBuildSystem/TargetFiles) (OutExtraModuleNames array). Be sure to use the appropriate macro for the game modules in the C++ code. At least one module must use the `IMPLEMENT_PRIMARY_GAME_MODULE` macro, while all others should use the `IMPLEMENT_GAME_MODULE` macro. UBT should then automatically discover the modules and compile additional game DLL files.
+You can create a primary game module, then any number of additional game-specific modules. You can create *.Build.cs files for these new modules, then add references to these modules to your game's [Target.cs file](Programming/UnrealBuildSystem/TargetFiles) (OutExtraModuleNames array). Be sure to use the appropriate macro for the game modules in the C++ code. At least one module must use the `IMPLEMENT_PRIMARY_GAME_MODULE` macro, while all others should use the `IMPLEMENT_GAME_MODULE` macro. UBT should then automatically discover the modules and compile additional game DLL files.
 
 ## Limitations
 
