#UE4doc #docs:
Added step (and image) for running Setup.sh as per UEDOC-2561

diff --git a/Engine/Documentation/Source/Programming/Development/BuildingUnrealEngine/BuildingUnrealEngine.INT.udn b/Engine/Documentation/Source/Programming/Development/BuildingUnrealEngine/BuildingUnrealEngine.INT.udn
--- a/Engine/Documentation/Source/Programming/Development/BuildingUnrealEngine/BuildingUnrealEngine.INT.udn
+++ b/Engine/Documentation/Source/Programming/Development/BuildingUnrealEngine/BuildingUnrealEngine.INT.udn
@@ -127,7 +127,11 @@ SkillLevel: Advanced
 		Our developers and support teams currently use the latest version of Ubuntu; as a result, we may not be able to provide support for other Linux distributions (including other versions of Ubuntu).
 		[/REGION]
 		
-		1.	Inside the root directory, run **GenerateProjectFiles.sh** from the terminal to set-up your project files. 
+		1.  Inside the root directory, run **Setup.sh** from the terminal to setup the files to prepare for generating the project files
+
+			![](RunSetupShellScript_Linux.png)
+
+		1.	Inside the root directory, run **GenerateProjectFiles.sh** from the terminal to generate your project files. 
 
 			![](RunGenPrjFilesShellScript_Linux.png)
 
