#UE4 Docs: Jira UEDOC-2317 (Animation: Metadata on Sequences)

diff --git a/Engine/Documentation/Source/Engine/Animation/AnimationComposite/AnimationComposite.INT.udn b/Engine/Documentation/Source/Engine/Animation/AnimationComposite/AnimationComposite.INT.udn
--- a/Engine/Documentation/Source/Engine/Animation/AnimationComposite/AnimationComposite.INT.udn
+++ b/Engine/Documentation/Source/Engine/Animation/AnimationComposite/AnimationComposite.INT.udn
@@ -11,22 +11,19 @@ version: 4.9
 [TOC(start:2 end:2)]
 
 
-
 In certain circumstances, you may have a situation in which you need to take several Animation Sequences and stitch them together so that they can be used as if they were just one sequence instead of many. That is exactly the purpose of the **Animation Composite**. An Animation Composite is an animation asset designed to allow you to combine multiple animations together and treat them as a single unit. However, be aware that a Composite simply appends animations; it does not provide any blending ability.
 
-Animation Composites can also have their own sets of [Notifies](Engine/Animation/Sequences\Notifies) and [Curves](Engine/Animation/Sequences\Curves). For more information on these topics, please see the associated documentation for each.
+[REGION:note]
+Animation Composites can also have their own sets of [Notifies](Engine/Animation/Sequences\Notifies) and [Curves](Engine/Animation/Sequences\Curves). 
+[/REGION]
 
 ## Creating a Composite
 
-Creating a new Animation Composite is as simple as adding a new asset in the **Content Browser**:
+Creating a new Animation Composite can be accomplished inside the **Content Browser** by **Right-clicking** (or select the **Add New** Button) and choosing **Animation** then **Animation Composite** from the context menu. 
 
 ![](AnimationCompositeContextMenu.png)
 
-[REGION:caption]
-**Right-click** in the **Content Browser** and choose **Animation > Animation Composite** from the context menu. Alternatively, you may use the ![](button_NewAsset.png) button.
-[/REGION]
-
-You will also need to specify the [Skeleton asset](Engine/Animation/Skeleton) with which to associate your new Composite. 
+You will also need to specify the [Skeleton asset](Engine/Animation/Skeleton) with which to associate your new Animation Composite. 
 
 ![](PickSkeleton.png)
 
@@ -44,26 +41,59 @@ Your Composite is now created and ready to edit!
 
 **Double-clicking** a Composite will open Persona in Animation Mode with the Composite selected and ready to edit.
 
-![](AnimationCompositeEdit.png)
+![](AnimationCompositeEdit.png)(w:600)
 
-At the top of the **Asset Editor** panel, you will see the **Composite Track**.
+[REGION:caption]
+The Animation Composite will be highlighted in the **Asset Browser** located in the lower-right panel of **Persona**.
+[/REGION]
+
+At the top of the **Asset Editor** panel, you will see the **Composite Tracks**.
 
 ![](CompositeTrack.png)(w:600)
 
-To add Animation Sequences to the Composite, simply drag and drop them into the **Composite Track**. As you drag more Sequences in, they will be appended automatically.
+To add Animation Sequences to the Composite, drag-and-drop them from the **Asset Browser** into the **Composite Tracks**.
 
-![](AddingAnimationsToComposite.png)(w:700)
+![](AddingAnimationsToComposite.png)(w:900)
 
-[REGION:note]
-As you drag an asset out of the **Asset Browser**, you may notice the **Asset Editor** gray out and warn that you are "Currently previewing another asset". This is expected behavior. Simply continue dragging the Animation Sequence onto the **Composite Track**.
+[REGION:caption]
+As you drag more Sequences in, they will be appended automatically.
 [/REGION]
 
-You can move Sequences along the **Composite Track** by dragging them, and remove them from the track by **Right-clicking** them and choosing **Delete Segment** from the context menu.
+You can move Sequences along the **Composite Tracks** by dragging them. You can also remove a track by **Right-clicking** on it and choosing **Delete Segment** from the context menu.
 
 ![](DeleteSegment.png)
 
 ## Using Composites
 
-Once assembled, a Composite is treated just like an Animation Sequence in the AnimGraph of an Animation Blueprint.
+Once assembled, an Animation Composite is treated just like an [Animation Sequence](Engine\Animation\Sequences) in the AnimGraph of an [Animation Blueprint](Engine\Animation\AnimBlueprints).
 
 ![](PlayComposite.png)
+
+[REGION:note]
+See the [](Engine/Animation/Sequences/UserGuide) for more details on playback of animations in **Animation Blueprints**.
+[/REGION]
+
+
+## Anim Asset Details
+
+When opening your Animation Composite in **Persona**, there are several properties that can be adjusted inside the **Anim Asset Details** panel which can determine how the asset is handled. These properties include things such as tweaking playback speed or assigning any Meta Data to the asset and others which are outlined in the table below. 
+
+![](AnimAssetDetailsComposite.png)
+
+|Composite Properties||
+|---|---|
+|[REGION:tablesection]Additive Settings[/REGION]||
+|**Preview Base Pose**| Preview Base pose for additive BlendSpace.|
+|[REGION:tablesection]Length[/REGION]||
+|**Sequence Length**| Length (in seconds) of this AnimSequence if played back with a speed of 1.0 (unadjustable).|
+|[REGION:tablesection]Animation[/REGION]||
+|**Rate Scale**| Number for tweaking playback rate of this animation globally.|
+|**Skeleton**| Pointer to the Skeleton this asset can be played on (unadjustable).|
+|[REGION:tablesection]Meta Data[/REGION]||
+|**Meta Data**| This is Meta Data that can be saved with the asset. The meta data is a Blueprintable class derived from the Anim Meta Data class. This allows you to add custom meta data to animation assets (Anim Sequence, Anim Montage, Anim Composite, and Blendspace are all supported). You can query the data from an animation asset in C++ by calling the `GetMetaData()` method, or by using `GetSectionMetaData()` on Anim Montages. |
+|[REGION:tablesection]Thumbnail[/REGION]||
+|**Orbit Pitch**| The pitch of the orbit camera around the asset. |
+|**Orbit Yaw**| The yaw of the orbit camera around the asset. |
+|**Orbit Zoom**| The offset from the bounds sphere distance from the asset.|
+
+
diff --git a/Engine/Documentation/Source/Engine/Animation/Blendspaces/Blendspaces.INT.udn b/Engine/Documentation/Source/Engine/Animation/Blendspaces/Blendspaces.INT.udn
--- a/Engine/Documentation/Source/Engine/Animation/Blendspaces/Blendspaces.INT.udn
+++ b/Engine/Documentation/Source/Engine/Animation/Blendspaces/Blendspaces.INT.udn
@@ -32,7 +32,7 @@ version: 4.9
 [/OBJECT]
 [/VAR]
 
-[TOC (start:2)]
+[TOC (start:2 end:2)]
 
 
 
@@ -172,6 +172,58 @@ Each item in the Sample Data array has the following properties:
 | **Sample Value** | The input values (X and Y corresponding to the Blend Parameters) where the Animation has full influence. |
 
 
+## Anim Asset Details
+
+When setting up your Blend Space asset, there are additional options that can be set via the **Anim Asset Details** panel inside of **Persona**. By default, this window is located in the lower left-hand corner of the **Animation** tab of the Persona Editor and contains the following properties: 
+
+![](BlendSpaceDetails.png)
+
+### Blend Space
+
+|---|---|
+|**Display Editor Vertically**| Only available for 1D Blend Space assets, toggles the Blend Space graph display in either vertical or horizontal display modes.|
+
+### Input Interpolation
+
+|---|---|
+|**Axis to Scale Animation**| Not available for 	1D Blend Space assets, this allows you (if you have input interpolation) to determine which axis to drive animation speed (scale). For example, for locomotion animation, speed axis will drive animation speed (thus scale).|
+|**Interpolation Params**| Input interpolation parameter for all 3 axis, for each axis input, decide how you’d like to interpolate input to by setting Interpolation Time and Interpolation Type. |
+
+### Sample Interpolation
+
+|---|---|
+|**Target Weight Interpolation Speed Per Sec**| When target samples are set, how fast you would like to get to target which improves target blending. For example, if you interpolate input, when you move from left to right rapidly, you will interpolate through forward, but if you use target weight interpolation, you will skip forward but interpolate between left and right.|
+|**Per Bone Blend**| Define target weight interpolation per bone. This will blend in different speed per each bone setting.|
+
+### Animation Notifies
+
+|---|---|
+|**Notify Trigger Mode**| Assign the current mode used by the Blend Space to decide which animation notifies to fire: All Animation (all notify events will fire), HighestWeightedAnimation (Notify events will only fire from the highest weighted animation) or None (no notify events will fire from any animations).|
+
+### Additive Settings
+
+|---|---|
+|**Preview Base Pose**| Preview Base pose for additive Blend Space.|
+
+### Animation
+
+|---|---|
+|**Skeleton**| Pointer to the Skeleton this asset can be played on (unadjustable).|
+
+### Meta Data
+
+|---|---|
+|**Meta Data**| This is Meta Data that can be saved with the asset. The meta data is a Blueprintable class derived from the Anim Meta Data class. This allows you to add custom meta data to animation assets (Anim Sequence, Anim Montage, Anim Composite, and Blendspace are all supported). You can query the data from an animation asset in C++ by calling the `GetMetaData()` method, or by using `GetSectionMetaData()` on Anim Montages. |
+
+### Thumbnail
+
+|---|---|
+|**Orbit Pitch**| The pitch of the orbit camera around the asset. |
+|**Orbit Yaw**| The yaw of the orbit camera around the asset. |
+|**Orbit Zoom**| The offset from the bounds sphere distance from the asset.|
+
+
+
 
 
 
diff --git a/Engine/Documentation/Source/Engine/Animation/Sequences/AnimationSequences.INT.udn b/Engine/Documentation/Source/Engine/Animation/Sequences/AnimationSequences.INT.udn
--- a/Engine/Documentation/Source/Engine/Animation/Sequences/AnimationSequences.INT.udn
+++ b/Engine/Documentation/Source/Engine/Animation/Sequences/AnimationSequences.INT.udn
@@ -106,6 +106,72 @@ Each option provides a different compression method as outlined in the table bel
 | **Remove Trivial Keys** | Removes trivial frames of tracks when position or orientation is constant over the entire animation from the raw animation data. |
 
 
+## Anim Asset Details
+
+When an Animation Sequence asset is opened in **Persona**, there are several properties that can be adjusted inside the **Anim Asset Details** panel which can determine how the asset is handled. These properties include things such as tweaking playback speed, enabling/disabling Root Motion, assigning any Meta Data to the asset and others which are outlined in the tables below. 
+
+![](AnimAssetDetails.png)
+
+### Linear Key Removal
+
+|---|---|
+|**Max Pos Diff**| Max position difference to use when testing if an animation key may be removed. Lower values retain more keys, but yield less compression.|
+|**Max Angle Diff**| Max angle difference to use when testing if an animation key may be removed. Lower values retain more keys, but yield less compression. |
+|**Max Scale Diff**| Max scale difference to use when testing if an animation key may be removed. Lower values retain more keys, but yield less compression. |
+|**Max Effector Diff**| As keys are tested for removal, we monitor the effects all the way down to the end effectors. If their position changes by more than this amount as a result of removing a key, the key will be retained. This value is used for all bones except the end-effectors parent. |
+|**Min Effector Diff**| Similar to Max Effector Diff, however this value is used for the end-effectors parent, allowing tighter restrictions near the end of a skeletal chain. |
+|**Effector Diff Socket**| Error threshold for End Effectors with Sockets attached to them. Typically for more important bones where we want to be less aggressive with compression. |
+|**Parent Key Scale**| A scale value which increases the likelihood that a bone will retain a key if its parent also had a key at the same time position. Higher values can remove shaking artifacts from the animation at the cost of compression. |
+|**Retarget**| If true, as the animation is compressed, adjust animated nodes to compensate for compression error. If false, do not adjust animated nodes. |
+|**Actually Filter Linear Keys**| Controls whether the final filtering step will occur, or only the retargeting after bitwise compression. If both this and **Retarget** are false, then the linear compressor will do no better than the underlying bitwise compressor, extremely slowly.|
+
+
+### Animation
+
+|---|---|
+|**Retarget Source**| When retargeting, this pose will be used as a base of animation.|
+|**Rate Scale**| Number for tweaking playback rate of this animation globally. |
+|**Skeleton**| Contains the skeleton associated with this Montage. Cannot be changed in the editor.|
+
+### Compression
+
+|---|---|
+|**Description**| Name of Compression Scheme used for this asset.|
+|**Do Not Override Compression**| Do not attempt to override compression scheme when running Compress Animations commandlet. Some high frequency animations are too sensitive and shouldn’t be changed. |
+
+### Additive Settings
+
+|---|---|
+|**Additive Anim Type**| The type of additive animation type to use: No Additive, Local Space or Mesh Space. Local space is additive and delta is calculated by local space. Mesh Space is additive and delta will be applied in component space. |
+|**Base Pose Type**| Decides how the additive delta is to be calculated: None (No Base Pose), Reference Pose (The mesh’s default reference pose), Selected Animation Scaled (choose another animation and scale the animation), Selected Animation Frame (choose another animation and only based on the frame number). Base Pose Type is only available if Additive Anim Type is set to Local or Mesh Space. |
+
+### Root Motion
+
+|---|---|
+|**Enable Root Motion**| If this is on, it will allow extracting of [Root Motion]( Engine/Animation/RootMotion). |
+|**RootMotion Root Lock**| Root Bone will be locked to that position specified when extracting root motion. |
+
+### Import Settings
+
+|---|---|
+|**Asset Import Data**| The importing data and options used for the mesh. |
+
+
+### Meta Data
+
+|---|---|
+|**Meta Data**| This is Meta Data that can be saved with the asset. The meta data is a Blueprintable class derived from the Anim Meta Data class. This allows you to add custom meta data to animation assets (Anim Sequence, Anim Montage, Anim Composite, and Blendspace are all supported). You can query the data from an animation asset in C++ by calling the `GetMetaData()` method, or by using `GetSectionMetaData()` on Anim Montages. |
+
+### Thumbnail
+
+
+|---|---|
+|**Orbit Pitch**| The pitch of the orbit camera around the asset. |
+|**Orbit Yaw**| The yaw of the orbit camera around the asset. |
+|**Orbit Zoom**| The offset from the bounds sphere distance from the asset.|
+
+
+
 
 
 
