#UE4doc #docs:
Removed text that was referring to features as not implemented, as they have been implemented at this time.

diff --git a/Engine/Documentation/Source/Programming/Online/OnlineSubsystem.INT.udn b/Engine/Documentation/Source/Programming/Online/OnlineSubsystem.INT.udn
--- a/Engine/Documentation/Source/Programming/Online/OnlineSubsystem.INT.udn
+++ b/Engine/Documentation/Source/Programming/Online/OnlineSubsystem.INT.udn
@@ -119,13 +119,13 @@ Provides the interface for accessing online leaderboards.
 
 #### Voice
 
-Provides the interface for voice communication over network while in the game (NOT IMPLEMENTED YET).
+Provides the interface for voice communication over network while in the game.
 
 <!--Details can be found on the [](Programming/Online/Interfaces/Voice) page.-->
 
 #### Achievements
 
-Provides the interface for reading/writing/unlocking achievements (NOT IMPLEMENTED YET).
+Provides the interface for reading/writing/unlocking achievements.
 
 <!--Details can be found on the [](Programming/Online/Interfaces/Achievements) page.-->
 
