#UE4doc #docs:
Removed italics from a few terms

diff --git a/Engine/Documentation/Source/Programming/Online/Interfaces/Session/Session.INT.udn b/Engine/Documentation/Source/Programming/Online/Interfaces/Session/Session.INT.udn
--- a/Engine/Documentation/Source/Programming/Online/Interfaces/Session/Session.INT.udn
+++ b/Engine/Documentation/Source/Programming/Online/Interfaces/Session/Session.INT.udn
@@ -7,8 +7,8 @@ Version: 4.9
 [TOC(start:2)]
 
 
-Matchmaking is the process of matching _players_ with _sessions_. A _session_ is basically an instance of 
-the game running on the server with a given set of properties which is either _advertised_ so that it can 
+Matchmaking is the process of matching players with sessions. A session is basically an instance of 
+the game running on the server with a given set of properties which is either advertised so that it can 
 be found and joined by players wanting to play the game or private so only players who are invited or 
 notified of it in some way can join.
 
@@ -42,7 +42,7 @@ interface must be created. Platforms, in this sense, refer to hardware platforms
 interface will ever exist at a time - the session interface for the platform the engine is currently running on.
 
 While the session interface performs all of the session handling, the game doesn't generally interact directly with 
-it. Instead, the **game session**, `AGameSession`, acts as a game-specific wrapper around the session interface 
+it. Instead, the **Game Session**, `AGameSession`, acts as a game-specific wrapper around the session interface 
 and the game code makes calls to it when it needs to interact with the session. The game session is created and 
 owned by the game mode, `AGameMode`, and also only exists on the server when running an online game.
 
@@ -52,7 +52,7 @@ using a dedicated server.
 
 ## Session Settings
 
-The **session settings**, defined by the `FOnlineSessionSettings` class, are a set of properties that determine 
+The **Session Settings**, defined by the `FOnlineSessionSettings` class, are a set of properties that determine 
 the characteristics of the session. In the base implementation, these are things like:
 
 * Number of players allowed
