#UE4doc #docs:
Removed monospace from a term

diff --git a/Engine/Documentation/Source/Programming/Online/Steam/OnlineSubsystemSteam.INT.udn b/Engine/Documentation/Source/Programming/Online/Steam/OnlineSubsystemSteam.INT.udn
--- a/Engine/Documentation/Source/Programming/Online/Steam/OnlineSubsystemSteam.INT.udn
+++ b/Engine/Documentation/Source/Programming/Online/Steam/OnlineSubsystemSteam.INT.udn
@@ -65,7 +65,7 @@ to make sure it references the new SDK directory location / version
 
 ### INI Configuration
 
-Turn on some settings in the game's `DefaultEngine.ini`.  
+Turn on some settings in the game's DefaultEngine.ini.  
 [REGION:note]
 The SteamDevAppId of 480 is Valve's test app id, shared by everyone.   You will need your own app id eventually, but most features of Steam should work before then.
 [/REGION]
