#UE4doc #docs:
Unbolded a term
Removed a bulleted list as it only had one entry.

diff --git a/Engine/Documentation/Source/Programming/Online/Types/Types.INT.udn b/Engine/Documentation/Source/Programming/Online/Types/Types.INT.udn
--- a/Engine/Documentation/Source/Programming/Online/Types/Types.INT.udn
+++ b/Engine/Documentation/Source/Programming/Online/Types/Types.INT.udn
@@ -17,7 +17,7 @@ This document serves to define the different data structures that are used throu
 
  There are often platform-specific representations of various constructs that, while they need to be handled by the game, do not need their details exposed. This allows the game to hold onto the data and replicate it to other machines without understanding the details.
  
- Classes deriving from **IOnlinePlatformData**:
+ Classes deriving from IOnlinePlatformData:
  * [FUniqueNetId](#FUniqueNetId)
  * [FSharedContentHandle](#FSharedContentHandle)
  * [FSessionInfo](#FSessionInfo)
@@ -57,8 +57,7 @@ This document serves to define the different data structures that are used throu
 
 **FUniqueNetId** contains a profile service online ID and can represent anything from an individual player, friend or a session.
 
-<!--See:
- * [](Programming/Online/Interfaces/Profile)-->
+<!--See: [](Programming/Online/Interfaces/Profile)-->
 
 
 ### FSharedContentHandle
@@ -74,8 +73,7 @@ This document serves to define the different data structures that are used throu
 
 **FSessionInfo** contains a session's platform-specific info.  In most cases, this would be information such as a session identifier, host address, or other means of identifying and connecting to a created session.
 
-See:
- * [](Programming/Online/Interfaces/Session)
+See: [](Programming/Online/Interfaces/Session)
 
 
 ### FEmsFile
