#UE4doc #docs:
Changed font settings for multiple terms
Replaced a couple images as they were outdated.

diff --git a/Engine/Documentation/Source/Programming/Plugins/Plugins.INT.udn b/Engine/Documentation/Source/Programming/Plugins/Plugins.INT.udn
--- a/Engine/Documentation/Source/Programming/Plugins/Plugins.INT.udn
+++ b/Engine/Documentation/Source/Programming/Plugins/Plugins.INT.udn
@@ -58,9 +58,9 @@ Plugins can have their own Content folder that contains asset files specific to
 section for more information.
 
 [REGION:warning]
-Note that Plugins do not support config files yet. This is something we are considering to add at a later date. 
+Plugins do not support config files yet. This is something we are considering to add at a later date. 
 
-Also note that plugins do not yet support their own Derived Data cache distribution. This is something we are looking 
+Also, plugins do not yet support their own Derived Data cache distribution. This is something we are looking 
 into adding at a later date.
 [/REGION]
 
@@ -72,11 +72,11 @@ the valid search paths for plugins in Unreal Engine.
 
 | Plugin type      | Search path                                        |
 | ---------------- | -------------------------------------------------- |
-| Engine plugin    | /**UE4 root**/Engine/Plugins/**My Engine Plugin**/ |
-| Game plugin      | /**My project**/Plugins/**My Game Plugin**/        |
+| Engine plugin    | /UE4 root/Engine/Plugins/My Engine Plugin/ |
+| Game plugin      | /My project/Plugins/My Game Plugin/        |
 
-You can also organize plugins into subdirectories under the base **Plugins** folder. The engine will scan all of your sub-folders 
-under the base **Plugins** folder for plugins to load, but it will never scan into subdirectories beneath where a plugin was already found.
+You can also organize plugins into subdirectories under the base Plugins folder. The engine will scan all of your sub-folders 
+under the base Plugins folder for plugins to load, but it will never scan into subdirectories beneath where a plugin was already found.
 
 Unreal Engine finds your plugin by searching for .uplugin files on disk. We call these files plugin descriptors. They are text 
 files that provide basic information about your plugin. Plugin descriptors are discovered and loaded automatically by Unreal Engine, 
@@ -101,9 +101,9 @@ For the most part, plugin source file layout is the same as any other C++ module
 Plugins are allowed to declare new UObject types (UCLASS, USTRUCT, etc) in header files in a Classes subdirectory in their module 
 code folders. The Unreal Engine build system will detect these files and generate code as needed to support these UObjects. You will 
 need to follow the normal rules for using UObjects within C++ modules, such as including the generated header file and the module's 
-`generated.inl` file in one of your module's source files.
+generated.inl file in one of your module's source files.
 
-One thing with plugin modules that is a bit different, is "Public" source file headers. Most plugin modules have no business exporting 
+One thing with plugin modules that is a bit different, is Public source file headers. Most plugin modules have no business exporting 
 public APIs in header files in a Public source folder, because they will never be a direct dependency (statically linked against) by 
 engine or game code. So usually your Public source folder will be empty. There are a few exceptions to this rule:
 
@@ -132,7 +132,7 @@ recommend using it.
 
 ## Plugins in your Game Project
 
-Plugins live under a "Plugins" subfolder in your game project's directory, and will be detected and loaded at startup by the game engine and editor. 
+Plugins live under a Plugins subfolder in your game project's directory, and will be detected and loaded at startup by the game engine and editor. 
 
 If the plugin contains modules that have Source folders (and *.Build.cs files), plugin code will be automatically added to generated 
 C++ project files, so that you can work on developing the plugin alongside your game project easily. Whenever you compile your game 
@@ -162,7 +162,7 @@ Here are the steps required to distribute a plugin.
 
 1. **Edit your plugin's descriptor file** (.uplugin file) and make sure the plugin name, modules, version and other settings are configured as desired.
 2. **Delete** your plugin's **Binaries** and **Intermediate** folders. This is just to make sure you are starting with a clean slate.
-3. If your plugin contains a **Source** directory, **compile plugin binaries** for **Win64 Development** configuration. This configuration will always be needed by an editor on PC platform to load your plugin. 
+3. If your plugin contains a **Source** directory, compile plugin binaries for **Win64 Development** configuration. This configuration will always be needed by an editor on PC platform to load your plugin. 
 4. For plugins with **Source**, if you intend to support additional configurations such as Mac, you should build those too.
 5. Make a copy of your plugin to a new "staging" folder outside of your project directory. It does not matter where, but you are going to be making changes to the plugin here.
 6. In the staging directory for your plugin, delete the **Intermediate** directory. These are temporary files that should never be distributed.
@@ -258,8 +258,8 @@ For plugins that contain code, the descriptor file will contain at least one mod
 | Field name		| Info			| Description |
 | -----------------	| -------------	| ----------- |
 | **Name**			| _Required_	| Unique name of this plugin module that will be loaded with the plugin. At runtime, the engine will expect appropriate plugin binaries to exist in the plugin's Binaries folder with the module name as specified here. For modules that have a Source directory, a matching *.Build.cs file is expected to exist within the module's subfolder tree. |
-| **Type**			| _Required_	| Sets the type of module. Valid options are **Runtime**, **RuntimeNoCommandlet**, **Developer**, **Editor**, **EditorNoCommandlet**, and **Program**. This type determines which types of applications this plugin's module is suitable for loading in. For example, some plugins may include modules that are only designed to be loaded when the editor is running. **Runtime** modules will be loaded in all cases, even in shipped games. **Developer** modules will only be loaded in development runtime or editor builds, but never in shipping builds. **Editor** modules will only be loaded when the editor is starting up. Your plugin can use a combination of modules of different types. |
-| **LoadingPhase**	| _Optional_	| If specified, controls when the plugin is loaded at startup. This is an advanced option that should not normally be required. The valid options are **Default** (which is used when no LoadingPhase is specified), **PreDefault**, and **PostConfigInit**. **PostConfigInit** allows the module to be loaded before the engine has finished starting up key subsystems. **PreDefault** loads just before the normal phase. Typically, this is only needed if you expect game modules to depend directly on content within your plugin, or types declared within your plugin's code. |
+| **Type**			| _Required_	| Sets the type of module. Valid options are **Runtime**, **RuntimeNoCommandlet**, **Developer**, **Editor**, **EditorNoCommandlet**, and **Program**. This type determines which types of applications this plugin's module is suitable for loading in. For example, some plugins may include modules that are only designed to be loaded when the editor is running. Runtime modules will be loaded in all cases, even in shipped games. Developer modules will only be loaded in development runtime or editor builds, but never in shipping builds. Editor modules will only be loaded when the editor is starting up. Your plugin can use a combination of modules of different types. |
+| **LoadingPhase**	| _Optional_	| If specified, controls when the plugin is loaded at startup. This is an advanced option that should not normally be required. The valid options are **Default** (which is used when no LoadingPhase is specified), **PreDefault**, and **PostConfigInit**. **PostConfigInit** allows the module to be loaded before the engine has finished starting up key subsystems. PreDefault loads just before the normal phase. Typically, this is only needed if you expect game modules to depend directly on content within your plugin, or types declared within your plugin's code. |
 
 		 
 
@@ -292,7 +292,7 @@ To use these examples as a starting point for your plugin:
 4. **Compile your game project** normally. Unreal Build Tool will detect the plugins and compile them as dependencies to your game!
 5. **Launch the editor** (or the game.)  Your plugin will be initially disabled, but you can turn it on in the editor UI.
 6. **Open the Plugins Editor** (Window -> Plugins), search for your plugin and enable it by clicking the check box.
-7. **Restart the Editor**.  Your plugin will be automatically **loaded at startup**.
+7. **Restart the Editor**.  Your plugin will be automatically loaded at startup.
 
 You can see that the plugins have loaded by opening the **Modules** viewer under the **Window -> Developer Tools** menu. 
 Another way is to use the code debugger to place a breakpoint in the plugin's startup code, such as in *FBlankPlugin::StartupModule()*.
