#UE4doc #docs:
Updated engine version field

diff --git a/Engine/Documentation/Source/Programming/QuickStart/QuickStart.INT.udn b/Engine/Documentation/Source/Programming/QuickStart/QuickStart.INT.udn
--- a/Engine/Documentation/Source/Programming/QuickStart/QuickStart.INT.udn
+++ b/Engine/Documentation/Source/Programming/QuickStart/QuickStart.INT.udn
@@ -5,7 +5,7 @@ Description:Create your first code project and add a new C++ class.
 prereq:GettingStarted/Installation
 prereq:Engine/Basics/Projects/Browser
 SkillLevel: Beginner
-Version: 4.8
+Version: 4.9
 
 [VAR:Topic]
 [OBJECT:Topic]
