#UE4doc #docs:
Updated engine version field
Changed font settings for multiple terms

diff --git a/Engine/Documentation/Source/Programming/QuickStart/1/QuickStart_1.INT.udn b/Engine/Documentation/Source/Programming/QuickStart/1/QuickStart_1.INT.udn
--- a/Engine/Documentation/Source/Programming/QuickStart/1/QuickStart_1.INT.udn
+++ b/Engine/Documentation/Source/Programming/QuickStart/1/QuickStart_1.INT.udn
@@ -3,7 +3,7 @@ Title:1 - Essential Project Setup
 Crumbs:%ROOT%, Programming, Programming/QuickStart
 Description:Creating a new project with the specifications we need
 SkillLevel: Beginner
-Version: 4.8
+Version: 4.9
 
 [VAR:Steps]
 [OBJECT:Navigation]
@@ -28,7 +28,7 @@ You should have a copy of **Visual Studio 2013** for Desktop (Preferably Communi
 
 1. Open **Unreal Engine** from the Launcher. The [](Engine/Basics/Projects/Browser) will appear.
 
-1. Click on the **New Project** tab and then select the **C++** tab. From there, select **Basic Code** so we have a clean starting point, and make sure "With Starter Content" is set. We'll need to enter a project name, so we'll use "QuickStart". We can now click **Create Project** and get started.
+1. Click on the **New Project** tab and then select the **C++** tab. From there, select **Basic Code** so we have a clean starting point, and make sure **With Starter Content** is set. We'll need to enter a project name, so we'll use "QuickStart". We can now click **Create Project** and get started.
 
 	[![](CreateProject.png)](CreateProject.png)
 
