#UE4doc #docs:
Updated engine version field
Changed font settings for a term

diff --git a/Engine/Documentation/Source/Programming/QuickStart/2/QuickStart_2.INT.udn b/Engine/Documentation/Source/Programming/QuickStart/2/QuickStart_2.INT.udn
--- a/Engine/Documentation/Source/Programming/QuickStart/2/QuickStart_2.INT.udn
+++ b/Engine/Documentation/Source/Programming/QuickStart/2/QuickStart_2.INT.udn
@@ -3,7 +3,7 @@ Title:2 - Create a C++ Class
 Crumbs:%ROOT%, Programming, Programming/QuickStart
 Description:Creating a C++ class
 SkillLevel: Beginner
-Version: 4.8
+Version: 4.9
 
 [VAR:Steps]
 [OBJECT:Navigation]
@@ -35,6 +35,6 @@ Version: 4.8
  
 	[![](NameActorClass.png)](NameActorClass.png)
 
-Now that we have created a C++ class, we can switch to **Visual Studio** to program it. `FloatingActor.cpp` will be opened for us automatically, and **Unreal Engine** will automatically compile and reload the code with our new class.
+Now that we have created a C++ class, we can switch to **Visual Studio** to program it. FloatingActor.cpp will be opened for us automatically, and **Unreal Engine** will automatically compile and reload the code with our new class.
 
 %Steps%
