#UE4doc #docs:
Updated engine version field
Changed font settings for multiple terms
Changed  a tip region to a note

diff --git a/Engine/Documentation/Source/Programming/QuickStart/3/QuickStart_3.INT.udn b/Engine/Documentation/Source/Programming/QuickStart/3/QuickStart_3.INT.udn
--- a/Engine/Documentation/Source/Programming/QuickStart/3/QuickStart_3.INT.udn
+++ b/Engine/Documentation/Source/Programming/QuickStart/3/QuickStart_3.INT.udn
@@ -3,7 +3,7 @@ Title:3 - Write and Compile C++ Code
 Crumbs:%ROOT%, Programming, Programming/QuickStart
 Description:Writing a C++ class and compiling it with Unreal Engine
 SkillLevel: Beginner
-Version: 4.8
+Version: 4.9
 
 [VAR:Steps]
 [OBJECT:Navigation]
@@ -23,17 +23,17 @@ Version: 4.8
 
 %Steps%
 
-1. In **Visual Studio**, we'll use the **Solution Explorer** pane to find our newly-created C++ files. In our example, they will be named `FloatingActor.cpp` and `FloatingActor.h` and will be inside the "QuickStart" project.
+1. In **Visual Studio**, we'll use the **Solution Explorer** pane to find our newly-created C++ files. In our example, they will be named FloatingActor.cpp and FloatingActor.h and will be inside the QuickStart project.
 
 	[![](SolutionExplorer_Files.png)](SolutionExplorer_Files.png)
 
 	[REGION:note]We're about to write code. All code used in this tutorial will be found at bottom of the page in its current state as of completing that page's instructions.[/REGION]
 
-1. In `FloatingActor.h`, we'll add the following code just before the closing brace and semicolon at the end of the file:
+1. In FloatingActor.h, we'll add the following code just before the closing brace and semicolon at the end of the file:
 
 		float RunningTime;
 
-1. Switching to `FloatingActor.cpp`, we'll add the following code just before the closing brace at the bottom of **AFloatingActor::Tick**:
+1. Switching to FloatingActor.cpp, we'll add the following code just before the closing brace at the bottom of **AFloatingActor::Tick**:
 
 		FVector NewLocation = GetActorLocation();
 		float DeltaHeight = (FMath::Sin(RunningTime + DeltaTime) - FMath::Sin(RunningTime));
@@ -41,7 +41,7 @@ Version: 4.8
 		RunningTime += DeltaTime;
 		SetActorLocation(NewLocation);
 
-	[REGION:tip] The code we've just written will cause **FloatingActors** to bob up and down smoothly, using the _RunningTime_ variable we created to keep track of our movement over time. [/REGION] 
+	[REGION:note] The code we've just written will cause **FloatingActors** to bob up and down smoothly, using the _RunningTime_ variable we created to keep track of our movement over time. [/REGION] 
 
 %Globals:OSSelect%
 
