Branch snapshot for CL 2723545

[CL 2723545 in Main branch]

diff --git a/Engine/Documentation/Source/Engine/Animation/AnimationComposite/AnimationComposite.INT.udn b/Engine/Documentation/Source/Engine/Animation/AnimationComposite/AnimationComposite.INT.udn
--- a/Engine/Documentation/Source/Engine/Animation/AnimationComposite/AnimationComposite.INT.udn
+++ b/Engine/Documentation/Source/Engine/Animation/AnimationComposite/AnimationComposite.INT.udn
@@ -11,19 +11,22 @@ version: 4.9
 [TOC(start:2 end:2)]
 
 
+
 In certain circumstances, you may have a situation in which you need to take several Animation Sequences and stitch them together so that they can be used as if they were just one sequence instead of many. That is exactly the purpose of the **Animation Composite**. An Animation Composite is an animation asset designed to allow you to combine multiple animations together and treat them as a single unit. However, be aware that a Composite simply appends animations; it does not provide any blending ability.
 
-[REGION:note]
-Animation Composites can also have their own sets of [Notifies](Engine/Animation/Sequences\Notifies) and [Curves](Engine/Animation/Sequences\Curves). 
-[/REGION]
+Animation Composites can also have their own sets of [Notifies](Engine/Animation/Sequences\Notifies) and [Curves](Engine/Animation/Sequences\Curves). For more information on these topics, please see the associated documentation for each.
 
 ## Creating a Composite
 
-Creating a new Animation Composite can be accomplished inside the **Content Browser** by **Right-clicking** (or select the **Add New** Button) and choosing **Animation** then **Animation Composite** from the context menu. 
+Creating a new Animation Composite is as simple as adding a new asset in the **Content Browser**:
 
 ![](AnimationCompositeContextMenu.png)
 
-You will also need to specify the [Skeleton asset](Engine/Animation/Skeleton) with which to associate your new Animation Composite. 
+[REGION:caption]
+**Right-click** in the **Content Browser** and choose **Animation > Animation Composite** from the context menu. Alternatively, you may use the ![](button_NewAsset.png) button.
+[/REGION]
+
+You will also need to specify the [Skeleton asset](Engine/Animation/Skeleton) with which to associate your new Composite. 
 
 ![](PickSkeleton.png)
 
@@ -41,59 +44,26 @@ Your Composite is now created and ready to edit!
 
 **Double-clicking** a Composite will open Persona in Animation Mode with the Composite selected and ready to edit.
 
-![](AnimationCompositeEdit.png)(w:600)
+![](AnimationCompositeEdit.png)
 
-[REGION:caption]
-The Animation Composite will be highlighted in the **Asset Browser** located in the lower-right panel of **Persona**.
-[/REGION]
-
-At the top of the **Asset Editor** panel, you will see the **Composite Tracks**.
+At the top of the **Asset Editor** panel, you will see the **Composite Track**.
 
 ![](CompositeTrack.png)(w:600)
 
-To add Animation Sequences to the Composite, drag-and-drop them from the **Asset Browser** into the **Composite Tracks**.
+To add Animation Sequences to the Composite, simply drag and drop them into the **Composite Track**. As you drag more Sequences in, they will be appended automatically.
 
-![](AddingAnimationsToComposite.png)(w:900)
+![](AddingAnimationsToComposite.png)(w:700)
 
-[REGION:caption]
-As you drag more Sequences in, they will be appended automatically.
+[REGION:note]
+As you drag an asset out of the **Asset Browser**, you may notice the **Asset Editor** gray out and warn that you are "Currently previewing another asset". This is expected behavior. Simply continue dragging the Animation Sequence onto the **Composite Track**.
 [/REGION]
 
-You can move Sequences along the **Composite Tracks** by dragging them. You can also remove a track by **Right-clicking** on it and choosing **Delete Segment** from the context menu.
+You can move Sequences along the **Composite Track** by dragging them, and remove them from the track by **Right-clicking** them and choosing **Delete Segment** from the context menu.
 
 ![](DeleteSegment.png)
 
 ## Using Composites
 
-Once assembled, an Animation Composite is treated just like an [Animation Sequence](Engine\Animation\Sequences) in the AnimGraph of an [Animation Blueprint](Engine\Animation\AnimBlueprints).
+Once assembled, a Composite is treated just like an Animation Sequence in the AnimGraph of an Animation Blueprint.
 
 ![](PlayComposite.png)
-
-[REGION:note]
-See the [](Engine/Animation/Sequences/UserGuide) for more details on playback of animations in **Animation Blueprints**.
-[/REGION]
-
-
-## Anim Asset Details
-
-When opening your Animation Composite in **Persona**, there are several properties that can be adjusted inside the **Anim Asset Details** panel which can determine how the asset is handled. These properties include things such as tweaking playback speed or assigning any Meta Data to the asset and others which are outlined in the table below. 
-
-![](AnimAssetDetailsComposite.png)
-
-|Composite Properties||
-|---|---|
-|[REGION:tablesection]Additive Settings[/REGION]||
-|**Preview Base Pose**| Preview Base pose for additive BlendSpace.|
-|[REGION:tablesection]Length[/REGION]||
-|**Sequence Length**| Length (in seconds) of this AnimSequence if played back with a speed of 1.0 (unadjustable).|
-|[REGION:tablesection]Animation[/REGION]||
-|**Rate Scale**| Number for tweaking playback rate of this animation globally.|
-|**Skeleton**| Pointer to the Skeleton this asset can be played on (unadjustable).|
-|[REGION:tablesection]Meta Data[/REGION]||
-|**Meta Data**| This is Meta Data that can be saved with the asset. The meta data is a Blueprintable class derived from the Anim Meta Data class. This allows you to add custom meta data to animation assets (Anim Sequence, Anim Montage, Anim Composite, and Blendspace are all supported). You can query the data from an animation asset in C++ by calling the `GetMetaData()` method, or by using `GetSectionMetaData()` on Anim Montages. |
-|[REGION:tablesection]Thumbnail[/REGION]||
-|**Orbit Pitch**| The pitch of the orbit camera around the asset. |
-|**Orbit Yaw**| The yaw of the orbit camera around the asset. |
-|**Orbit Zoom**| The offset from the bounds sphere distance from the asset.|
-
-
diff --git a/Engine/Documentation/Source/Engine/Animation/Blendspaces/Blendspaces.INT.udn b/Engine/Documentation/Source/Engine/Animation/Blendspaces/Blendspaces.INT.udn
--- a/Engine/Documentation/Source/Engine/Animation/Blendspaces/Blendspaces.INT.udn
+++ b/Engine/Documentation/Source/Engine/Animation/Blendspaces/Blendspaces.INT.udn
@@ -32,7 +32,7 @@ version: 4.9
 [/OBJECT]
 [/VAR]
 
-[TOC (start:2 end:2)]
+[TOC (start:2)]
 
 
 
@@ -172,58 +172,6 @@ Each item in the Sample Data array has the following properties:
 | **Sample Value** | The input values (X and Y corresponding to the Blend Parameters) where the Animation has full influence. |
 
 
-## Anim Asset Details
-
-When setting up your Blend Space asset, there are additional options that can be set via the **Anim Asset Details** panel inside of **Persona**. By default, this window is located in the lower left-hand corner of the **Animation** tab of the Persona Editor and contains the following properties: 
-
-![](BlendSpaceDetails.png)
-
-### Blend Space
-
-|---|---|
-|**Display Editor Vertically**| Only available for 1D Blend Space assets, toggles the Blend Space graph display in either vertical or horizontal display modes.|
-
-### Input Interpolation
-
-|---|---|
-|**Axis to Scale Animation**| Not available for 	1D Blend Space assets, this allows you (if you have input interpolation) to determine which axis to drive animation speed (scale). For example, for locomotion animation, speed axis will drive animation speed (thus scale).|
-|**Interpolation Params**| Input interpolation parameter for all 3 axis, for each axis input, decide how you’d like to interpolate input to by setting Interpolation Time and Interpolation Type. |
-
-### Sample Interpolation
-
-|---|---|
-|**Target Weight Interpolation Speed Per Sec**| When target samples are set, how fast you would like to get to target which improves target blending. For example, if you interpolate input, when you move from left to right rapidly, you will interpolate through forward, but if you use target weight interpolation, you will skip forward but interpolate between left and right.|
-|**Per Bone Blend**| Define target weight interpolation per bone. This will blend in different speed per each bone setting.|
-
-### Animation Notifies
-
-|---|---|
-|**Notify Trigger Mode**| Assign the current mode used by the Blend Space to decide which animation notifies to fire: All Animation (all notify events will fire), HighestWeightedAnimation (Notify events will only fire from the highest weighted animation) or None (no notify events will fire from any animations).|
-
-### Additive Settings
-
-|---|---|
-|**Preview Base Pose**| Preview Base pose for additive Blend Space.|
-
-### Animation
-
-|---|---|
-|**Skeleton**| Pointer to the Skeleton this asset can be played on (unadjustable).|
-
-### Meta Data
-
-|---|---|
-|**Meta Data**| This is Meta Data that can be saved with the asset. The meta data is a Blueprintable class derived from the Anim Meta Data class. This allows you to add custom meta data to animation assets (Anim Sequence, Anim Montage, Anim Composite, and Blendspace are all supported). You can query the data from an animation asset in C++ by calling the `GetMetaData()` method, or by using `GetSectionMetaData()` on Anim Montages. |
-
-### Thumbnail
-
-|---|---|
-|**Orbit Pitch**| The pitch of the orbit camera around the asset. |
-|**Orbit Yaw**| The yaw of the orbit camera around the asset. |
-|**Orbit Zoom**| The offset from the bounds sphere distance from the asset.|
-
-
-
 
 
 
diff --git a/Engine/Documentation/Source/Engine/Animation/Sequences/AnimationSequences.INT.udn b/Engine/Documentation/Source/Engine/Animation/Sequences/AnimationSequences.INT.udn
--- a/Engine/Documentation/Source/Engine/Animation/Sequences/AnimationSequences.INT.udn
+++ b/Engine/Documentation/Source/Engine/Animation/Sequences/AnimationSequences.INT.udn
@@ -106,72 +106,6 @@ Each option provides a different compression method as outlined in the table bel
 | **Remove Trivial Keys** | Removes trivial frames of tracks when position or orientation is constant over the entire animation from the raw animation data. |
 
 
-## Anim Asset Details
-
-When an Animation Sequence asset is opened in **Persona**, there are several properties that can be adjusted inside the **Anim Asset Details** panel which can determine how the asset is handled. These properties include things such as tweaking playback speed, enabling/disabling Root Motion, assigning any Meta Data to the asset and others which are outlined in the tables below. 
-
-![](AnimAssetDetails.png)
-
-### Linear Key Removal
-
-|---|---|
-|**Max Pos Diff**| Max position difference to use when testing if an animation key may be removed. Lower values retain more keys, but yield less compression.|
-|**Max Angle Diff**| Max angle difference to use when testing if an animation key may be removed. Lower values retain more keys, but yield less compression. |
-|**Max Scale Diff**| Max scale difference to use when testing if an animation key may be removed. Lower values retain more keys, but yield less compression. |
-|**Max Effector Diff**| As keys are tested for removal, we monitor the effects all the way down to the end effectors. If their position changes by more than this amount as a result of removing a key, the key will be retained. This value is used for all bones except the end-effectors parent. |
-|**Min Effector Diff**| Similar to Max Effector Diff, however this value is used for the end-effectors parent, allowing tighter restrictions near the end of a skeletal chain. |
-|**Effector Diff Socket**| Error threshold for End Effectors with Sockets attached to them. Typically for more important bones where we want to be less aggressive with compression. |
-|**Parent Key Scale**| A scale value which increases the likelihood that a bone will retain a key if its parent also had a key at the same time position. Higher values can remove shaking artifacts from the animation at the cost of compression. |
-|**Retarget**| If true, as the animation is compressed, adjust animated nodes to compensate for compression error. If false, do not adjust animated nodes. |
-|**Actually Filter Linear Keys**| Controls whether the final filtering step will occur, or only the retargeting after bitwise compression. If both this and **Retarget** are false, then the linear compressor will do no better than the underlying bitwise compressor, extremely slowly.|
-
-
-### Animation
-
-|---|---|
-|**Retarget Source**| When retargeting, this pose will be used as a base of animation.|
-|**Rate Scale**| Number for tweaking playback rate of this animation globally. |
-|**Skeleton**| Contains the skeleton associated with this Montage. Cannot be changed in the editor.|
-
-### Compression
-
-|---|---|
-|**Description**| Name of Compression Scheme used for this asset.|
-|**Do Not Override Compression**| Do not attempt to override compression scheme when running Compress Animations commandlet. Some high frequency animations are too sensitive and shouldn’t be changed. |
-
-### Additive Settings
-
-|---|---|
-|**Additive Anim Type**| The type of additive animation type to use: No Additive, Local Space or Mesh Space. Local space is additive and delta is calculated by local space. Mesh Space is additive and delta will be applied in component space. |
-|**Base Pose Type**| Decides how the additive delta is to be calculated: None (No Base Pose), Reference Pose (The mesh’s default reference pose), Selected Animation Scaled (choose another animation and scale the animation), Selected Animation Frame (choose another animation and only based on the frame number). Base Pose Type is only available if Additive Anim Type is set to Local or Mesh Space. |
-
-### Root Motion
-
-|---|---|
-|**Enable Root Motion**| If this is on, it will allow extracting of [Root Motion]( Engine/Animation/RootMotion). |
-|**RootMotion Root Lock**| Root Bone will be locked to that position specified when extracting root motion. |
-
-### Import Settings
-
-|---|---|
-|**Asset Import Data**| The importing data and options used for the mesh. |
-
-
-### Meta Data
-
-|---|---|
-|**Meta Data**| This is Meta Data that can be saved with the asset. The meta data is a Blueprintable class derived from the Anim Meta Data class. This allows you to add custom meta data to animation assets (Anim Sequence, Anim Montage, Anim Composite, and Blendspace are all supported). You can query the data from an animation asset in C++ by calling the `GetMetaData()` method, or by using `GetSectionMetaData()` on Anim Montages. |
-
-### Thumbnail
-
-
-|---|---|
-|**Orbit Pitch**| The pitch of the orbit camera around the asset. |
-|**Orbit Yaw**| The yaw of the orbit camera around the asset. |
-|**Orbit Zoom**| The offset from the bounds sphere distance from the asset.|
-
-
-
 
 
 
diff --git a/Engine/Documentation/Source/Programming/Development/BuildingUnrealEngine/BuildingUnrealEngine.INT.udn b/Engine/Documentation/Source/Programming/Development/BuildingUnrealEngine/BuildingUnrealEngine.INT.udn
--- a/Engine/Documentation/Source/Programming/Development/BuildingUnrealEngine/BuildingUnrealEngine.INT.udn
+++ b/Engine/Documentation/Source/Programming/Development/BuildingUnrealEngine/BuildingUnrealEngine.INT.udn
@@ -127,11 +127,7 @@ SkillLevel: Advanced
 		Our developers and support teams currently use the latest version of Ubuntu; as a result, we may not be able to provide support for other Linux distributions (including other versions of Ubuntu).
 		[/REGION]
 		
-		1.  Inside the root directory, run **Setup.sh** from the terminal to setup the files to prepare for generating the project files
-
-			![](RunSetupShellScript_Linux.png)
-
-		1.	Inside the root directory, run **GenerateProjectFiles.sh** from the terminal to generate your project files. 
+		1.	Inside the root directory, run **GenerateProjectFiles.sh** from the terminal to set-up your project files. 
 
 			![](RunGenPrjFilesShellScript_Linux.png)
 
diff --git a/Engine/Documentation/Source/Programming/Modules/API/ModuleAPI.INT.udn b/Engine/Documentation/Source/Programming/Modules/API/ModuleAPI.INT.udn
--- a/Engine/Documentation/Source/Programming/Modules/API/ModuleAPI.INT.udn
+++ b/Engine/Documentation/Source/Programming/Modules/API/ModuleAPI.INT.udn
@@ -2,7 +2,6 @@ Availability:Public
 Title:Module API Specifiers
 Crumbs:%ROOT%, Programming
 Description:Exposing classes and functions to other modules.
-version: 4.9
 
 The easiest way to think about these specifiers is that they are used to tag functions, classes or data as public to your module's DLL file. If you mark a function in the Engine module as ENGINE_API, then any module that imports Engine can access that function directly.
 
diff --git a/Engine/Documentation/Source/Programming/Modules/Gameplay/GameplayModules.INT.udn b/Engine/Documentation/Source/Programming/Modules/Gameplay/GameplayModules.INT.udn
--- a/Engine/Documentation/Source/Programming/Modules/Gameplay/GameplayModules.INT.udn
+++ b/Engine/Documentation/Source/Programming/Modules/Gameplay/GameplayModules.INT.udn
@@ -16,13 +16,13 @@ If your game project is not already created, see [Creating a Game Project](Engin
 
 ## Module Creation
 
-Gameplay modules must contain, at the very least, a header file (.h), C++ file (.cpp), and a build file (*.Build.cs). The header file must be located in the Public folder of the module's directory, i.e. [GameName]\Source\[ModuleName]\Public. This file includes any header files - including the module's auto-generated headers - necessary to compile the classes contained within the module.
+Gameplay modules must contain, at the very least, a header file (`.h`), C++ file (`.cpp`), and a build file (`*.Build.cs`). The header file must be located in the Public folder of the module's directory, i.e. [GameName]\Source\[ModuleName]\Public. This file includes any header files - including the module's auto-generated headers - necessary to compile the classes contained within the module.
 
 	#include "Engine.h"
 	#include "EnginePrivate.h"
 	#include "<ModuleName>Classes.h"
 
-The C++ file is placed in the Private folder of the module's directory, i.e. [GameName] \Source\ [ModuleName]\Private, registers and implements the module.
+The C++ file is placed in the Private folder of the module's directory, i.e. [GameName]\Source\[ModuleName]\Private, registers and implements the module.
 
 [REGION:note]At least one module in your game must be registered using `IMPLEMENT_PRIMARY_GAME_MODULE`. Additional modules can use the alternative `IMPLEMENT_GAME_MODULE` method. See the [Multiple Gameplay Modules](#MultipleGameplayModules) section for more details on using more than one gameplay module in a game project.[/REGION]
 
@@ -32,7 +32,7 @@ The C++ file is placed in the Private folder of the module's directory, i.e. [Ga
 	// Designate the module as primary
 	IMPLEMENT_PRIMARY_GAME_MODULE( <ModuleName>, "<GameName>" );
 
-The build file is placed in the root directory of the gameplay module, i.e. [GameName] \Source\ [ModuleName], and it defines certain pieces of information used by the UnrealBuildTool to compile the module.
+The build file is placed in the root directory of the gameplay module, i.e. [GameName]\Source\[ModuleName], and it defines certain pieces of information used by the UnrealBuildTool to compile the module.
 
 	using UnrealBuildTool;
 
@@ -49,7 +49,7 @@ The build file is placed in the root directory of the gameplay module, i.e. [Gam
 
 Because the new gameplay modules will have UObject code, some configuration is necessary.
 
-1. The module needs to be added in several areas of the DefaultEngine.ini file:
+1. The module needs to be added in several areas of the `DefaultEngine.ini` file:
 
 	The `EditPackages` array in the `[UnrealEd.EditorEngine]` section:
 
@@ -71,7 +71,7 @@ Because the new gameplay modules will have UObject code, some configuration is n
 
 There are philosophical choices about game DLL modularity. Splitting a game up into a lot of DLL files may be more trouble than it is worth, but this is a decision to be made by each individual team based on their needs and principles. Using multiple gameplay modules will result in better link times and faster code iteration, but with more modules you will need to deal with DLL exports and/or interface classes more often. This trade-off is the correct one for engine and editor code, but it is questionable for gameplay.
 
-You can create a primary game module, then any number of additional game-specific modules. You can create *.Build.cs files for these new modules, then add references to these modules to your game's [Target.cs file](Programming/UnrealBuildSystem/TargetFiles) (OutExtraModuleNames array). Be sure to use the appropriate macro for the game modules in the C++ code. At least one module must use the `IMPLEMENT_PRIMARY_GAME_MODULE` macro, while all others should use the `IMPLEMENT_GAME_MODULE` macro. UBT should then automatically discover the modules and compile additional game DLL files.
+You can create a primary game module, then any number of additional game-specific modules. You can create `*.Build.cs` files for these new modules, then add references to these modules to your game's [`Target.cs` file](Programming/UnrealBuildSystem/TargetFiles) (OutExtraModuleNames array). Be sure to use the appropriate macro for the game modules in the C++ code. At least one module must use the `IMPLEMENT_PRIMARY_GAME_MODULE` macro, while all others should use the `IMPLEMENT_GAME_MODULE` macro. UBT should then automatically discover the modules and compile additional game DLL files.
 
 ## Limitations
 
diff --git a/Engine/Documentation/Source/Programming/Online/Interfaces/Session/Session.INT.udn b/Engine/Documentation/Source/Programming/Online/Interfaces/Session/Session.INT.udn
--- a/Engine/Documentation/Source/Programming/Online/Interfaces/Session/Session.INT.udn
+++ b/Engine/Documentation/Source/Programming/Online/Interfaces/Session/Session.INT.udn
@@ -7,8 +7,8 @@ Version: 4.9
 [TOC(start:2)]
 
 
-Matchmaking is the process of matching players with sessions. A session is basically an instance of 
-the game running on the server with a given set of properties which is either advertised so that it can 
+Matchmaking is the process of matching _players_ with _sessions_. A _session_ is basically an instance of 
+the game running on the server with a given set of properties which is either _advertised_ so that it can 
 be found and joined by players wanting to play the game or private so only players who are invited or 
 notified of it in some way can join.
 
@@ -42,7 +42,7 @@ interface must be created. Platforms, in this sense, refer to hardware platforms
 interface will ever exist at a time - the session interface for the platform the engine is currently running on.
 
 While the session interface performs all of the session handling, the game doesn't generally interact directly with 
-it. Instead, the **Game Session**, `AGameSession`, acts as a game-specific wrapper around the session interface 
+it. Instead, the **game session**, `AGameSession`, acts as a game-specific wrapper around the session interface 
 and the game code makes calls to it when it needs to interact with the session. The game session is created and 
 owned by the game mode, `AGameMode`, and also only exists on the server when running an online game.
 
@@ -52,7 +52,7 @@ using a dedicated server.
 
 ## Session Settings
 
-The **Session Settings**, defined by the `FOnlineSessionSettings` class, are a set of properties that determine 
+The **session settings**, defined by the `FOnlineSessionSettings` class, are a set of properties that determine 
 the characteristics of the session. In the base implementation, these are things like:
 
 * Number of players allowed
diff --git a/Engine/Documentation/Source/Programming/Online/OnlineSubsystem.INT.udn b/Engine/Documentation/Source/Programming/Online/OnlineSubsystem.INT.udn
--- a/Engine/Documentation/Source/Programming/Online/OnlineSubsystem.INT.udn
+++ b/Engine/Documentation/Source/Programming/Online/OnlineSubsystem.INT.udn
@@ -119,13 +119,13 @@ Provides the interface for accessing online leaderboards.
 
 #### Voice
 
-Provides the interface for voice communication over network while in the game.
+Provides the interface for voice communication over network while in the game (NOT IMPLEMENTED YET).
 
 <!--Details can be found on the [](Programming/Online/Interfaces/Voice) page.-->
 
 #### Achievements
 
-Provides the interface for reading/writing/unlocking achievements.
+Provides the interface for reading/writing/unlocking achievements (NOT IMPLEMENTED YET).
 
 <!--Details can be found on the [](Programming/Online/Interfaces/Achievements) page.-->
 
diff --git a/Engine/Documentation/Source/Programming/Online/Steam/OnlineSubsystemSteam.INT.udn b/Engine/Documentation/Source/Programming/Online/Steam/OnlineSubsystemSteam.INT.udn
--- a/Engine/Documentation/Source/Programming/Online/Steam/OnlineSubsystemSteam.INT.udn
+++ b/Engine/Documentation/Source/Programming/Online/Steam/OnlineSubsystemSteam.INT.udn
@@ -65,7 +65,7 @@ to make sure it references the new SDK directory location / version
 
 ### INI Configuration
 
-Turn on some settings in the game's DefaultEngine.ini.  
+Turn on some settings in the game's `DefaultEngine.ini`.  
 [REGION:note]
 The SteamDevAppId of 480 is Valve's test app id, shared by everyone.   You will need your own app id eventually, but most features of Steam should work before then.
 [/REGION]
diff --git a/Engine/Documentation/Source/Programming/Online/Types/Types.INT.udn b/Engine/Documentation/Source/Programming/Online/Types/Types.INT.udn
--- a/Engine/Documentation/Source/Programming/Online/Types/Types.INT.udn
+++ b/Engine/Documentation/Source/Programming/Online/Types/Types.INT.udn
@@ -17,7 +17,7 @@ This document serves to define the different data structures that are used throu
 
  There are often platform-specific representations of various constructs that, while they need to be handled by the game, do not need their details exposed. This allows the game to hold onto the data and replicate it to other machines without understanding the details.
  
- Classes deriving from IOnlinePlatformData:
+ Classes deriving from **IOnlinePlatformData**:
  * [FUniqueNetId](#FUniqueNetId)
  * [FSharedContentHandle](#FSharedContentHandle)
  * [FSessionInfo](#FSessionInfo)
@@ -57,7 +57,8 @@ This document serves to define the different data structures that are used throu
 
 **FUniqueNetId** contains a profile service online ID and can represent anything from an individual player, friend or a session.
 
-<!--See: [](Programming/Online/Interfaces/Profile)-->
+<!--See:
+ * [](Programming/Online/Interfaces/Profile)-->
 
 
 ### FSharedContentHandle
@@ -73,7 +74,8 @@ This document serves to define the different data structures that are used throu
 
 **FSessionInfo** contains a session's platform-specific info.  In most cases, this would be information such as a session identifier, host address, or other means of identifying and connecting to a created session.
 
-See: [](Programming/Online/Interfaces/Session)
+See:
+ * [](Programming/Online/Interfaces/Session)
 
 
 ### FEmsFile
diff --git a/Engine/Documentation/Source/Programming/Plugins/Plugins.INT.udn b/Engine/Documentation/Source/Programming/Plugins/Plugins.INT.udn
--- a/Engine/Documentation/Source/Programming/Plugins/Plugins.INT.udn
+++ b/Engine/Documentation/Source/Programming/Plugins/Plugins.INT.udn
@@ -58,9 +58,9 @@ Plugins can have their own Content folder that contains asset files specific to
 section for more information.
 
 [REGION:warning]
-Plugins do not support config files yet. This is something we are considering to add at a later date. 
+Note that Plugins do not support config files yet. This is something we are considering to add at a later date. 
 
-Also, plugins do not yet support their own Derived Data cache distribution. This is something we are looking 
+Also note that plugins do not yet support their own Derived Data cache distribution. This is something we are looking 
 into adding at a later date.
 [/REGION]
 
@@ -72,11 +72,11 @@ the valid search paths for plugins in Unreal Engine.
 
 | Plugin type      | Search path                                        |
 | ---------------- | -------------------------------------------------- |
-| Engine plugin    | /UE4 root/Engine/Plugins/My Engine Plugin/ |
-| Game plugin      | /My project/Plugins/My Game Plugin/        |
+| Engine plugin    | /**UE4 root**/Engine/Plugins/**My Engine Plugin**/ |
+| Game plugin      | /**My project**/Plugins/**My Game Plugin**/        |
 
-You can also organize plugins into subdirectories under the base Plugins folder. The engine will scan all of your sub-folders 
-under the base Plugins folder for plugins to load, but it will never scan into subdirectories beneath where a plugin was already found.
+You can also organize plugins into subdirectories under the base **Plugins** folder. The engine will scan all of your sub-folders 
+under the base **Plugins** folder for plugins to load, but it will never scan into subdirectories beneath where a plugin was already found.
 
 Unreal Engine finds your plugin by searching for .uplugin files on disk. We call these files plugin descriptors. They are text 
 files that provide basic information about your plugin. Plugin descriptors are discovered and loaded automatically by Unreal Engine, 
@@ -101,9 +101,9 @@ For the most part, plugin source file layout is the same as any other C++ module
 Plugins are allowed to declare new UObject types (UCLASS, USTRUCT, etc) in header files in a Classes subdirectory in their module 
 code folders. The Unreal Engine build system will detect these files and generate code as needed to support these UObjects. You will 
 need to follow the normal rules for using UObjects within C++ modules, such as including the generated header file and the module's 
-generated.inl file in one of your module's source files.
+`generated.inl` file in one of your module's source files.
 
-One thing with plugin modules that is a bit different, is Public source file headers. Most plugin modules have no business exporting 
+One thing with plugin modules that is a bit different, is "Public" source file headers. Most plugin modules have no business exporting 
 public APIs in header files in a Public source folder, because they will never be a direct dependency (statically linked against) by 
 engine or game code. So usually your Public source folder will be empty. There are a few exceptions to this rule:
 
@@ -132,7 +132,7 @@ recommend using it.
 
 ## Plugins in your Game Project
 
-Plugins live under a Plugins subfolder in your game project's directory, and will be detected and loaded at startup by the game engine and editor. 
+Plugins live under a "Plugins" subfolder in your game project's directory, and will be detected and loaded at startup by the game engine and editor. 
 
 If the plugin contains modules that have Source folders (and *.Build.cs files), plugin code will be automatically added to generated 
 C++ project files, so that you can work on developing the plugin alongside your game project easily. Whenever you compile your game 
@@ -162,7 +162,7 @@ Here are the steps required to distribute a plugin.
 
 1. **Edit your plugin's descriptor file** (.uplugin file) and make sure the plugin name, modules, version and other settings are configured as desired.
 2. **Delete** your plugin's **Binaries** and **Intermediate** folders. This is just to make sure you are starting with a clean slate.
-3. If your plugin contains a **Source** directory, compile plugin binaries for **Win64 Development** configuration. This configuration will always be needed by an editor on PC platform to load your plugin. 
+3. If your plugin contains a **Source** directory, **compile plugin binaries** for **Win64 Development** configuration. This configuration will always be needed by an editor on PC platform to load your plugin. 
 4. For plugins with **Source**, if you intend to support additional configurations such as Mac, you should build those too.
 5. Make a copy of your plugin to a new "staging" folder outside of your project directory. It does not matter where, but you are going to be making changes to the plugin here.
 6. In the staging directory for your plugin, delete the **Intermediate** directory. These are temporary files that should never be distributed.
@@ -258,8 +258,8 @@ For plugins that contain code, the descriptor file will contain at least one mod
 | Field name		| Info			| Description |
 | -----------------	| -------------	| ----------- |
 | **Name**			| _Required_	| Unique name of this plugin module that will be loaded with the plugin. At runtime, the engine will expect appropriate plugin binaries to exist in the plugin's Binaries folder with the module name as specified here. For modules that have a Source directory, a matching *.Build.cs file is expected to exist within the module's subfolder tree. |
-| **Type**			| _Required_	| Sets the type of module. Valid options are **Runtime**, **RuntimeNoCommandlet**, **Developer**, **Editor**, **EditorNoCommandlet**, and **Program**. This type determines which types of applications this plugin's module is suitable for loading in. For example, some plugins may include modules that are only designed to be loaded when the editor is running. Runtime modules will be loaded in all cases, even in shipped games. Developer modules will only be loaded in development runtime or editor builds, but never in shipping builds. Editor modules will only be loaded when the editor is starting up. Your plugin can use a combination of modules of different types. |
-| **LoadingPhase**	| _Optional_	| If specified, controls when the plugin is loaded at startup. This is an advanced option that should not normally be required. The valid options are **Default** (which is used when no LoadingPhase is specified), **PreDefault**, and **PostConfigInit**. **PostConfigInit** allows the module to be loaded before the engine has finished starting up key subsystems. PreDefault loads just before the normal phase. Typically, this is only needed if you expect game modules to depend directly on content within your plugin, or types declared within your plugin's code. |
+| **Type**			| _Required_	| Sets the type of module. Valid options are **Runtime**, **RuntimeNoCommandlet**, **Developer**, **Editor**, **EditorNoCommandlet**, and **Program**. This type determines which types of applications this plugin's module is suitable for loading in. For example, some plugins may include modules that are only designed to be loaded when the editor is running. **Runtime** modules will be loaded in all cases, even in shipped games. **Developer** modules will only be loaded in development runtime or editor builds, but never in shipping builds. **Editor** modules will only be loaded when the editor is starting up. Your plugin can use a combination of modules of different types. |
+| **LoadingPhase**	| _Optional_	| If specified, controls when the plugin is loaded at startup. This is an advanced option that should not normally be required. The valid options are **Default** (which is used when no LoadingPhase is specified), **PreDefault**, and **PostConfigInit**. **PostConfigInit** allows the module to be loaded before the engine has finished starting up key subsystems. **PreDefault** loads just before the normal phase. Typically, this is only needed if you expect game modules to depend directly on content within your plugin, or types declared within your plugin's code. |
 
 		 
 
@@ -292,7 +292,7 @@ To use these examples as a starting point for your plugin:
 4. **Compile your game project** normally. Unreal Build Tool will detect the plugins and compile them as dependencies to your game!
 5. **Launch the editor** (or the game.)  Your plugin will be initially disabled, but you can turn it on in the editor UI.
 6. **Open the Plugins Editor** (Window -> Plugins), search for your plugin and enable it by clicking the check box.
-7. **Restart the Editor**.  Your plugin will be automatically loaded at startup.
+7. **Restart the Editor**.  Your plugin will be automatically **loaded at startup**.
 
 You can see that the plugins have loaded by opening the **Modules** viewer under the **Window -> Developer Tools** menu. 
 Another way is to use the code debugger to place a breakpoint in the plugin's startup code, such as in *FBlankPlugin::StartupModule()*.
diff --git a/Engine/Documentation/Source/Programming/QuickStart/1/QuickStart_1.INT.udn b/Engine/Documentation/Source/Programming/QuickStart/1/QuickStart_1.INT.udn
--- a/Engine/Documentation/Source/Programming/QuickStart/1/QuickStart_1.INT.udn
+++ b/Engine/Documentation/Source/Programming/QuickStart/1/QuickStart_1.INT.udn
@@ -3,7 +3,7 @@ Title:1 - Essential Project Setup
 Crumbs:%ROOT%, Programming, Programming/QuickStart
 Description:Creating a new project with the specifications we need
 SkillLevel: Beginner
-Version: 4.9
+Version: 4.8
 
 [VAR:Steps]
 [OBJECT:Navigation]
@@ -28,7 +28,7 @@ You should have a copy of **Visual Studio 2013** for Desktop (Preferably Communi
 
 1. Open **Unreal Engine** from the Launcher. The [](Engine/Basics/Projects/Browser) will appear.
 
-1. Click on the **New Project** tab and then select the **C++** tab. From there, select **Basic Code** so we have a clean starting point, and make sure **With Starter Content** is set. We'll need to enter a project name, so we'll use "QuickStart". We can now click **Create Project** and get started.
+1. Click on the **New Project** tab and then select the **C++** tab. From there, select **Basic Code** so we have a clean starting point, and make sure "With Starter Content" is set. We'll need to enter a project name, so we'll use "QuickStart". We can now click **Create Project** and get started.
 
 	[![](CreateProject.png)](CreateProject.png)
 
diff --git a/Engine/Documentation/Source/Programming/QuickStart/2/QuickStart_2.INT.udn b/Engine/Documentation/Source/Programming/QuickStart/2/QuickStart_2.INT.udn
--- a/Engine/Documentation/Source/Programming/QuickStart/2/QuickStart_2.INT.udn
+++ b/Engine/Documentation/Source/Programming/QuickStart/2/QuickStart_2.INT.udn
@@ -3,7 +3,7 @@ Title:2 - Create a C++ Class
 Crumbs:%ROOT%, Programming, Programming/QuickStart
 Description:Creating a C++ class
 SkillLevel: Beginner
-Version: 4.9
+Version: 4.8
 
 [VAR:Steps]
 [OBJECT:Navigation]
@@ -35,6 +35,6 @@ Version: 4.9
  
 	[![](NameActorClass.png)](NameActorClass.png)
 
-Now that we have created a C++ class, we can switch to **Visual Studio** to program it. FloatingActor.cpp will be opened for us automatically, and **Unreal Engine** will automatically compile and reload the code with our new class.
+Now that we have created a C++ class, we can switch to **Visual Studio** to program it. `FloatingActor.cpp` will be opened for us automatically, and **Unreal Engine** will automatically compile and reload the code with our new class.
 
 %Steps%
diff --git a/Engine/Documentation/Source/Programming/QuickStart/3/QuickStart_3.INT.udn b/Engine/Documentation/Source/Programming/QuickStart/3/QuickStart_3.INT.udn
--- a/Engine/Documentation/Source/Programming/QuickStart/3/QuickStart_3.INT.udn
+++ b/Engine/Documentation/Source/Programming/QuickStart/3/QuickStart_3.INT.udn
@@ -3,7 +3,7 @@ Title:3 - Write and Compile C++ Code
 Crumbs:%ROOT%, Programming, Programming/QuickStart
 Description:Writing a C++ class and compiling it with Unreal Engine
 SkillLevel: Beginner
-Version: 4.9
+Version: 4.8
 
 [VAR:Steps]
 [OBJECT:Navigation]
@@ -23,17 +23,17 @@ Version: 4.9
 
 %Steps%
 
-1. In **Visual Studio**, we'll use the **Solution Explorer** pane to find our newly-created C++ files. In our example, they will be named FloatingActor.cpp and FloatingActor.h and will be inside the QuickStart project.
+1. In **Visual Studio**, we'll use the **Solution Explorer** pane to find our newly-created C++ files. In our example, they will be named `FloatingActor.cpp` and `FloatingActor.h` and will be inside the "QuickStart" project.
 
 	[![](SolutionExplorer_Files.png)](SolutionExplorer_Files.png)
 
 	[REGION:note]We're about to write code. All code used in this tutorial will be found at bottom of the page in its current state as of completing that page's instructions.[/REGION]
 
-1. In FloatingActor.h, we'll add the following code just before the closing brace and semicolon at the end of the file:
+1. In `FloatingActor.h`, we'll add the following code just before the closing brace and semicolon at the end of the file:
 
 		float RunningTime;
 
-1. Switching to FloatingActor.cpp, we'll add the following code just before the closing brace at the bottom of **AFloatingActor::Tick**:
+1. Switching to `FloatingActor.cpp`, we'll add the following code just before the closing brace at the bottom of **AFloatingActor::Tick**:
 
 		FVector NewLocation = GetActorLocation();
 		float DeltaHeight = (FMath::Sin(RunningTime + DeltaTime) - FMath::Sin(RunningTime));
@@ -41,7 +41,7 @@ Version: 4.9
 		RunningTime += DeltaTime;
 		SetActorLocation(NewLocation);
 
-	[REGION:note] The code we've just written will cause **FloatingActors** to bob up and down smoothly, using the _RunningTime_ variable we created to keep track of our movement over time. [/REGION] 
+	[REGION:tip] The code we've just written will cause **FloatingActors** to bob up and down smoothly, using the _RunningTime_ variable we created to keep track of our movement over time. [/REGION] 
 
 %Globals:OSSelect%
 
diff --git a/Engine/Documentation/Source/Programming/QuickStart/4/QuickStart_4.INT.udn b/Engine/Documentation/Source/Programming/QuickStart/4/QuickStart_4.INT.udn
--- a/Engine/Documentation/Source/Programming/QuickStart/4/QuickStart_4.INT.udn
+++ b/Engine/Documentation/Source/Programming/QuickStart/4/QuickStart_4.INT.udn
@@ -3,7 +3,7 @@ Title:4 - Test Your Code
 Crumbs:%ROOT%, Programming, Programming/QuickStart
 Description:Creating a test object in your level
 SkillLevel: Beginner
-Version: 4.9
+Version: 4.8
 
 [VAR:Steps]
 [OBJECT:Navigation]
diff --git a/Engine/Documentation/Source/Programming/QuickStart/5/QuickStart_5.INT.udn b/Engine/Documentation/Source/Programming/QuickStart/5/QuickStart_5.INT.udn
--- a/Engine/Documentation/Source/Programming/QuickStart/5/QuickStart_5.INT.udn
+++ b/Engine/Documentation/Source/Programming/QuickStart/5/QuickStart_5.INT.udn
@@ -3,7 +3,7 @@ Title:5 - On Your Own!
 Crumbs:%ROOT%, Programming, Programming/QuickStart
 Description:Ideas for further exploration!
 SkillLevel: Beginner
-Version: 4.9
+Version: 4.8
 
 [VAR:Steps]
 [OBJECT:Navigation]
@@ -24,13 +24,13 @@ Version: 4.9
 
 Using what you have learned, try to do the following:
 
-* Add a [Particle System Component](Engine/Rendering/ParticleSystems) to your FloatingActor. Some pre-built Particle Systems are already included in your project.
-* Use **Unreal Engine's** [UProperty](Programming/UnrealArchitecture/Reference/Properties) macro to expose a variable for the magnitude of your FloatingActor's movement, instead of using a hardcoded value. You might want to check the [%Programming/Tutorials/VariablesTimersEvents:title%](Programming/Tutorials/VariablesTimersEvents) tutorial for help on this topic.
-* Add periodic motion on the X and/or Y axis, and multiply the `DeltaTime` value by a number between 0.6 and 1.4, so your FloatingActor appears to float around freely. This can look great for powerups!
+* Add a [Particle System Component](Engine/Rendering/ParticleSystems) to your **FloatingActor**. Some pre-built Particle Systems are already included in your project.
+* Use **Unreal Engine's** [UProperty](Programming/UnrealArchitecture/Reference/Properties) macro to expose a variable for the magnitude of your **FloatingActor's** movement, instead of using a hardcoded value. You might want to check the [%Programming/Tutorials/VariablesTimersEvents:title%](Programming/Tutorials/VariablesTimersEvents) tutorial for help on this topic.
+* Add periodic motion on the X and/or Y axis, and multiply the `DeltaTime` value by a number between 0.6 and 1.4, so your **FloatingActor** appears to float around freely. This can look great for powerups!
 
 As for the specifics covered in this tutorial:
 
-* For more information on Actors, see the [Actor](API:AActor) manual page.
+* For more information on **Actors**, see the [Actor](API:AActor) manual page.
 * For further tutorials, see the [%Programming/Tutorials:title%](Programming/Tutorials) page.
 
 %Steps%
diff --git a/Engine/Documentation/Source/Programming/QuickStart/QuickStart.INT.udn b/Engine/Documentation/Source/Programming/QuickStart/QuickStart.INT.udn
--- a/Engine/Documentation/Source/Programming/QuickStart/QuickStart.INT.udn
+++ b/Engine/Documentation/Source/Programming/QuickStart/QuickStart.INT.udn
@@ -5,7 +5,7 @@ Description:Create your first code project and add a new C++ class.
 prereq:GettingStarted/Installation
 prereq:Engine/Basics/Projects/Browser
 SkillLevel: Beginner
-Version: 4.9
+Version: 4.8
 
 [VAR:Topic]
 [OBJECT:Topic]
diff --git a/Engine/Documentation/Source/Resources/Showcases/BlueprintMultiplayer/BlueprintMultiplayer.INT.udn b/Engine/Documentation/Source/Resources/Showcases/BlueprintMultiplayer/BlueprintMultiplayer.INT.udn
--- a/Engine/Documentation/Source/Resources/Showcases/BlueprintMultiplayer/BlueprintMultiplayer.INT.udn
+++ b/Engine/Documentation/Source/Resources/Showcases/BlueprintMultiplayer/BlueprintMultiplayer.INT.udn
@@ -8,7 +8,6 @@ Related: Programming/Online
 Related: Programming/Online/Steam
 Related: Gameplay/Networking/Actors
 Related: Resources/ContentExamples/Networking
-Version: 4.9
 
 
 [VAR:Topic]
diff --git a/Engine/Documentation/Source/Resources/Showcases/Effects/Effects.INT.udn b/Engine/Documentation/Source/Resources/Showcases/Effects/Effects.INT.udn
--- a/Engine/Documentation/Source/Resources/Showcases/Effects/Effects.INT.udn
+++ b/Engine/Documentation/Source/Resources/Showcases/Effects/Effects.INT.udn
@@ -7,7 +7,6 @@ Related: Engine/Rendering/ParticleSystems
 Related: Engine/Rendering/ParticleSystems\Cascade
 Related: Engine/Rendering/ParticleSystems\Reference\TypeData\GPUSprites
 Related: Engine/Rendering/ParticleSystems/Reference
-Version: 4.9
 
 
 
@@ -64,7 +63,7 @@ The Particle Effects Example shows off a variety of particle systems that demons
 
 ## Scene Breakdown
 
-This example showcases a cavern level, high up in a snowy arctic mountain range. Within the cave, water cascades across moss-covered rock formations, trickling into rich, flowing pools on the cavern floor. Mysterious and seemingly magical flame sources provide an eerie light. Outside, a harsh blizzard rips through the rocky crags. 
+The Particle Effects Example showcases a cavern level, high up in a snowy arctic mountain range. Within the cave, water cascades across moss-covered rock formations, trickling into rich, flowing pools on the cavern floor. Mysterious and seemingly magical flame sources provide an eerie light. Outside, a harsh blizzard rips through the rocky crags. 
 
 
 ## Particle System Examples
@@ -119,7 +118,7 @@ The documentation included with this example is intended to help you get off the
 
 ![](FindInCB.png)
 
-This will select the Particle System in the Content Browser. **Double-click** its icon in the Content Browser to open the effect up in **Cascade**, Unreal Engine 4's powerful particle editing system.
+This will select the Particle System in the **Content Browser**. **Double-click** its icon in the **Content Browser** to open the effect up in **Cascade**, Unreal Engine 4's powerful particle editing system.
 
 ![](CascadeEditorOpen.png)(w:600)
 
diff --git a/Engine/Source/Runtime/Engine/Private/Collision/CollisionConversions.cpp b/Engine/Source/Runtime/Engine/Private/Collision/CollisionConversions.cpp
--- a/Engine/Source/Runtime/Engine/Private/Collision/CollisionConversions.cpp
+++ b/Engine/Source/Runtime/Engine/Private/Collision/CollisionConversions.cpp
@@ -427,7 +427,7 @@ static void SetHitResultFromShapeAndFaceIndex(const PxShape* PShape,  const PxRi
 	OutResult.FaceIndex = INDEX_NONE;
 }
 
-EConvertQueryResult ConvertQueryImpactHit(const UWorld* World, const PxLocationHit& PHit, FHitResult& OutResult, float CheckLength, const PxFilterData& QueryFilter, const FVector& StartLoc, const FVector& EndLoc, const PxGeometry* const Geom, const PxTransform& QueryTM, bool bReturnFaceIndex, bool bReturnPhysMat)
+void ConvertQueryImpactHit(const UWorld* World, const PxLocationHit& PHit, FHitResult& OutResult, float CheckLength, const PxFilterData& QueryFilter, const FVector& StartLoc, const FVector& EndLoc, const PxGeometry* const Geom, const PxTransform& QueryTM, bool bReturnFaceIndex, bool bReturnPhysMat)
 {
 	SCOPE_CYCLE_COUNTER(STAT_ConvertQueryImpactHit);
 
@@ -436,7 +436,7 @@ EConvertQueryResult ConvertQueryImpactHit(const UWorld* World, const PxLocationH
 	if (bInitialOverlap && Geom != nullptr)
 	{
 		ConvertOverlappedShapeToImpactHit(World, PHit, StartLoc, EndLoc, OutResult, *Geom, QueryTM, QueryFilter, bReturnPhysMat);
-		return EConvertQueryResult::Valid;
+		return;
 	}
 
 	// See if this is a 'blocking' hit
@@ -456,24 +456,10 @@ EConvertQueryResult ConvertQueryImpactHit(const UWorld* World, const PxLocationH
 	OutResult.Location = SafeLocationToFitShape;
 
 	const bool bUsePxPoint = ((PHit.flags & PxHitFlag::ePOSITION) && !bInitialOverlap);
-	if (bUsePxPoint && !PHit.position.isFinite())
-	{
-		OutResult.Reset();
-		UE_LOG(LogCollision, Error, TEXT("ConvertQueryImpactHit() received NaN/Inf for position: %.2f %.2f %.2f"), PHit.position.x, PHit.position.y, PHit.position.z);
-		return EConvertQueryResult::Invalid;
-	}
-
 	OutResult.ImpactPoint = bUsePxPoint ? P2UVector(PHit.position) : StartLoc;
 	
 	// Caution: we may still have an initial overlap, but with null Geom. This is the case for RayCast results.
 	const bool bUsePxNormal = ((PHit.flags & PxHitFlag::eNORMAL) && !bInitialOverlap);
-	if (bUsePxNormal && !PHit.normal.isFinite())
-	{
-		OutResult.Reset();
-		UE_LOG(LogCollision, Error, TEXT("ConvertQueryImpactHit() received NaN/Inf for normal: %.2f %.2f %.2f"), PHit.normal.x, PHit.normal.y, PHit.normal.z);
-		return EConvertQueryResult::Invalid;
-	}
-
 	FVector Normal = bUsePxNormal ? P2UVector(PHit.normal).GetSafeNormal() : -TraceStartToEnd.GetSafeNormal();
 	OutResult.Normal = Normal;
 	OutResult.ImpactNormal = Normal;
@@ -523,15 +509,11 @@ EConvertQueryResult ConvertQueryImpactHit(const UWorld* World, const PxLocationH
 			OutResult.FaceIndex	= PTriMeshGeom.triangleMesh->getTrianglesRemap()[PHit.faceIndex];
 		}
 	}
-
-	return EConvertQueryResult::Valid;
 }
 
-EConvertQueryResult ConvertRaycastResults(bool& OutHasValidBlockingHit, const UWorld* World, int32 NumHits, PxRaycastHit* Hits, float CheckLength, const PxFilterData& QueryFilter, TArray<FHitResult>& OutHits, const FVector& StartLoc, const FVector& EndLoc, bool bReturnFaceIndex, bool bReturnPhysMat)
+void ConvertRaycastResults(const UWorld* World, int32 NumHits, PxRaycastHit* Hits, float CheckLength, const PxFilterData& QueryFilter, TArray<FHitResult>& OutHits, const FVector& StartLoc, const FVector& EndLoc, bool bReturnFaceIndex, bool bReturnPhysMat)
 {
 	OutHits.Reserve(OutHits.Num() + NumHits);
-	EConvertQueryResult ConvertResult = EConvertQueryResult::Valid;
-	bool bHadBlockingHit = false;
 
 	PxTransform PStartTM(U2PVector(StartLoc));
 	for(int32 i=0; i<NumHits; i++)
@@ -539,28 +521,16 @@ EConvertQueryResult ConvertRaycastResults(bool& OutHasValidBlockingHit, const UW
 		FHitResult& NewResult = OutHits[OutHits.AddDefaulted()];
 		const PxRaycastHit& PHit = Hits[i];
 
-		if (ConvertQueryImpactHit(World, PHit, NewResult, CheckLength, QueryFilter, StartLoc, EndLoc, NULL, PStartTM, bReturnFaceIndex, bReturnPhysMat) == EConvertQueryResult::Valid)
-		{
-			bHadBlockingHit |= NewResult.bBlockingHit;
-		}
-		else
-		{
-			// Reject invalid result (this should be rare). Remove from the results.
-			OutHits.Pop(/*bAllowShrinking=*/ false);
-			ConvertResult = EConvertQueryResult::Invalid;
-		}
+		ConvertQueryImpactHit(World, PHit, NewResult, CheckLength, QueryFilter, StartLoc, EndLoc, NULL, PStartTM, bReturnFaceIndex, bReturnPhysMat);
 	}
 
 	// Sort results from first to last hit
 	OutHits.Sort( FCompareFHitResultTime() );
-	OutHasValidBlockingHit = bHadBlockingHit;
-	return ConvertResult;
 }
 
-EConvertQueryResult AddSweepResults(bool& OutHasValidBlockingHit, const UWorld* World, int32 NumHits, const PxSweepHit* Hits, float CheckLength, const PxFilterData& QueryFilter, TArray<FHitResult>& OutHits, const FVector& StartLoc, const FVector& EndLoc, const PxGeometry& Geom, const PxTransform& QueryTM, float MaxDistance, bool bReturnPhysMat)
+bool AddSweepResults(const UWorld* World, int32 NumHits, const PxSweepHit* Hits, float CheckLength, const PxFilterData& QueryFilter, TArray<FHitResult>& OutHits, const FVector& StartLoc, const FVector& EndLoc, const PxGeometry& Geom, const PxTransform& QueryTM, float MaxDistance, bool bReturnPhysMat)
 {
 	OutHits.Reserve(OutHits.Num() + NumHits);
-	EConvertQueryResult ConvertResult = EConvertQueryResult::Valid;
 	bool bHadBlockingHit = false;
 
 	for(int32 i=0; i<NumHits; i++)
@@ -570,24 +540,14 @@ EConvertQueryResult AddSweepResults(bool& OutHasValidBlockingHit, const UWorld*
 		if(PHit.distance <= MaxDistance)
 		{
 			FHitResult& NewResult = OutHits[OutHits.AddDefaulted()];
-			if (ConvertQueryImpactHit(World, PHit, NewResult, CheckLength, QueryFilter, StartLoc, EndLoc, &Geom, QueryTM, false, bReturnPhysMat) == EConvertQueryResult::Valid)
-			{
-				bHadBlockingHit |= NewResult.bBlockingHit;
-			}
-			else
-			{
-				// Reject invalid result (this should be rare). Remove from the results.
-				OutHits.Pop(/*bAllowShrinking=*/ false);
-				ConvertResult = EConvertQueryResult::Invalid;
-			}
-			
+			ConvertQueryImpactHit(World, PHit, NewResult, CheckLength, QueryFilter, StartLoc, EndLoc, &Geom, QueryTM, false, bReturnPhysMat);
+			bHadBlockingHit |= NewResult.bBlockingHit;
 		}
 	}
 
 	// Sort results from first to last hit
 	OutHits.Sort( FCompareFHitResultTime() );
-	OutHasValidBlockingHit = bHadBlockingHit;
-	return ConvertResult;
+	return bHadBlockingHit;
 }
 
 /* Function to find the best normal from the list of triangles that are overlapping our geom. */
diff --git a/Engine/Source/Runtime/Engine/Private/Collision/CollisionConversions.h b/Engine/Source/Runtime/Engine/Private/Collision/CollisionConversions.h
--- a/Engine/Source/Runtime/Engine/Private/Collision/CollisionConversions.h
+++ b/Engine/Source/Runtime/Engine/Private/Collision/CollisionConversions.h
@@ -7,12 +7,6 @@
 
 #include "../PhysicsEngine/PhysXSupport.h"
 
-enum class EConvertQueryResult
-{
-	Valid,
-	Invalid
-};
-
 /** 
  * Util to convert single physX hit (raycast or sweep) to our hit result
  *
@@ -25,14 +19,12 @@ enum class EConvertQueryResult
  * @param	Geom
  * @param	bReturnFaceIndex	True if we want to lookup the face index
  * @param	bReturnPhysMat		True if we want to lookup the physical material
- * @return	Whether result passed NaN/Inf checks.
  */
-EConvertQueryResult ConvertQueryImpactHit(const UWorld* World, const PxLocationHit& PHit, FHitResult& OutResult, float CheckLength, const PxFilterData& QueryFilter, const FVector& StartLoc, const FVector& EndLoc, const PxGeometry* const Geom, const PxTransform& QueryTM, bool bReturnFaceIndex, bool bReturnPhysMat);
+void ConvertQueryImpactHit(const UWorld* World, const PxLocationHit& PHit, FHitResult& OutResult, float CheckLength, const PxFilterData& QueryFilter, const FVector& StartLoc, const FVector& EndLoc, const PxGeometry* const Geom, const PxTransform& QueryTM, bool bReturnFaceIndex, bool bReturnPhysMat);
 
 /** 
  * Util to convert physX raycast results to our hit results
  *
- * @param	OutHasValidBlockingHit set to whether there is a valid blocking hit found in the results.
  * @param	NumHits		Number of Hits
  * @param	Hits		Array of hits
  * @param	CheckLength	Distance of trace
@@ -41,14 +33,12 @@ EConvertQueryResult ConvertQueryImpactHit(const UWorld* World, const PxLocationH
  * @param	EndLoc		End of trace
  * @param	bReturnFaceIndex	True if we want to lookup the face index
  * @param	bReturnPhysMat		True if we want to lookup the physical material
- * @return	Whether all results passed NaN/Inf checks.
  */
-EConvertQueryResult ConvertRaycastResults(bool& OutHasValidBlockingHit, const UWorld* World, int32 NumHits, PxRaycastHit* Hits, float CheckLength, const PxFilterData& QueryFilter, TArray<FHitResult>& OutHits, const FVector& StartLoc, const FVector& EndLoc, bool bReturnFaceIndex, bool bReturnPhysMat);
+void ConvertRaycastResults(const UWorld* World, int32 NumHits, PxRaycastHit* Hits, float CheckLength, const PxFilterData& QueryFilter, TArray<FHitResult>& OutHits, const FVector& StartLoc, const FVector& EndLoc, bool bReturnFaceIndex, bool bReturnPhysMat);
 
 /** 
  * Util to convert physX sweep results to unreal hit results and add to array
  *
- * @param	OutHasValidBlockingHit set to whether there is a valid blocking hit found in the results.
  * @param	NumHits		Number of Hits
  * @param	Hits		Array of hits
  * @param	CheckLength	Distance of trace
@@ -56,9 +46,9 @@ EConvertQueryResult ConvertRaycastResults(bool& OutHasValidBlockingHit, const UW
  * @param	StartLoc	Start of trace
  * @param	EndLoc		End of trace
  * @param	Geom
- * @return	Whether all results passed NaN/Inf checks.
+ * @return	true if any blocking hit was found within MaxDistance.
  */
-EConvertQueryResult AddSweepResults(bool& OutHasValidBlockingHit, const UWorld* World, int32 NumHits, const PxSweepHit* Hits, float CheckLength, const PxFilterData& QueryFilter, TArray<FHitResult>& OutHits, const FVector& StartLoc, const FVector& EndLoc, const PxGeometry& Geom, const PxTransform& QueryTM, float MaxDistance, bool bReturnPhysMat);
+bool AddSweepResults(const UWorld* World, int32 NumHits, const PxSweepHit* Hits, float CheckLength, const PxFilterData& QueryFilter, TArray<FHitResult>& OutHits, const FVector& StartLoc, const FVector& EndLoc, const PxGeometry& Geom, const PxTransform& QueryTM, float MaxDistance, bool bReturnPhysMat);
 
 /** 
  * Util to convert physX overlap query to our overlap result
diff --git a/Engine/Source/Runtime/Engine/Private/Collision/PhysXCollision.cpp b/Engine/Source/Runtime/Engine/Private/Collision/PhysXCollision.cpp
--- a/Engine/Source/Runtime/Engine/Private/Collision/PhysXCollision.cpp
+++ b/Engine/Source/Runtime/Engine/Private/Collision/PhysXCollision.cpp
@@ -991,17 +991,7 @@ bool RaycastSingle(const UWorld* World, struct FHitResult& OutHit, const FVector
 			if (bHaveBlockingHit) // If we got a hit
 			{
 				PxTransform PStartTM(U2PVector(Start));
-				if (ConvertQueryImpactHit(World, PHit, OutHit, DeltaMag, PFilter, Start, End, NULL, PStartTM, Params.bReturnFaceIndex, Params.bReturnPhysicalMaterial) == EConvertQueryResult::Invalid)
-				{
-					bHaveBlockingHit = false;
-					UE_LOG(LogCollision, Error, TEXT("RaycastSingle resulted in a NaN/INF in PHit!"));
-#if ENABLE_NAN_DIAGNOSTIC
-					UE_LOG(LogCollision, Error, TEXT("--------TraceChannel : %d"), (int32)TraceChannel);
-					UE_LOG(LogCollision, Error, TEXT("--------Start : %s"), *Start.ToString());
-					UE_LOG(LogCollision, Error, TEXT("--------End : %s"), *End.ToString());
-					UE_LOG(LogCollision, Error, TEXT("--------%s"), *Params.ToString());
-#endif
-				}
+				ConvertQueryImpactHit(World, PHit, OutHit, DeltaMag, PFilter, Start, End, NULL, PStartTM, Params.bReturnFaceIndex, Params.bReturnPhysicalMaterial);
 			}
 		}
 #endif //WITH_PHYSX
@@ -1180,17 +1170,7 @@ bool RaycastMulti(const UWorld* World, TArray<struct FHitResult>& OutHits, const
 
 		if (NumHits > 0)
 		{
-			if (ConvertRaycastResults(bBlockingHit, World, NumHits, PHits, DeltaMag, PFilter, OutHits, Start, End, Params.bReturnFaceIndex, Params.bReturnPhysicalMaterial) == EConvertQueryResult::Invalid)
-			{
-				// We don't need to change bBlockingHit, that's done by ConvertRaycastResults if it removed the blocking hit.
-				UE_LOG(LogCollision, Error, TEXT("RaycastMulti resulted in a NaN/INF in PHit!"));
-#if ENABLE_NAN_DIAGNOSTIC
-				UE_LOG(LogCollision, Error, TEXT("--------TraceChannel : %d"), (int32)TraceChannel);
-				UE_LOG(LogCollision, Error, TEXT("--------Start : %s"), *Start.ToString());
-				UE_LOG(LogCollision, Error, TEXT("--------End : %s"), *End.ToString());
-				UE_LOG(LogCollision, Error, TEXT("--------%s"), *Params.ToString());
-#endif
-			}
+			ConvertRaycastResults(World, NumHits, PHits, DeltaMag, PFilter, OutHits, Start, End, Params.bReturnFaceIndex, Params.bReturnPhysicalMaterial);
 		}
 
 		bHaveBlockingHit = bBlockingHit;
@@ -1349,6 +1329,23 @@ bool GeomSweepSingle(const UWorld* World, const struct FCollisionShape& Collisio
 		PxSweepHit PHit;
 		bHaveBlockingHit = SyncScene->sweepSingle(PGeom, PStartTM, PDir, DeltaMag, POutputFlags, PHit, PQueryFilterData, &PQueryCallbackSweep);
 
+		// Emergency NaN guarding. If we get a NaN back, act like there was no hit
+		if (PHit.position.isFinite() == false || PHit.normal.isFinite() == false)
+		{
+			bHaveBlockingHit = false;
+#if ENABLE_NAN_DIAGNOSTIC
+			UE_LOG(LogCollision, Error, TEXT("GeomSweepSingle resulted in a NaN/INF in PHit!"));
+			UE_LOG(LogCollision, Error, TEXT("--------TraceChannel : %d"), (int32)TraceChannel);
+			UE_LOG(LogCollision, Error, TEXT("--------Start : %s"), *Start.ToString());
+			UE_LOG(LogCollision, Error, TEXT("--------End : %s"), *End.ToString());
+			UE_LOG(LogCollision, Error, TEXT("--------%s"), *Params.ToString());
+
+			//convert to FVector will do further diagnostic and print as needed
+			FVector NanCheckVectorPosition = P2UVector(PHit.position);
+			FVector NanCheckVectorNormal = P2UVector(PHit.normal);
+#endif
+		}
+
 		if (!bHaveBlockingHit)
 		{
 			// Not using anything from this scene, so unlock it.
@@ -1381,17 +1378,7 @@ bool GeomSweepSingle(const UWorld* World, const struct FCollisionShape& Collisio
 
 		if(bHaveBlockingHit) // If we got a hit, convert it to unreal type
 		{
-			if (ConvertQueryImpactHit(World, PHit, OutHit, DeltaMag, PFilter, Start, End, &PGeom, PStartTM, false, Params.bReturnPhysicalMaterial) == EConvertQueryResult::Invalid)
-			{
-				bHaveBlockingHit = false;
-				UE_LOG(LogCollision, Error, TEXT("GeomSweepSingle resulted in a NaN/INF in PHit!"));
-#if ENABLE_NAN_DIAGNOSTIC
-				UE_LOG(LogCollision, Error, TEXT("--------TraceChannel : %d"), (int32)TraceChannel);
-				UE_LOG(LogCollision, Error, TEXT("--------Start : %s"), *Start.ToString());
-				UE_LOG(LogCollision, Error, TEXT("--------End : %s"), *End.ToString());
-				UE_LOG(LogCollision, Error, TEXT("--------%s"), *Params.ToString());
-#endif
-			}
+			ConvertQueryImpactHit(World, PHit, OutHit, DeltaMag, PFilter, Start, End, &PGeom, PStartTM, false, Params.bReturnPhysicalMaterial);
 		}
 	}
 
@@ -1523,17 +1510,7 @@ bool GeomSweepMulti_PhysX(const UWorld* World, const PxGeometry& PGeom, const Px
 		// Convert all hits to unreal structs. This will remove any hits further than MinBlockDistance, and sort results.
 		if (NumHits > 0)
 		{
-			if (AddSweepResults(bBlockingHit, World, NumHits, PHits, DeltaMag, PFilter, OutHits, Start, End, PGeom, PStartTM, MinBlockDistance, Params.bReturnPhysicalMaterial) == EConvertQueryResult::Invalid)
-			{
-				// We don't need to change bBlockingHit, that's done by AddSweepResults if it removed the blocking hit.
-				UE_LOG(LogCollision, Error, TEXT("GeomSweepMulti resulted in a NaN/INF in PHit!"));
-#if ENABLE_NAN_DIAGNOSTIC				
-				UE_LOG(LogCollision, Error, TEXT("--------TraceChannel : %d"), (int32)TraceChannel);
-				UE_LOG(LogCollision, Error, TEXT("--------Start : %s"), *Start.ToString());
-				UE_LOG(LogCollision, Error, TEXT("--------End : %s"), *End.ToString());
-				UE_LOG(LogCollision, Error, TEXT("--------%s"), *Params.ToString());
-#endif
-			}
+			bBlockingHit |= AddSweepResults(World, NumHits, PHits, DeltaMag, PFilter, OutHits, Start, End, PGeom, PStartTM, MinBlockDistance, Params.bReturnPhysicalMaterial);
 		}
 	}
 
