Branch snapshot for CL 2723545

[CL 2723545 in Main branch]

diff --git a/Engine/Documentation/Source/Engine/Animation/AnimMontage/AnimMontage.INT.udn b/Engine/Documentation/Source/Engine/Animation/AnimMontage/AnimMontage.INT.udn
--- a/Engine/Documentation/Source/Engine/Animation/AnimMontage/AnimMontage.INT.udn
+++ b/Engine/Documentation/Source/Engine/Animation/AnimMontage/AnimMontage.INT.udn
@@ -11,28 +11,34 @@ Related: Engine/Content/Tools/MayaRiggingTool
 Related: Gameplay/Input
 version: 4.9
 
+
+[TOC(start:2 end:3)]
+
+
+
 [EXCERPT:Intro]
-**Animation Montages** (or just Montages for short) are a multipurpose tool that allows for a wide variety of animation effects, primarily related to exposing animation controls within code or [Blueprint Visual Scripting](Engine/Blueprints). It can also be used to create a wide variety of animation effects including intelligent loops of animation, logic-based animation switching and much more.
+**Animation Montages** (or just Montages for short) are a multipurpose tool that allows for a wide variety of animation effects, primarily related to exposing animation controls within code or Blueprint. It can also be used to create a wide variety of animation effects including intelligent loops of animation, logic-based animation switching, root motion handling, and much more.
 [/EXCERPT:Intro]
 
-[TOC(start:2 end:2)]
-
-AnimMontages are animation assets that can be created and exist within the **Content Browser**. These assets can then be plugged directly into the AnimGraph in run-time and you can modify any state of it. For example, you can jump to different sections or you can re-link different sections. This is mostly for code driven animations or one-off animations such as melee attacks, allowing you to control triggers, stop when you want, or change between states (looping or non-looping). 
+AnimMontages are animation assets that can be created and exist within the **Content Browser**. These assets can then be plugged directly into the AnimGraph in run-time and you can modify any state of it. For example, you can jump to different sections or you can re-link different sections. This is mostly for code driven animations or one-off animations such as melee attacks, allowing you to control triggers, stop when you want, or change state (looping or non-looping) between (see the image below). 
 
 ![](Montage_Screen2.png)
 
-The image above is a melee attack with 3 sections [start, loop, and end]. When the player **Left-clicks**, it triggers the _Start_ section by default when you ask to play the montage. The middle section called _Loop_ (which is set to loop), immediately follows the _Start_ section and will loop there forever. If the player lets the mouse button go, it will stop, but you do not want to stop right away since the animation will pop in the middle of the loop. You would like to relink the _Loop_ to the _End_ section, where it will then transition out to the end and finish the animation. 
+The image above is a melee attack with 3 sections [start, loop, and end]. When the player **Left-clicks**, it triggers the start section by default when you ask to play the montage. Now the middle section is set to loop. When the start is done, it will transition to loop, and will loop there forever. If the player lets the mouse button go, it will stop, but you do not want to stop right away since the animation will pop in the middle of the loop. You would like to relink the loop to the end, where it will then transition out to the end and finish the animation. 
 
 Some additional uses for Montages include:
 * The ability to play an animation from within an AnimBlueprint's **EventGraph**.
 * Chaining together a complex sequence of animations that you want to think of as a single animation.
 * Looping only a specific portion of animation or animations based on code or Blueprint script.
 * Handling event-based switching of multiple animations based on code or Blueprint script.
+* Proper handling of [Root Motion](Engine/Animation/RootMotion) for your characters.
 * The ability to assign complex animation sequences to named slots that can be switched between in code or Blueprints.
 * Precise switching between various AnimSequences based on code or Blueprint script.
 
-[REGION:note]
-An Animation Montage applied to a character can also be seen on the [Animation Content Examples](Resources\ContentExamples\Animation) page under section 1.5.
+So, as you can see, there are a lot of different things you can do with Montages. This can easily make Montages seem overwhelming as a system. However, if you really boil it down, the key point of Montages is to expose animation controls to code or Blueprints. Nearly everything else the system can do follows from that single point.
+
+[REGION:tip]
+An example of an Animation Montage applied to a character can also be seen on the [Animation Content Examples](Resources\ContentExamples\Animation) page under section 1.5.
 [/REGION]
 
 
@@ -44,29 +50,28 @@ There are a few details to keep in mind regarding AnimMontages:
 * [INCLUDE:#MontageReplication]
 
 
+
 ## Montage Properties
 
 The following is a breakdown of the AnimMontage asset properties. These are available in **Persona** when looking at a Montage, and can also be accessed by **Right-clicking** on a Montage within the **Content Browser** and choosing Properties from the context menu.
 
 ![](MontageProperties.png)
 
+
+
 |Montage Properties||
 |---|---|
 |[REGION:tablesection]Montage[/REGION]||
-|**Blend In Time**| The amount of time at the beginning of Montage playback during which the character will blend in from its current pose.|
-|**Blend Out Time**|The amount of time at the end of Montage playback during which the character will blend back to its original pose.|
-|**Blend Out Trigger Time**| Time from the Sequence End to trigger blend out. < 0 uses BlendOutTime, so BlendOut finishes as the Montage Ends. > = 0 uses SequenceEnd-BlendOutTriggerTime to trigger blend out. |
+|**Blend In Time**|This is an amount of time at the beginning of Montage playback during which the character will blend in from its current pose.|
+|**Blend Out Time**|This is an amount of time at the end of Montage playback during which the character will blend back to its original pose.|
+|[REGION:tablesection]Root Motion[/REGION]||
+|**Enable Root Motion Translation**|Enables handling of root motion translation, cancelling out any translation applied to the root so that it can then be applied back to the character's collision capsule. Please see [](Engine/Animation/RootMotion) for more details.|
+|**Enable Root Motion Rotation**|Enables handling of root motion rotation, cancelling out any rotation applied to the root so that it can then be applied back to the character's collision capsule. Please see [](Engine/Animation/RootMotion) for more details.|
 |[REGION:tablesection]Additive Settings[/REGION]||
 |**Preview Base Pose**|Sets a base preview pose used for additive Blend Spaces.|
 |[REGION:tablesection]Animation[/REGION]||
 |**Rate Scale**|A multiplier value for how fast the Montage will play back. Default is 1.0.|
 |**Skeleton**|Contains the skeleton associated with this Montage. Cannot be changed in the editor.|
-|[REGION:tablesection]Meta Data[/REGION]||
-|**Meta Data**| This is Meta Data that can be saved with the asset. The meta data is a Blueprintable class derived from the Anim Meta Data class. This allows you to add custom meta data to animation assets (Anim Sequence, Anim Montage, Anim Composite, and Blendspace are all supported). You can query the data from an animation asset in C++ by calling the `GetMetaData()` method, or by using `GetSectionMetaData()` on Anim Montages. |
-|[REGION:tablesection]Thumbnail[/REGION]||
-|**Orbit Pitch**| The pitch of the orbit camera around the asset. |
-|**Orbit Yaw**| The yaw of the orbit camera around the asset. |
-|**Orbit Zoom**| The offset from the bounds sphere distance from the asset.|
 
 ## Montage UI
 
@@ -91,10 +96,11 @@ The Montage Area breaks down like so:
 [/REGION]
 
 1. **Section Track** - Shows any [Sections](#Sections) that have been defined for this Montage. Sections can be dragged to different positions along the timeline with the left mouse button.
-1. **Slot Track** - Shows the current [Slot](#Slots), along with the Slot name on the right. This Slot may be filled with as many animations as you desire; they will be played back in order. Notice that multiple animations will have alternating positions in the Slot Track - top, then bottom, then top again, and so on. This is to help you differentiate between different animations.
+1. **Slot Track** - Shows the current Slot, along with the Slot name on the right. This Slot may be filled with as many animations as you desire; they will be played back in order. Notice that multiple animations will have alternating positions in the Slot Track - top, then bottom, then top again, and so on. This is to help you differentiate between different animations.
+1. **Branch Point Track** - Shows any Branch Points that have been defined for this Montage. Branch Points can be dragged to different positions along the timeline with the left mouse button.
 
 [REGION:note]
-You may have as many Slot Tracks as you like within a single Montage, each with their own name and containing their own unique animations. However, you may only have one Section Track for each Montage.
+You may have as many Slot Tracks as you like within a single Montage, each with their own name and containing their own unique animations. However, you may only have one Section Track and one Branch Point track for each Montage.
 [/REGION]
 
 
@@ -115,7 +121,7 @@ Within a Montage, a **Slot** is just a single track that can hold any number of
 It is an important point to keep in mind that although much of your Montage control will take place in the Animation Blueprint's **Event Graph**, Slots are actually handled within the **Anim Graph**. This is done by way of the Slot node, which simply takes in the name for a Slot. By positioning this node at a strategic point along your AnimGraph's execution, you can have multiple Montages that utilize the same Slot name.
 [/EXCERPT:Slots]
 
-<!---------------- Branch moved to Notify page.
+
 ### Branch Points
 
 Branch Points allow you to create events that coincide with animation playback. These Branch Point Events can be used in code or Blueprint to cause other things to happen, but specifically, Branch Points are useful for switching to other animation Sections within a Montage. 
@@ -128,8 +134,6 @@ High precision is important when you need to jump to a specific animation at a v
 Due to their synchronous nature and the precision resulting from it, Branch Points are more performance expensive than Notifies. You should only use them when an event must be fired at a precise moment along the animation timeline, such as jumping directly to another animation that matches up frame-to-frame. If being off by a frame (or some percentage of one) is not important, you should use Notifies instead.
 [/REGION]
 
--------------------->
-
 ## Sections Area
 
 The Sections area is where you can establish relationships between the Sections you define in the Montage Area. For instance, you may want a certain Section of animation (or group of Sections) to play in a specific sequence, or even loop. 
@@ -143,6 +147,7 @@ The Sections area is where you can establish relationships between the Sections
 1. **Section Buttons** - In this area, you will see a button for each of the Sections you define in the Montage area. By choosing an existing Section and then clicking the one of these buttons, you associate the Section for that button with the selected track. For instance, in the image above, we have associated Swing2 with Swing1. The order actually goes Swing1, Swing2, and then Swing1 again, which causes a loop. See the [Looping](#Looping) section below for details.
 1. **Section Association Tracks** - This is where you can visualize and preview the relationships between animation Sections. By clicking on the **Preview** buttons, you can see the result of each individual track, or click the **Preview All Sections** button and see all Sections play back one after the other.
 
+
 ### Looping
 
 Sections can be set up to loop indefinitely, which is extremely useful for any action that you need to repeat. By associating the same Section more than once in a Section Association Track, you cause that association to run in a loop. The section will turn blue to show this. As an example, consider an animation in which a character is reloading a shotgun, one shell at a time. You could take just the section in which the character inserts a shell, and loop it. Then, by using Notifies, you could create Notify Events in Blueprint that increment the ammo count each time the animation plays through. Once that count reaches a set number (full ammo), you could then switch to an animation of the character closing the receiver and returning to idle.
@@ -153,7 +158,7 @@ The Element Timing area takes information from the Montage area and the Notifies
 
 ![](ElementTiming1.png)
 
-Each node in the track is given a number denoting the trigger order of that object in the entire montage, and are colored to match what they describe. The default colors are:
+Each node in the track is given a number denoting the trigger order of that object in the entire montage, and are coloured to match what they describe. The default colours are:
 * Red - Notifies, Notify states and notify state ending markers
 * Light Blue - Branching point notifies
 * Green - Montage sections
@@ -186,7 +191,6 @@ For more information, see [](Engine/Animation/Sequences\Notifies).
 
 For more information, see [](Engine/Animation/Sequences\Curves).
 
-<!-------------------- Example removed. This should be a HowTo.
 
 ## Montage Practical Example
 
@@ -294,7 +298,6 @@ Dragging off the **IsAttacking** node which extends from the **Cast To** node, w
 
 And that is it! If we compile, we will now see that the character continues to swing as long as the mouse button is down, and performs an intelligent end animation when the mouse button is released!
 
-
 [OBJECT:EmbeddedVideo]
 [PARAMLITERAL:width]
 640
@@ -307,12 +310,13 @@ And that is it! If we compile, we will now see that the character continues to s
 [/PARAMLITERAL]
 [/OBJECT]
 
--------------------->
-
+<!--------------------
 
 
+[EXCERPT:MontageReplication]At this time, root motion is replicated across the server. However the AnimMontage is not. That means that users utilizing root motion will need to make sure they are replicating the state of the AnimMontage across the network. (Position, play rate, etc.).[/EXCERPT:MontageReplication]
 
 
+-------------------->
 
 
 
diff --git a/Engine/Documentation/Source/Engine/Animation/Sequences/Notifies/AnimNotifies.INT.udn b/Engine/Documentation/Source/Engine/Animation/Sequences/Notifies/AnimNotifies.INT.udn
--- a/Engine/Documentation/Source/Engine/Animation/Sequences/Notifies/AnimNotifies.INT.udn
+++ b/Engine/Documentation/Source/Engine/Animation/Sequences/Notifies/AnimNotifies.INT.udn
@@ -3,10 +3,6 @@ Title:Animation Notifications (Notifies)
 Crumbs: %ROOT%, Engine, Engine/Animation, Engine/Animation/Sequences
 Description:Notifies are a system for setting up and receiving events in Animation Sequences to perform external actions.
 version: 4.9
-Related: Engine/Animation/Sequences/Editor
-Related: Engine/Animation/AnimMontage
-Related: Engine/Animation
-Related: Engine/Animation/AnimHowTo
 
 [VAR:Topic]
 [OBJECT:Topic]
@@ -29,38 +25,40 @@ Related: Engine/Animation/AnimHowTo
 [/VAR]
 
 
+[TOC (start:2 end:2)] 
+
+
 [EXCERPT:Intro]
 **Animation Notifications** (**AnimNotifies** or just **Notifies**) provide a way for animation programmers to setup events to occur at specific points during an **Animation Sequence**. Notifies are commonly used to add effects like footstep sounds during a walk or run animation or spawning a particle system during animation. However, they can have any number of different uses, as the system can be extended with custom notification types to cater to the needs of any type of game.
 [/EXCERPT:Intro]
 
-[TOC (start:2 end:2)] 
+Creating and editing Notifies is performed within the **Notify Editor** in **Persona**. For more information, refer to the [Editing Animation Sequences](Engine/Animation/Sequences/Editor) documentation. 
 
-## Notify Types
+## Animation Notify Types
 
-Different types of Notifies cause different events to be triggered. Camera effects, particle effects, sounds, etc. can all be triggered at any point during an animation through the use of a Notify. Adding a Notify can be done by **Right-clicking** on a **Notify Track** at the point in the animation you want the Notify to occur, and then choosing one of the type of Notify to add.
+Different types of Notifies cause different events to be triggered. Camera effects, particle effects, sounds, etc. can all be triggered at any point during an animation. 
+
+Below you will find the currently available Notifies and how they are used. 
 
-[REGION:note]
-For more information on adding/editing Notifies, see the [](Engine/Animation/Sequences/Editor) page. 
-[/REGION]
 
 ### Play Particle Effect
 
-The **PlayParticleEffect** Notify spawns and plays a particle system at a given socket location and at a specific point along an animation's timeline. This is useful for such effects as causing a cloud of dust to appear around a player's feet when they move, fire to spew from an exhaust pipe when the player uses a Nitro-Boost power-up, or other such effects.
+The **PlayParticleEffect** notify spawns and plays a particle system at a given socket location and at a specific point along an animation's timeline. This is useful for such effects as causing a cloud of dust to appear around a player's feet when they move, fire to spew from an exhaust pipe when the player uses a Nitro-Boost power-up, or other such effects.
 
 ![](PlayParticleEffectNotify.png)
 
-A PlayParticleEffect Notify does not use the Animation Blueprint to control behavior. Instead, it uses a few simple properties defined in the **Details** panel to determine what Particle System is used and from what Socket the particles will originate.
+A PlayParticleEffect notify does not use the Animation Blueprint to control behavior. Instead, it uses a few simple properties defined in the **Details** panel to determine what Particle System is used and from what Socket the particles will originate.
 
 ![](PlayParticleEffectDetails.png)
 
 | **Property** | **Description** |
 | --- | --- |
 | **Socket Name** | This contains the Skeletal Mesh socket that will be used as the spawn point for the particle system. |
-| **PSTemplate** | This holds the particle system that will be spawned at the Notifies set point in time. |
+| **PSTemplate** | This holds the particle system that will be spawned at the notify's set point in time. |
 | **Attached** | Determines if the effect should attach to the bone/socket. |
 | **Location Offset** | Location offset from the socket. |
 | **Rotation Offset** | Rotation offset from the socket. |
-| **Notify Color** | Changes the color of the Notify on the Notify bar. |
+| **Notify Color** | Changes the color of the notify on the notify bar. |
 
 Once set up, the associated Particle System will play, using the Socket location as the spawn point.
 
@@ -69,11 +67,11 @@ Once set up, the associated Particle System will play, using the Socket location
 
 ### Play Sound
 
-The **PlaySound** Notify provides a simple means to get your sound effects synced with animation. Much like Play Particle Effect, this Notify has a few simple properties that can be set in the **Details** panel, which allow for fast and easy sound effect setup. This is useful for such things as mechanical whirrs when a robot moves, the rustle of cloth during character motion, and similar types of effects.
+The **PlaySound** notify provides a simple means to get your sound effects synced with animation. Much like Play Particle Effect, this notify has a few simple properties that can be set in the **Details** panel, which allow for fast and easy sound effect setup. This is useful for such things as mechanical whirrs when a robot moves, the rustle of cloth during character motion, and similar types of effects.
 
 ![](PlaySound.png)
 
-The **Play Sound** Notify uses a few simple properties defined in the **Details** panel to determine what sound effect is used and from what **Socket** or **Bone** the sound will originate.
+The **Play Sound** notify uses a few simple properties defined in the **Details** panel to determine what sound effect is used and from what **Socket** or **Bone** the sound will originate.
 
 ![](PlaySoundDetails.png)
 
@@ -83,53 +81,68 @@ The **Play Sound** Notify uses a few simple properties defined in the **Details*
 | **Volume Multiplier** | This scaling multiplier is applied to the volume of the sound effect. |
 | **Pitch Multiplier** | This scaling multiplier is applied to the pitch of the sound effect. |
 | **Follow** | If checked, the sound effect source will follow the Static Mesh as it moves. If unchecked, the sound remains behind at the location it was spawned. |
-| **Notify Color** | Changes the color of the Notify on the Notify bar. |
+| **Notify Color** | Changes the color of the notify on the notify bar. |
 
-### Custom Notifies
+### Custom
 
-The **Custom** Notify serves as a catch-all for anything an artist may want to do within a Animation Blueprint at a certain point in an animation. Custom Notifies are created by **Right-clicking** on a Notify track and choosing **New Notify** from the context menu.
+The **Custom** notify serves as a catch-all for anything an artist may want to do within a Animation Blueprint at a certain point in an animation. Custom notifies are created by **Right-clicking** on a notify track and choosing **New Notify** from the context menu.
 
 ![](MySpecialNotify.png)
 
-Custom Notifies allow you to create a **Notify Event** of the same name in the Event Graph or Transition Graphs in the [State Machine](Engine/Animation/StateMachines) for a Skeletal Mesh's [Animation Blueprint](Engine\Animation\AnimBlueprints). To add your Custom Notify, **Right-click** anywhere in your Event or Transition Graphs and there is a section for adding Animation Notifies events.
+As soon as you do this, you will be required to supply a name for the new notify.
 
-![](rightClickAddNotify.png)
+![](NotifyName.png)
+
+After the name is typed in, you will see the new notify appear in the timeline.
+
+![](MyNotifyIn.png)
+
+Custom notifies will not have any properties to edit in the **Details** panel. Instead, they simply allow you to create a notify event of the same name in the **Event Graph** or any of the **Transitions** in a state machine of the Skeletal Mesh's Animation Blueprint.
 
-After adding the Custom Notify, an Event Node will be added to the graph which will be executed when the Notify is called. 
+Switch to the **Graph** edit mode
+<!--![](SwitchToAnimBP.png) Commented out picture as it didn't seem useful--> 
+
+**Right-click** anywhere in the **Event Graph** and there is a section for adding Animation Notifies events.
+
+![](rightClickAddNotify.png)
 
 ![](MySpecialNotifyBlueprint.png)
 
-## Notify States 
+## Animation Notify States 
 
-**AnimNotifyStates** (**Notify States**) work much like the standard Notifies above. They have 3 distinct events: a start, the tick, and an end. They start straightforward, firing at the moment the Notify begins and ends, and the **Event Graphs** in them fire off when their time comes up in the animation. The tick fires off every animation update until the end event is hit. The major difference between normal Notifies and Notify States is that Notify States are self-contained [Blueprints](Engine/Blueprints).
+**AnimNotifyStates** (**Notify States**) work much like the standard notifies above.  They have 3 distinct events: a start, the tick, and an end. They start straightforward, firing at the moments the notify begins and ends, and the **Event Graphs** in them fire off when their time comes up in the animation. The tick fires off every animation update until the end event is hit. The major difference between the standard Notifies is that AnimNotifyStates are self-contained Blueprints.
 
-Creating Notify States (which are the same as [Native Notify States](#nativenotifystates) but are custom Blueprints that you can provide any logic you need to occurring during an animation) is done the same as adding a regular Notify by **Right-clicking** in the **Notify Track** of an animation and selecting **Add Notify State...** then selecting your custom Notify.
+To use an AnimNotifyState:
 
-[REGION:note]
-To use a custom Notify, you must first [create a Blueprint Class](Engine/Blueprints/UserGuide/Types/ClassBlueprint/Creation) of the AnimNotifyState class. 
-[/REGION]
+1. Create a new AnimNotifyState asset in the **Content Browser** by creating a new **Blueprint** that inherits from `AnimNotifyState`:
 
-![](addNotifyState.png)
+	![](newBP.png)
 
-An AnimNotifyState Blueprint Class called **MyAnimNotifyState** (above) was created so that we may select it as the Notify State to use in our animation. Once you have created your Anim Notify State Blueprint, you can open it and begin implementing the functions that determine how your Notify is used (below).
+	![](AnimNotifyStateCreate.png)
 
-![](animNotifyStateImplementFunction.png)
+1. On the Animation Sequence or Animation Montage, you can **Right-click** in the notify area and now add your new AnimNoitfyState.
 
-It is not necessary to implement every function, but at the bare minimum you should implement **Received Notify Tick** or both **Received Notify Start** and **Received Notify End**. If implementing just **Received Notify Start** or **Received Notify End**, you could simply use a standard Animation Notify.
+	![](addNotifyState.png)
 
-![](Recieved_NotifyStart.png)
+1. You can set the length of the AnimNotifyState by **Right-clicking** on it and choosing **Set AnimNotifyState Duration**.
 
-_The Notify Start Override._
+	![](AnimNotifyStateSetLength.png)
 
-![](Recieved_NotifyEnd.png)
+	![](duration.png)
 
-_The Notify End Override._
+1. Finally from the **Content Browser**, **Double Left-click** your AnimNotifyState to bring up its Blueprint so you can start defining what if does.
+	
 
-![](Recieved_NotifyTick.png)
+![](animNotifyStateImplementFunction.png)
 
-_The Notify Tick Override._
 
-Below is a description of how each pin type is used in the override functions:
+At this point you can start implementing the 3 functions of the **AnimNotifyState**. It is not necessary to implement every function, but at the bare minimum you should implement **Recieved Notify Tick** or both **Recieved Notify Start** and **Recieved Notify End**. If implementing just **Recieved Notify Start** or **Recieved Notify End**, you could simply use a standard Animation Notify.
+
+![](Recieved_NotifyStart.png)
+
+![](Recieved_NotifyEnd.png)
+
+![](Recieved_NotifyTick.png)
 
 | **Pin** | **Description** |
 | --- | ----------- |
@@ -163,21 +176,20 @@ When choosing to add a Notify State, you can also assign one of two **Native Not
 
 ### Timed Particle Effects
 
-The Timed Particle Effect is very similar to the **Play Particle Effect** Notify in that they both allow you to play a Particle Effect at any given point in an animation, the difference however is that with a Timed Particle Effect, you are given the ability to define the length of the effect inside the Notify itself. 
+The Timed Particle Effect is very similar to the **Play Particle Effect** notify in that they both allow you to play a Particle Effect at any given point in an animation, the difference however is that with a Timed Particle Effect, you are given the ability to define the length of the effect inside the notify itself. 
 
 ![](TimedParticleEffect.png)
 
-In the image above, a Timed Particle Effect is assigned to an animation. By **Left-Clicking** on the handle bar at the end (or at the start), you can drag to increase or decrease the length of the effect.
+In the image above, both a Play Particle Effect and Timed Particle Effect are assigned to an animation. Inside the yellow box is the Timed Particle Effect and you will notice that this notify has two handle bars; one at the start and one at the end of the event. By **Left-Clicking** on the handle bar at the end (denoted by the green arrow), you can drag to increase or decrease the length of the effect.
 
 ### Animation Trails
 Animation Trails (or **AnimTrail** for short) are emitters that allow you to leave trails behind moving objects in a scene. 
 They do this by sampling the positions of two sockets on a Skeletal Mesh each tick and stitching triangles between the sockets, down the length of the trail. 
   
-To use an AnimTrail emitter you must:
-1. Create a [Particle System](Engine/Rendering/ParticleSystems) that contains an emitter with AnimTrail type data. 
-2. Create and add a Notify State of the **Trail** type to an animation.
+To use an AnimTrail emitter you must do the following; 
+1.       Create a Particle System that contains an emitter with AnimTrail type data. 
+2.       In the **Animation Editor**, add an AnimNotifyState of type **Trail**. 
 
-<!-- To be included as part of a HowTo.
 Setting Up the Particle System
 
 Create a Particle System as usual and add a new emitter with type data AnimTrail to it. 
@@ -208,25 +220,24 @@ However, the more tessellated a trail is, the more expensive it is.
 | **Distance Tessellation Step Size** | The (approx.) distance in world units that the trail must be to warrant another tessellation point.   |
 | **Tangent Tessellation Step Size** | The (approx.) angle difference in degrees that the tangent must change from one particle to another to warrant another tessellation point. This allows you to ensure that any sharp turns in the trail are nice and smooth.  |
 | **Width Tessellation Step Size** | The (approx.) world unit change in the width of the trail that warrants an additional tessellation point.|
-
+ 
 Setting Up the AnimNotify
 
 Now that we have a particle system, we must use it with an AnimNotify. 
-Open the **Animation Editor**, add or use an existing **Timeline** under the **Notifies** section and **Right-click** on it to bring up the menu. The Trails Notify can be found under the **Add Notify State** section. 
+Open the **Animation Editor**, add or use an existing **Timeline** under the **Notifies** section and **Right-click** on it to bring up the menu. The Trails notify can be found under the **Add Notify State** section. 
 
 ![](SetUpTrail.png)
--->
 
-Once you have added an AnimTrail, you can adjust its properties from the **Details** panel: 
+Below is the **Details** panel and **Properties** for the Trail notify. 
 
 ![](TrailDetails.png)
 
 | **Property** | **Description** |
 | --- | --- |
-| **PSTemplate** | This is the particle system containing the anim trail that was described above. |
 | **First Socket Name** | This is the name of the first Bone/Socket defining the trail. |
 | **Second Socket Name** | This is the name of the second Bone/Socket defining the trail. |
-
+| **PSTemplate** | This is the particle system containing the anim trail that was described above. |
+  
 The above properties are required but the below are optional. They allow you to alter the width of the trail during the animation. You could also attach the trail to sockets that move as you need but this may not be ideal. 
 This method is based on providing a width scale value. A width scale of 1.0 will keep the trail the same width, 0.0 will cause a trail of zero width. 
 
@@ -248,7 +259,6 @@ There are also some debug rendering options available:
 | **Render Tangents** | This toggles the debug rendering of the curve tangents. |
 | **Render Tessellation** | This toggles debug rendering of all vertices on the curve. This shows how the curve is tessellated between the spawn points. |
 
-<!-- To be included with AnimTrail HowTo.
 ### Controlling Animation Trails in Blueprint
 While the above method is the standard approach to controlling AnimTrails, you can also control them from inside a Blueprint without the need of an Anim Notify. 
 The only restriction with this is that the Particle System Component must be a child of the Skeletal Mesh Component from which you are sampling socket locations. 
@@ -261,34 +271,4 @@ The only restriction with this is that the Particle System Component must be a c
  ![](TrailBlueprint3.png)
 1. You can also alter the trail's properties while it is running (the example below alters width each tick using the game time).  
  ![](TrailBlueprint4.png)
-  -->
-
-## Additional Notify Properties
-In addition to the properties that can be modified in the **Anim Notify** section of the **Details** panel for each type of Notify (when applicable), the properties below can also be used to define how your Notify is used.
-
-![](AdditionalProperties.png)
-
-|---|---|
-|[REGION:tablesection]Trigger Settings[/REGION]||
-|**Notify Trigger Chance**| Defines the chance of triggering the Notify (0=No Chance, 1=Always Triggers).|
-|**Trigger on Dedicated Server**| If disabled, this Notify will be skipped on Dedicated Servers.|
-|**Notify Filter Type**| Defines a method for filtering notifies (stopping them from triggering). Can be set to filter by the meshes current LOD.|
-|[REGION:tablesection]Event[/REGION]||
-|**Trigger Weight Threshold**| Determines how much much weight the animation sequence needs to have to trigger the Notify. |
-|**Montage Tick Type**| The method to determine how the Notify is handled, Queued or Branching Point (See [Montage Tick Type](#montageticktype) below). |
-|[REGION:tablesection]Anim Link[/REGION]||
-|**Slot**| The slot index currently being used within the LinkedMontage. |
-|**Link Method**| The method used to calculate times, Absolute (stays at a specific time without moving), Relative (moves with its segment, but not when the segment changes size) or Proportional (moves with its segment and will stay at a certain proportion through the segment). |
-
-
-### Montage Tick Type
-
-The key difference between Queued and Branching Point settings for the Montage Tick Type is that Queued Notifies is asynchronous, while Branching Points are synchronous. What this means to the end user is that Branching Points come with a much higher degree of precision for where they will take place along the animation timeline.
-
-High precision is important when you need to jump to a specific animation at a very precise moment in time. While you could use a Queued Notify to do the same job, the asynchronous nature of Queued Notifies means that the event could be fired at the incorrect animation frame, which can lead to hitches and jumps in your motion.
-
-[REGION:warning]
-Due to their synchronous nature and the precision resulting from it, Branching Points are more performance expensive than Queued. You should only use them when an event must be fired at a precise moment along the animation timeline, such as jumping directly to another animation that matches up frame-to-frame. If being off by a frame (or some percentage of one) is not important, you should use Queued instead.
-[/REGION]
-
 
diff --git a/Engine/Documentation/Source/Engine/Paper2D/Flipbooks/Paper2DFlipbooks.INT.udn b/Engine/Documentation/Source/Engine/Paper2D/Flipbooks/Paper2DFlipbooks.INT.udn
--- a/Engine/Documentation/Source/Engine/Paper2D/Flipbooks/Paper2DFlipbooks.INT.udn
+++ b/Engine/Documentation/Source/Engine/Paper2D/Flipbooks/Paper2DFlipbooks.INT.udn
@@ -7,8 +7,6 @@ Related: Engine\Paper2D\Flipbooks\Components
 Related: Engine/Paper2D/Sprites
 Related: Engine/Rendering/Materials
 Related: Engine/Content/Types/Textures
-Version: 4.9
-SkillLevel: Beginner
 
 [VAR:Topic]
 [OBJECT:Topic]
@@ -56,6 +54,7 @@ SkillLevel: Beginner
 
 [TOC (start:2 end:3)]
 
+## Overview
 
 [EXCERPT:intro]
 The best way to think of **Paper 2D Flipbooks** (or **Flipbooks** for short) is in the form of hand-drawn animation, where a series of images that are slightly different are "flipped" through to produce what appears to be motion. In Unreal Engine 4, Flipbooks consist of a series of key frames, each of which contains a [Sprite](Engine\Paper2D\Sprites) to be displayed and a duration (in frames) to display it. 
diff --git a/Engine/Documentation/Source/Engine/QuickStart/1/LDQS_1.INT.udn b/Engine/Documentation/Source/Engine/QuickStart/1/LDQS_1.INT.udn
--- a/Engine/Documentation/Source/Engine/QuickStart/1/LDQS_1.INT.udn
+++ b/Engine/Documentation/Source/Engine/QuickStart/1/LDQS_1.INT.udn
@@ -26,7 +26,7 @@ Related:GettingStarted/Installation
 %Steps%
 
 Before we get started, the first thing that we will need to do is create a [Project](Engine/Basics/Projects) to work from for the duration of this guide. In the steps below we will use Unreal Engine 4's [Project Browser](Engine/Basics/Projects/Browser) to create our project. 
-After completing this step, you will have a Template that you can build onto for your future projects.  
+After completing this step, our project will be ready and you will gone through the process of creating a Project and using a Template which you can then use for your very own future projects. 
 
 ## Steps
 
diff --git a/Engine/Documentation/Source/Engine/Rendering/LightingAndShadows/LightFunctions/LightFunctions.INT.udn b/Engine/Documentation/Source/Engine/Rendering/LightingAndShadows/LightFunctions/LightFunctions.INT.udn
--- a/Engine/Documentation/Source/Engine/Rendering/LightingAndShadows/LightFunctions/LightFunctions.INT.udn
+++ b/Engine/Documentation/Source/Engine/Rendering/LightingAndShadows/LightFunctions/LightFunctions.INT.udn
@@ -52,10 +52,6 @@ Version: 4.9
 [EXCERPT:LightFunctions]
 **Light Functions** are - in essence - a Material that can be applied to filter a light's intensity.  
 
-[Region:note]
-It should be noted, however, that you cannot change the color of lights using a Light Function, only the Light Color setting can do that. Also, Static Lights cannot not use Light Functions, as this feature is only available for non-lightmapped lights, like Stationary and Movable lights.
-[/Region]
-
 It should be noted however that you cannot change the color of lights using a Light Function, only the Light Color setting can do that. 
 [/EXCERPT:LightFunctions]
 
diff --git a/Engine/Documentation/Source/Engine/Rendering/LightingAndShadows/LightMobility/StationaryLights/StationaryLights.INT.udn b/Engine/Documentation/Source/Engine/Rendering/LightingAndShadows/LightMobility/StationaryLights/StationaryLights.INT.udn
--- a/Engine/Documentation/Source/Engine/Rendering/LightingAndShadows/LightMobility/StationaryLights/StationaryLights.INT.udn
+++ b/Engine/Documentation/Source/Engine/Rendering/LightingAndShadows/LightMobility/StationaryLights/StationaryLights.INT.udn
@@ -124,7 +124,6 @@ Translucency also receives shadowing very cheaply with Stationary lights - Light
 	[/PARAM]
 [/OBJECT]
 
-The resolution of the static shadow depth map is controlled by StaticShadowDepthMapTransitionSampleDistanceX and StaticShadowDepthMapTransitionSampleDistanceY in BaseLightmass.ini, with a default setting of 100 meaning one texel every meter.
 
 ### Dynamic Shadowing
 
@@ -186,10 +185,10 @@ In the following image you see the difference between the two shadowing methods.
 
 [OBJECT:ComparisonSlider]
 	[PARAM:before]
-	![Area Shadows Enabled](AreaShadows.png)(w:900)
+	![Area Shadows Enabled](Area_Shadows_On.png)(w:900)
 	[/PARAM]
 	[PARAM:after]
-	![Area Shadows Disabled](UniformPenumbraSharp.png)(w:900)
+	![Area Shadows Disabled](Area_Shadows_Off.png)(w:900)
 	[/PARAM]
 [/OBJECT]
 
diff --git a/Engine/Documentation/Source/Engine/Rendering/LightingAndShadows/LitTranslucency/LitTranslucency.INT.udn b/Engine/Documentation/Source/Engine/Rendering/LightingAndShadows/LitTranslucency/LitTranslucency.INT.udn
--- a/Engine/Documentation/Source/Engine/Rendering/LightingAndShadows/LitTranslucency/LitTranslucency.INT.udn
+++ b/Engine/Documentation/Source/Engine/Rendering/LightingAndShadows/LitTranslucency/LitTranslucency.INT.udn
@@ -32,11 +32,6 @@ Translucent effects generally fall into a few categories: volumetric, volumetric
 
 Lit translucency gets most of its lighting through a series of cascaded volume textures oriented around the view frustum.  This allows lighting to be known in a single forward pass for any point inside the volumes, but has the downside that the volume texture is fairly low resolution, and can only cover a limited depth range from the viewer.  
 
-The volume is configured through cvars that can be set differently based on the scalability level:
-* r.TranslucencyLightingVolumeDim, which defaults to 64.  Raising this by a factor of 2 increases the cost to light volume by a factor of 8.
-* r.TranslucencyLightingVolumeInnerDistance, which defaults to 1500.  Raising this increases lighting volume coverage but reduces effective resolution.
-* r.TranslucencyLightingVolumeOuterDistance, which defaults to 5000.  Raising this increases lighting volume coverage but reduces effective resolution.
-
 Shadowed direct lighting from all movable light types is injected into the translucency lighting volume.  Light functions are also taken into account.
 
 [REGION:fullwidth]
@@ -63,10 +58,6 @@ Translucent self-shadowing uses per-object shadows, which means that it needs us
 | | |
 [/REGION]
 
-### Static shadowing
-
-Translucency can receive static shadowing from [stationary lights](Engine/Rendering/LightingAndShadows/LightMobility/StationaryLights) through a special static shadow depth map generated by Lightmass at lighting build time.
-
 ## Translucent Surfaces
 
 ### Reflection Captures
@@ -79,5 +70,6 @@ _The left sphere is translucent, the right is opaque, and they are both setup as
 
 ## Limitations
 
+* Static shadowing from static lights is currently not handled for lit translucency. However, dynamic shadows from stationary lights is supported.
 * Lit translucent surfaces are missing direct specular.
 * Lit translucent surfaces get all their direct lighting through the translucency volume lighting texture, which causes them to be lower resolution than needed for most surface materials (glass, water).
diff --git a/Engine/Documentation/Source/Engine/Rendering/PostProcessEffects/PostProcessMaterials/PostProcessMaterials.INT.udn b/Engine/Documentation/Source/Engine/Rendering/PostProcessEffects/PostProcessMaterials/PostProcessMaterials.INT.udn
--- a/Engine/Documentation/Source/Engine/Rendering/PostProcessEffects/PostProcessMaterials/PostProcessMaterials.INT.udn
+++ b/Engine/Documentation/Source/Engine/Rendering/PostProcessEffects/PostProcessMaterials/PostProcessMaterials.INT.udn
@@ -14,7 +14,7 @@ Version: 4.9
 
 [EXCERPT:PostProcessMaterials_Desc]
 The engine already features complex post processing based on a graph of post processing nodes. The **Post Processing Materials**
-can be additionally inserted in some specific position. See [FAQ](#FAQ) **r.CompositionGraphDebug** to get a dump of the full graph.
+can be additionally inserted in some specific position. See [**FAQ**](#FAQ) **r.CompositionGraphDebug** to get a dump of the full graph.
 The graph is actually not only doing post processing but also some parts of the lighting. Over time we plan to make more
 parts be customizable through the material editor.
 [/EXCERPT:PostProcessMaterials_Desc]
diff --git a/Engine/Documentation/Source/GettingStarted/DownloadingUnrealEngine/Perforce/Setup/Setup.INT.udn b/Engine/Documentation/Source/GettingStarted/DownloadingUnrealEngine/Perforce/Setup/Setup.INT.udn
--- a/Engine/Documentation/Source/GettingStarted/DownloadingUnrealEngine/Perforce/Setup/Setup.INT.udn
+++ b/Engine/Documentation/Source/GettingStarted/DownloadingUnrealEngine/Perforce/Setup/Setup.INT.udn
@@ -17,8 +17,8 @@ To learn about setting up and establishing a VPN connection, see the [Setting Up
 
 [REGION:red]
 Epic Games has executed a custom licensing agreement with Perforce. To be in compliance with this license, only 
-one machine per license should sync full builds, but multiple users may connect with the same login name to 
-view source metadata and past revisions. Individual UDN login information will not work.
+one machine per license should sync full builds, but multiple users **may** connect with the same login name to 
+view source metadata and past revisions. Individual UDN login information will **not** work.
 [/REGION]
 
 In some cases, larger organizations may have an account for each team; even still, only one machine per team 
@@ -177,7 +177,7 @@ this on our proxy.
 
 
 It is a VBScript that removes files that have not been accessed within a certain time period, keeping things nice and 
-tidy. It is best to set up as a Windows "at job" or a scheduled task. Be sure to run it as a user with the appropriate 
+tidy. It is best to set up as a Windows _at job_ or a scheduled task. Be sure to run it as a user with the appropriate 
 permissions to the location on the file system.
 
 The script can be called as:
diff --git a/Engine/Documentation/Source/GettingStarted/HowTo/PlacingActors/PlacingActors.INT.udn b/Engine/Documentation/Source/GettingStarted/HowTo/PlacingActors/PlacingActors.INT.udn
--- a/Engine/Documentation/Source/GettingStarted/HowTo/PlacingActors/PlacingActors.INT.udn
+++ b/Engine/Documentation/Source/GettingStarted/HowTo/PlacingActors/PlacingActors.INT.udn
@@ -6,7 +6,6 @@ Related: Engine/Actors/Placement
 Related: Engine/UI/LevelEditor/Modes/PlaceMode
 Related: Engine/UI/ClassViewer
 Related: GettingStarted\HowTo\ManipulatingActors
-version: 4.9
 
 At the most fundamental level, an **Actor** is any Object that you can place in a level and this page will show you the various methods in which you can place those Actors inside your levels. 
 
@@ -73,7 +72,7 @@ This method of placing an Actor involves using the **Content Browser** and dragg
 
 A more advanced method of placing Actors is performed by using the [](Engine/UI/ClassViewer), which is a hierarchical list of classes used by the Editor.
 
-1.  To open the **Class Viewer**, click **Window** then **Developer Tools** and under **Miscellaneous**, select **Class Viewer**. 
+1.  To open the **Class Viewer**, click **Window** then **Developer Tools** and under _Miscellaneous_, select **Class Viewer**. 
 
 	![](Place6.png)
 
diff --git a/Engine/Documentation/Source/GettingStarted/Installation/InstallingUnrealEngine.INT.udn b/Engine/Documentation/Source/GettingStarted/Installation/InstallingUnrealEngine.INT.udn
--- a/Engine/Documentation/Source/GettingStarted/Installation/InstallingUnrealEngine.INT.udn
+++ b/Engine/Documentation/Source/GettingStarted/Installation/InstallingUnrealEngine.INT.udn
@@ -8,7 +8,7 @@ Related: Engine/Content/QuickStart
 Related: Programming/QuickStart
 Related: Engine/QuickStart
 Related: Support/Builds
-Version: 4.9
+Version: 4.8
 SkillLevel: Beginner
 
 By the end of this tutorial, you'll be able to:
@@ -21,7 +21,7 @@ By the end of this tutorial, you'll be able to:
 
 ![](InstallationGoal.png)
 
-The Epic Games Launcher enables our subscribers to conveniently download and install updates for UE4.
+_The Epic Games Launcher enables our subscribers to conveniently download and install updates for UE4._
 
 [REGION:warning]
 If you already have an Epic Games account, skip ahead to learn about [Downloading and Running the Installer](#Bookmark2).
@@ -126,15 +126,15 @@ Space Requirements
 Unreal Engine uses approximately eight gigabytes of disk space for every version of the engine that is installed on a machine. Make sure you have plenty of disk space prior to installing UE4.
 [/REGION]
 
-1.	Click the **Unreal Engine** tab inside the Epic Games Launcher Loading Screen.
+1.	**Click** the **Unreal Engine** tab inside the Epic Games Launcher Loading Screen.
 
 	![](LauncherLoadScreen.png)
 
-1.	Click the **Install Engine** button to download and install the latest version of Unreal Engine.
+1.	**Click** the **Install Engine** button to download and install the latest version of Unreal Engine.
 
 	![](InstallEngine.png)
 
-1.	Read our [Unreal Engine End User License Agreement](https://www.unrealengine.com/eula) (EULA) prior to clicking the **Accept** button.
+1.	**Read** our [Unreal Engine End User License Agreement](https://www.unrealengine.com/eula) (EULA) prior to **clicking** the **Accept** button.
 
 	![](AcceptEULALauncher.png)
 
@@ -148,7 +148,7 @@ Unreal Engine uses approximately eight gigabytes of disk space for every version
 
 (#Bookmark5)
 
-After the Epic Games Launcher successfully downloads and installs Unreal Engine, you're ready to click the **Launch** button.
+After the Epic Games Launcher successfully downloads and installs Unreal Engine, you're ready to **click** the **Launch** button.
 
 ![](LaunchUE4.png)
 
@@ -162,7 +162,6 @@ Congratulations! You've learned how to:
 
 Now, you're ready to start using Unreal Engine.
 
-<!--Deprecated as it is redundant with the Related Topics section directly below it.
 ## Getting Started with Unreal Engine
 
 Learn how to use Unreal Engine by referring to the [](GettingStarted) documentation!
@@ -173,9 +172,9 @@ If you're looking to quickly get started with UE4, refer to the following tutori
 *	[](Engine/QuickStart)
 
 
-[REGION:note]
+[REGION:tip]
 UE4's in-editor help features are a great way to get your questions answered.
-[/REGION]-->
+[/REGION]
 <!--Deprecated on 21 July 2015 due to outdated screenshots and outdated procedures
 %Globals:OSSelect%
 
diff --git a/Engine/Documentation/Source/GettingStarted/QuickStartDirectory/QuickStartDirectory.INT.udn b/Engine/Documentation/Source/GettingStarted/QuickStartDirectory/QuickStartDirectory.INT.udn
--- a/Engine/Documentation/Source/GettingStarted/QuickStartDirectory/QuickStartDirectory.INT.udn
+++ b/Engine/Documentation/Source/GettingStarted/QuickStartDirectory/QuickStartDirectory.INT.udn
@@ -2,7 +2,6 @@ Availability:Public
 Title:Unreal Engine 4 Quick Start Directory
 Crumbs: %ROOT%, GettingStarted
 Description:All of the quick start guides available to help you get started with Unreal Engine, at a glance.
-version: 4.9
 
 [VAR:Topic]
 [OBJECT:Topic]
diff --git a/Engine/Documentation/Source/GettingStarted/RecommendedSpecifications/Specifications.INT.udn b/Engine/Documentation/Source/GettingStarted/RecommendedSpecifications/Specifications.INT.udn
--- a/Engine/Documentation/Source/GettingStarted/RecommendedSpecifications/Specifications.INT.udn
+++ b/Engine/Documentation/Source/GettingStarted/RecommendedSpecifications/Specifications.INT.udn
@@ -3,7 +3,6 @@ Title:Hardware & Software Specifications
 SEO-Title:Hardware & Software Specifications
 Crumbs: %ROOT%, GettingStarted
 Description:Minimum and recommended hardware specifications and necessary software for developing with Unreal Engine.
-version: 4.9
 
 [PUBLISH:Rocket]
 %Globals:OSSelect_Linux%
diff --git a/Engine/Documentation/Source/GettingStarted/RunningUnrealEngine/RunningUnrealEngine.INT.udn b/Engine/Documentation/Source/GettingStarted/RunningUnrealEngine/RunningUnrealEngine.INT.udn
--- a/Engine/Documentation/Source/GettingStarted/RunningUnrealEngine/RunningUnrealEngine.INT.udn
+++ b/Engine/Documentation/Source/GettingStarted/RunningUnrealEngine/RunningUnrealEngine.INT.udn
@@ -2,7 +2,6 @@ Availability:Public
 Title:Running Unreal Engine
 Crumbs: %ROOT%, GettingStarted
 Description:Running the engine as a game or editor, including how to launch a specific map and useful in-game commands.
-version: 4.9
 
 %Globals:OSSelect%
 
@@ -26,11 +25,11 @@ This process will launch a copy of Unreal Editor.
 
 The process of running the editor simply requires passing the name of the project to run as an argument to the executable.
 
-[REGION:note]
+[REGION:tip]
 You can add the `-debug` switch to force the executable to load the debug version of the modules for your project, which 
 contain all of the debugging symbols. This is necessary even when debugging through Visual Studio with the configuration 
 set to **Debug** because the main executable is always compiled using the **Development** configuration. Of course, you 
-must first compile your modules using the Debug configuration so they exist for the executable to load.
+must first compile your modules using the **Debug** configuration so they exist for the executable to load.
 [/REGION]
 
 [/PUBLISH:Rocket]
@@ -172,15 +171,14 @@ or open [content examples and sample games](Resources).
 ## Running an Uncooked Game
 
 Once a project is loaded in Unreal Editor, you can [test your gameplay](Engine/UI/LevelEditor/InEditorTesting#modes) in Uncooked Game mode from the **Play In** menu. 
-To play the uncooked game in its own window, select the [New Window At > Default Player Start mode](Engine/UI/LevelEditor/InEditorTesting#modes) using 
-the [Play In dropdown menu](Engine/UI/LevelEditor/InEditorTesting#toolbar) in the Level Editor toolbar.
-
+To play the uncooked game in its own window, select the [**New Window At** > **Default Player Start** mode](Engine/UI/LevelEditor/InEditorTesting#modes) using 
+the [**Play In** dropdown menu](Engine/UI/LevelEditor/InEditorTesting#toolbar) in the Level Editor toolbar.
 
+**Click to see full-size image:**
 
 [REGION:lightbox]
 [![](Engine/UI/LevelEditor/InEditorTesting/PIE_newwindow.png)(w:300)](Engine/UI/LevelEditor/InEditorTesting/PIE_newwindow.png)
 [/REGION:lightbox]
-_Click to see full-size image:_
 
 Running the uncooked game version of the engine executable using either of the below methods will result in the same behavior.
 
diff --git a/Engine/Documentation/Source/GettingStarted/SubEditors/SubEditors.INT.udn b/Engine/Documentation/Source/GettingStarted/SubEditors/SubEditors.INT.udn
--- a/Engine/Documentation/Source/GettingStarted/SubEditors/SubEditors.INT.udn
+++ b/Engine/Documentation/Source/GettingStarted/SubEditors/SubEditors.INT.udn
@@ -2,7 +2,6 @@ Availability: Public
 Title:Tools and Editors
 Crumbs: %ROOT%, Gameplay, GettingStarted
 Description:There are several different types of Editors contained within Unreal Engine 4, this page highlights each of them. 
-version: 4.9
 
 [VAR:Topic]
 [OBJECT:Topic]
diff --git a/Engine/Documentation/Source/GettingStarted/Terminology/Terminology.INT.udn b/Engine/Documentation/Source/GettingStarted/Terminology/Terminology.INT.udn
--- a/Engine/Documentation/Source/GettingStarted/Terminology/Terminology.INT.udn
+++ b/Engine/Documentation/Source/GettingStarted/Terminology/Terminology.INT.udn
@@ -2,7 +2,6 @@ Availability: Public
 Title:Unreal Engine 4 Terminology
 Crumbs: %ROOT%, Gameplay, GettingStarted
 Description: Covers the most commonly used terms when working with Unreal Engine 4. 
-version: 4.9
 
 [VAR:Topic]
 [OBJECT:Topic]
diff --git a/Engine/Documentation/Source/Platforms/Android/Android.INT.udn b/Engine/Documentation/Source/Platforms/Android/Android.INT.udn
--- a/Engine/Documentation/Source/Platforms/Android/Android.INT.udn
+++ b/Engine/Documentation/Source/Platforms/Android/Android.INT.udn
@@ -4,7 +4,6 @@ Crumbs:%ROOT%, Platforms
 Description:Developing games for the Android platform.
 Platform:Android
 Navigation:platform
-version: 4.9
 
 
 [OBJECT:TopicButtonList]
diff --git a/Engine/Documentation/Source/Platforms/Mobile/Lighting/MobileLighting.INT.udn b/Engine/Documentation/Source/Platforms/Mobile/Lighting/MobileLighting.INT.udn
--- a/Engine/Documentation/Source/Platforms/Mobile/Lighting/MobileLighting.INT.udn
+++ b/Engine/Documentation/Source/Platforms/Mobile/Lighting/MobileLighting.INT.udn
@@ -3,7 +3,6 @@ Title:Lighting for Mobile Platforms
 Crumbs:%ROOT%, Platforms, Platforms/Mobile
 Description:Setting up lighting for mobile games.
 Platform:Mobile
-version: 4.9
 
 ## Supported Features
 
diff --git a/Engine/Documentation/Source/Platforms/Mobile/Materials/MobileMaterials.INT.udn b/Engine/Documentation/Source/Platforms/Mobile/Materials/MobileMaterials.INT.udn
--- a/Engine/Documentation/Source/Platforms/Mobile/Materials/MobileMaterials.INT.udn
+++ b/Engine/Documentation/Source/Platforms/Mobile/Materials/MobileMaterials.INT.udn
@@ -9,10 +9,11 @@ Related: Engine/Rendering/Materials/CustomizedUVs
 Related: Engine/Rendering/Materials/MaterialProperties/BlendModes
 Related: Engine/Rendering/Materials/ExpressionReference/Utility
 Related: Platforms/Mobile/Lighting
-version: 4.9
 
 [TOC(start:2)]
 
+## Overview
+
 Materials on mobile platforms (Feature Level ES 2) use the same node-based creation process used by other platforms. Almost all of the Material inputs are available to mobile Materials, 
 and Materials will generally look and behave the same on mobile platforms as they do on other platforms. For instance, the following Material inputs all work exactly the same on mobile platforms 
 as they do on PC.
@@ -30,7 +31,7 @@ There are some considerations you must take into account when creating Materials
 * Only Default and Unlit shading models are available.
 * [](Engine/Rendering/Materials/CustomizedUVs) should be used to avoid dependent texture fetches (no math on texture UVs).
 * Translucent and Masked Materials are extremely expensive. It is recommended that you use [Opaque Materials](Engine/Rendering/Materials/MaterialProperties/BlendModes#opaque) wherever possible.
-* Depth fade can be used in translucent Materials for **iOS platforms**. It is not supported on platforms where the hardware does not support fetching from the depth buffer and 
+* Depth fade can be used in translucent Materials for **iOS platforms**. It is _**NOT**_ supported on platforms where the hardware does not support fetching from the depth buffer and 
 	emulating this would incur an unacceptable performance cost.
 
 
@@ -59,5 +60,5 @@ To maintain compatibility with older devices, use the [Feature Level](Engine/Ren
 
 | Issue | Solution |
 | ----- | -------- |
-| Default Material is displayed in the Mobile Previewer. | Open the [log file](Engine/Basics/DirectoryStructure) in **Engine\Saved** and search for error messages to know how to fix the Material. |
+| Default Material is displayed in the Mobile Previewer. | Open the [log file](Engine/Basics/DirectoryStructure) in **Engine\Saved** and search for **error messages** to know how to fix the Material. |
 | Material Animation with Time Expressions looks choppy but framerate is good elsewhere. | Make sure the Period property is enabled in the **Time** expression along with a Period. This will move the computation to the CPU and run at full-precision instead of half. |
diff --git a/Engine/Documentation/Source/Platforms/Mobile/Mobile.INT.udn b/Engine/Documentation/Source/Platforms/Mobile/Mobile.INT.udn
--- a/Engine/Documentation/Source/Platforms/Mobile/Mobile.INT.udn
+++ b/Engine/Documentation/Source/Platforms/Mobile/Mobile.INT.udn
@@ -4,7 +4,6 @@ Crumbs:%ROOT%, Platforms
 Description:Information for developers creating mobile games.
 Platform:Mobile
 Navigation:platform
-version: 4.9
 
 [REGION:banner]
 ![](mobile_banner.png)
diff --git a/Engine/Documentation/Source/Platforms/Mobile/Previewer/MobilePreviewer.INT.udn b/Engine/Documentation/Source/Platforms/Mobile/Previewer/MobilePreviewer.INT.udn
--- a/Engine/Documentation/Source/Platforms/Mobile/Previewer/MobilePreviewer.INT.udn
+++ b/Engine/Documentation/Source/Platforms/Mobile/Previewer/MobilePreviewer.INT.udn
@@ -3,7 +3,6 @@ Title:Mobile Previewer
 Crumbs:%ROOT%, Platforms, Platforms/Mobile
 Description: Previewing games on mobile platforms will look in the UE4 editor.
 Platform:Mobile
-version: 4.9
 
 ![](T_Rendering_Level_Examples.png)
 
@@ -18,7 +17,7 @@ When enabling different preview levels the Materials in your scene will be re-co
 You can also switch back and forth between rendering levels seamlessly inside of UE4 without having to re-start the editor.
 To change the which rendering level is used in the current UE4 Editor session, you will need to do the following. 
 
-[region:note]
+[region:tip]
 While you do not need to re-start the Editor for the changes to take effect the first time you change the Preview Rendering Level the editor will need time to re-compile shaders.
 [/region]
 
@@ -32,7 +31,7 @@ While you do not need to re-start the Editor for the changes to take effect the
 ![](T_PRL.png)
 
 [region:note]
-The Mobile previewer is intended to match mobile devices as closely as possible but is not indicative of what your project is going to look like on the target device.
+Please note that the Mobile previewer is intended to match mobile devices as closely as possible but is not indicative of what your project is going to look like on the target device.
 You should always make sure to fully test your project on your target device and only use the mobile previewer to see if your work is headed in the right direction.
 [/region]
 
diff --git a/Engine/Documentation/Source/Platforms/Platforms.INT.udn b/Engine/Documentation/Source/Platforms/Platforms.INT.udn
--- a/Engine/Documentation/Source/Platforms/Platforms.INT.udn
+++ b/Engine/Documentation/Source/Platforms/Platforms.INT.udn
@@ -3,7 +3,6 @@ Title:Platform Development
 Crumbs:%ROOT%
 Description:Information over developing for platforms other than PC.
 Navigation:platform
-version: 4.9
 
 [VAR:Params]
 [PARAM:image]
diff --git a/Engine/Documentation/Source/Programming/Animation/AnimNodes/CreatingAnimNodes.INT.udn b/Engine/Documentation/Source/Programming/Animation/AnimNodes/CreatingAnimNodes.INT.udn
--- a/Engine/Documentation/Source/Programming/Animation/AnimNodes/CreatingAnimNodes.INT.udn
+++ b/Engine/Documentation/Source/Programming/Animation/AnimNodes/CreatingAnimNodes.INT.udn
@@ -2,7 +2,7 @@ Availability:Public
 Title: Animation Node Technical Guide
 Crumbs: %ROOT%, Engine, Engine/Animation
 Description: Guide to creating new nodes for use within graphs in Anim Blueprints.
-Version: 4.9
+Version: 4.5
 
 [REGION:banner]
 ![Anim Nodes](animnodes_banner.png)(convert:false)
@@ -10,7 +10,7 @@ Version: 4.9
 
 [TOC(start:2)]
 
-Animation nodes are used inside of [Anim Blueprints](Engine/Animation/AnimBlueprints) 
+Animation nodes are used inside of [_Anim Blueprints_](Engine/Animation/AnimBlueprints) 
 to perform operations, such as blending of poses or direct manipulation of bones. Several animation 
 nodes are provided, but custom nodes can be created to suit the needs of any game.
 
@@ -100,13 +100,16 @@ The following metadata keys are available:
 | `AlwaysAsPin` | The property is always exposed as a data pin. |
 
 (#OptionalPins)
-**Optional Pins**
-[REGION:none]
-	For properties with `PinHiddenByDefault` or `PinShownByDefault`, a check box will show up next 
-	to the property in the **Details** panel, allowing you to show or hide it.
-[/REGION]	
-[REGION:none]
-	![Property Details](node_property_details.png)
+[REGION:note]
+	**Optional Pins**
+	[REGION:none]
+		For properties with `PinHiddenByDefault` or `PinShownByDefault`, a check box will show up next 
+		to the property in the **Details** panel, allowing you to show or hide it.
+	[/REGION]	
+	[REGION:none]
+		![Property Details](node_property_details.png)
+	[/REGION]
+
 [/REGION]
 
 ## Editor Node
@@ -123,7 +126,7 @@ The editor-time class should contain an instance of your runtime node exposed as
 
 ![Node Title](node_title.png)
 
-The background color and text of the title of the animation node displayed in the graph of the Animation Blueprint in 
+The background color and text of the title of the animation node displayed in the graph of the _Animation Blueprint_ in 
 Persona are defined by overriding the `GetNodeTitle()` and `GetNodeTitleColor()` functions.
 
 For example, the `UAnimGraphNode_ApplyAdditive` node uses a gray background and displays "Apply Additive":
@@ -151,7 +154,7 @@ The tooltip displayed when hovering over the node in Persona is defined by overr
 
 ### Context Menu
 
-Each animation node can add node-specific options to the context menu displayed when **Right-clicking** on the 
+Each animation node can add node-specific options to the context menu displayed when **right-clicking** on the 
 node in a graph in Persona. Options are added using the `GetContextMenuActions()` function which is a 
 member of all Blueprint nodes.
 
diff --git a/Engine/Documentation/Source/Programming/Assets/AsyncLoading/AsyncLoading.INT.udn b/Engine/Documentation/Source/Programming/Assets/AsyncLoading/AsyncLoading.INT.udn
--- a/Engine/Documentation/Source/Programming/Assets/AsyncLoading/AsyncLoading.INT.udn
+++ b/Engine/Documentation/Source/Programming/Assets/AsyncLoading/AsyncLoading.INT.udn
@@ -2,7 +2,7 @@ Availability:Public
 Title:Asynchronous Asset Loading
 Crumbs:%ROOT%, Programming
 Description:Methods for loading and unloading assets during runtime.
-Version: 4.9
+Version: 4.5
 
 [TOC(start:2)]
 
diff --git a/Engine/Documentation/Source/Programming/Assets/ReferencingAssets/ReferencingAssets.INT.udn b/Engine/Documentation/Source/Programming/Assets/ReferencingAssets/ReferencingAssets.INT.udn
--- a/Engine/Documentation/Source/Programming/Assets/ReferencingAssets/ReferencingAssets.INT.udn
+++ b/Engine/Documentation/Source/Programming/Assets/ReferencingAssets/ReferencingAssets.INT.udn
@@ -2,7 +2,7 @@ Title: Referencing Assets
 Crumbs:
 Description:
 Availability: Public
-Version: 4.9
+Version: 4.7
 
 Unreal Engine 4 provides a number of mechanisms to control how an asset is referenced and by extension loaded into memory. You can think of references in two ways: a hard reference where object A refers to object B and causes object B to be loaded when object A is loaded; and a soft reference where object A refers to object B via an indirect mechanism such as the string form of the path to the object. The first two sections following cover hard references with the remaining sections exploring soft references. 
 
diff --git a/Engine/Documentation/Source/Programming/Assets/Registry/AssetRegistry.INT.udn b/Engine/Documentation/Source/Programming/Assets/Registry/AssetRegistry.INT.udn
--- a/Engine/Documentation/Source/Programming/Assets/Registry/AssetRegistry.INT.udn
+++ b/Engine/Documentation/Source/Programming/Assets/Registry/AssetRegistry.INT.udn
@@ -2,7 +2,7 @@ Availability:Public
 Title:Asset Registry
 Crumbs: %ROOT%, Engine, Engine/UI, Programming
 Description:How assets are discovered by the editor and how to make it know more about asset types before they are loaded.
-Version: 4.9
+Version: 4.5
 
 The **Asset Registry** is an editor subsystem which gathers information about unloaded assets asynchronously as the editor loads.
 This information is stored in memory so the editor can create lists of assets without loading them.
@@ -102,10 +102,10 @@ For example (from `UTexture`):
 	UPROPERTY(Category=Texture, AssetRegistrySearchable)
 	TEnumAsByte<enum TextureFilter> Filter;
 
-Once this flag was added to the Filter property of `UTexture`, all `UTextures` that were saved thereafter now have an entry in their `FAssetData`'s `TagsAndValues` map whose key is `Filter` and whose value is the string representation of the enum value, such as `"TF_Linear"`.
+Once this flag was added to the 'Filter' property of `UTexture`, all `UTextures` that were saved thereafter now have an entry in their `FAssetData`'s `TagsAndValues` map whose key is `"Filter"` and whose value is the string representation of the enum value, such as `"TF_Linear"`.
 
 [REGION:note]
-	This requires assets to be resaved before their properties will be discovered by the Asset Registry.
+	This **requires assets to be resaved** before their properties will be discovered by the Asset Registry.
 [/REGION]
 
 If you want the Asset Registry to be able to search for information that is not directly a UProperty, 
diff --git a/Engine/Documentation/Source/Programming/Automation/Automation.INT.udn b/Engine/Documentation/Source/Programming/Automation/Automation.INT.udn
--- a/Engine/Documentation/Source/Programming/Automation/Automation.INT.udn
+++ b/Engine/Documentation/Source/Programming/Automation/Automation.INT.udn
@@ -3,7 +3,6 @@ Title:Automation System Overview
 Crumbs: %ROOT%, Programming
 Description:Overview of the automation system used for unit testing, feature testing, and content stress testing.
 Navigation:topic
-version: 4.9
 
 [VAR:Topic]
 [OBJECT:Topic]
@@ -46,27 +45,27 @@ version: 4.9
 [/VAR]
 
 [EXCERPT:Intro]
-The **Automation System** provides the ability to perform Unit Testing, Feature Testing, and Content 
-Stress Testing using the power of the **Unreal Message Bus** in order 
+The **Automation System** provides the ability to perform _Unit Testing_, _Feature Testing_, and _Content 
+Stress Testing_ using the power of the **Unreal Message Bus** in order 
 to increase stability.
 
-The Automation system works by performing one or more automation tests. Automation tests can be broken down 
+The Automation system works by performing one or more **automation tests**. Automation tests can be broken down 
 into different categories depending on their purpose or function:
 [/EXCERPT:Intro]
 
-$ Unit Test : API level verification tests.  See TimespanTest.cpp or DateTimeTest.cpp for examples.
+$ Unit Test : API level verification tests.  See `TimespanTest.cpp` or `DateTimeTest.cpp` for examples.
 $ Feature Tests : System level verification tests such as verify PIE works, verify in-game stats work, 
-and verify changing resolution works.  See EditorAutomationTests.cpp or EngineAutomationTests.cpp for examples. 
+and verify changing resolution works.  See `EditorAutomationTests.cpp` or `EngineAutomationTests.cpp` for examples. 
 $ Smoke Tests : Smoke Tests are just considered a speed promise by the implementer.  They are intended to 
 be fast so they can run _everytime_ the editor, game, or a commandlet starts.  They are also selected by 
 default in the [UI](Programming/Automation/UserGuide#UserInterface).
 	
 	[REGION:warning]
-		All Smoke tests are intended to complete within 1 second.  Only mark Unit Tests or fast Feature Tests as Smoke Tests.
+		All Smoke tests are intended to complete within **1 second**.  Only mark Unit Tests or fast Feature Tests as Smoke Tests.
 	[/REGION]
 
 $ Content Stress Tests : More thorough testing of a particular system to avoid crashes, such as loading all maps 
-or loading and compiling all blueprints.  See EditorAutomationTests.cpp or EngineAutomationTests.cpp 
+or loading and compiling all blueprints.  See `EditorAutomationTests.cpp` or `EngineAutomationTests.cpp` 
 for examples.
 
 $ Screen Shot Comparison :The Screen Shot Comparison tool enables QA to quickly compare screenshots to identify potential 
diff --git a/Engine/Documentation/Source/Programming/Automation/ScreenShotComparison/ScreenShotComparisonGuide.INT.udn b/Engine/Documentation/Source/Programming/Automation/ScreenShotComparison/ScreenShotComparisonGuide.INT.udn
--- a/Engine/Documentation/Source/Programming/Automation/ScreenShotComparison/ScreenShotComparisonGuide.INT.udn
+++ b/Engine/Documentation/Source/Programming/Automation/ScreenShotComparison/ScreenShotComparisonGuide.INT.udn
@@ -2,7 +2,7 @@ Availability:Public
 Title:Screen Shot Comparison User Guide
 Crumbs: %ROOT%, Programming, Programming/Automation
 Description:Developers's guide to creating and comparing screen shots
-Version: 4.9
+Version: 4.5
 
 The screenshot browser tool allows QA to quickly and easily compare screenshots. The screenshots are generated via an automation test, and can be viewed in the Unreal Frontend tool. 
 A history of screenshots is maintained, so QA can compare the latest screenshots to a history, and identify any obvious rendering errors with the latest build.
@@ -33,9 +33,9 @@ Run the test by clicking the run button. This will start an automation test that
 
 As the screenshots are being captured, they will be sent to the Unreal Frontend app that requested the test. Then the UI will update automatically when it receives a new screenshot.
 
-$ Select **Screen Comparison** from the Session Frontend's Window menu :
+$ Select **Screen Comparison** from the Session Frontend's *Window* menu :
 
-This will display a UI that updates as it receives images.
+This will display an UI that update as it receives images.
 
 $ UI Navigation :
 
diff --git a/Engine/Documentation/Source/Programming/Automation/TechnicalGuide/AutomationTechnicalGuide.INT.udn b/Engine/Documentation/Source/Programming/Automation/TechnicalGuide/AutomationTechnicalGuide.INT.udn
--- a/Engine/Documentation/Source/Programming/Automation/TechnicalGuide/AutomationTechnicalGuide.INT.udn
+++ b/Engine/Documentation/Source/Programming/Automation/TechnicalGuide/AutomationTechnicalGuide.INT.udn
@@ -18,7 +18,7 @@ Blueprints would be good fits for complex automation tests.
 ## Architecture and Execution
 
 All automation tests derive from the **FAutomationTestBase** class which defines generic functionality for 
-performing a set of commands. The important functions of the FAutomationTestBase class used when setting 
+performing a set of commands. The important functions of the **FAutomationTestBase** class used when setting 
 up a new automation test are:
 
 | Member Functions | Description |
@@ -69,7 +69,7 @@ For example, the declaration of a new simple test for the `SetRes` command funct
 
 	IMPLEMENT_SIMPLE_AUTOMATION_TEST( FSetResTest, "Windows.SetResolution", ATF_Game )
 
-Using the SetRes example above, the implementation would be:
+Using the _SetRes_ example above, the implementation would be:
 
 	bool FSetResTest::RunTest(const FString& Parameters)
 	{
@@ -141,7 +141,7 @@ Using the map loading example, the implementations would be:
 ## Latent Commands
 
 Latent Commands allow for sections of an automation test to be run over multiple frames. These are intended to 
-be queued up during RunTest calls.
+be queued up during _RunTest_ calls.
 
 The first step is declaring a Latent Command with the following syntax:
 
@@ -149,9 +149,9 @@ The first step is declaring a Latent Command with the following syntax:
 
 	bool FExecStringLatentCommand::Update()
 
-The Update call should return true when it is done executing and false if it would like to stop execution of the 
+The Update call should return _true_ when it is done executing and _false_ if it would like to stop execution of the 
 automation test and try again next frame. Commands execute in order and execution will not continue if a latent 
-command returns false from Update.
+command returns _false_ from _Update_.
 
 The second step is adding the latent command to the queue for execution. This is wrapped in a macro to avoid boiler 
 plate code as follows:
diff --git a/Engine/Documentation/Source/Programming/Basics/Basics.INT.udn b/Engine/Documentation/Source/Programming/Basics/Basics.INT.udn
--- a/Engine/Documentation/Source/Programming/Basics/Basics.INT.udn
+++ b/Engine/Documentation/Source/Programming/Basics/Basics.INT.udn
@@ -2,12 +2,6 @@ Availability:Public
 Title:Programming Basics
 Crumbs: %ROOT%, Programming
 Description:Information for programmers developing with Unreal Engine.
-related: Programming\UnrealArchitecture\Reference\Classes
-related: Programming\UnrealArchitecture\Reference\Functions
-related: Programming\UnrealArchitecture\Reference\Properties
-related: Programming\UnrealArchitecture\Reference\Structs
-related: Programming\UnrealArchitecture\Reference\Interfaces
-version: 4.9
 
 
 An individual game is defined by a game project that contains all of the code, content, and settings specific to that particular game. Gameplay code is contained within one or more 
@@ -62,10 +56,29 @@ Different types of components can be used to control how Actors move, how they a
 UObject class, which is the base class of all gameplay Objects.  This means they cannot be directly instanced into the world; they must belong to an Actor.
 
 Each Actor or Object is a single instance of a class.  The class sets up the template for the Actor or Object. It defines the variables that can be set for that Actor or Object, and the functions that
-can be carried out within that Actor or Object.  You can create new classes, or types of Objects and Actors, with C++ code. Blueprint Classes primarily allow you to create classes that set up new Actors, 
-although you can extend a few Objects with Blueprint Classes as well. You can also combine the two, by creating a new C++ class and then making a Blueprint Class derived from that C++ class. To learn more about 
+can be carried out within that Actor or Object.  You can create new classes, or types of Objects and Actors, with C++ code. _Blueprint Classes_ primarily allow you to create classes that set up new Actors, 
+although you can extend a few Objects with _Blueprint Classes_ as well. You can also combine the two, by creating a new C++ class and then making a _Blueprint Class_ derived from that C++ class. To learn more about 
 creating classes so you can make new kinds of Actors and Objects, see the [Class Creation Basics](Gameplay\ClassCreation) page.
 
 For more on Objects, Actors, and Components, see the [](Programming/UnrealArchitecture) documentation.
 
 
+[OBJECT:TopicList]
+	[PARAM:icon]
+		![](%ROOT%/reference_icon.png)(convert:false)
+	[/PARAM]
+	[PARAM:title]
+		Related Pages
+	[/PARAM]
+	[PARAM:description]
+	[/PARAM]
+	[PARAM:links]
+	* [](Programming\UnrealArchitecture\Reference\Classes)
+	* [](Programming\UnrealArchitecture\Reference\Functions)
+	* [](Programming\UnrealArchitecture\Reference\Properties)
+	* [](Programming\UnrealArchitecture\Reference\Structs)
+	* [](Programming\UnrealArchitecture\Reference\Interfaces)
+	[/PARAM]
+[/OBJECT]
+
+
diff --git a/Engine/Documentation/Source/Programming/Basics/CommandLineArguments/CommandLineArguments.INT.udn b/Engine/Documentation/Source/Programming/Basics/CommandLineArguments/CommandLineArguments.INT.udn
--- a/Engine/Documentation/Source/Programming/Basics/CommandLineArguments/CommandLineArguments.INT.udn
+++ b/Engine/Documentation/Source/Programming/Basics/CommandLineArguments/CommandLineArguments.INT.udn
@@ -23,7 +23,7 @@ These arguments are used to force the main editor executable (`UE4Editor.exe`) t
 These commands are not case sensitive.
 [/REGION]
 
-Example:
+*Example:*
 
       
 	UE4Editor.exe -game
diff --git a/Engine/Documentation/Source/Programming/Basics/ConfigurationFiles/ConfigurationFiles.INT.udn b/Engine/Documentation/Source/Programming/Basics/ConfigurationFiles/ConfigurationFiles.INT.udn
--- a/Engine/Documentation/Source/Programming/Basics/ConfigurationFiles/ConfigurationFiles.INT.udn
+++ b/Engine/Documentation/Source/Programming/Basics/ConfigurationFiles/ConfigurationFiles.INT.udn
@@ -14,27 +14,27 @@ Engine configuration files are used for object and variable default values. User
 are created when you create a new blank project with the **Project Wizard**. New projects that begin with templates may also generate DefaultEditor.ini and DefaultInput.ini configuration files
 if needed.  
 
-The SaveConfig() function can be called on a class with the  [](Programming\UnrealArchitecture\Reference\Classes\Specifiers\Config) class specifier.  This saves any properties marked with the 
-**Config** property specifier to the appropriate configuration file. Generally, variables saved by **SaveConfig()** are in a section title
-which follows the format [(package).(classname)].  For instance, the [/Script/Engine.Engine] section in DefaultEngine.ini points to the Engine class stored within the Engine package.
+The `SaveConfig()` function can be called on a class with the  [](Programming\UnrealArchitecture\Reference\Classes\Specifiers\Config) class specifier.  This saves any properties marked with the 
+`Config` property specifier to the appropriate configuration file. Generally, variables saved by **SaveConfig()** are in a section title
+which follows the format `[(package).(classname)]`.  For instance, the `[/Script/Engine.Engine]` section in DefaultEngine.ini points to the Engine class stored within the Engine package.
 There are some exceptions with hard-coded section names.
 
 A number of the settings previously available only through editing configuration files are available in Unreal Editor in the [](Engine\UI\ProjectSettings) editor.
 
 ## Specifying Configurable Variables
 
-In order to specify which variables should be read in from configuration files, the class that contains those variables must first be given the Config specifier in its UCLASS macro.
+In order to specify which variables should be read in from configuration files, the class that contains those variables must first be given the `Config` specifier in its `UCLASS` macro.
 
 	UCLASS(Config=Game)
 	class AExampleClass : public AActor
 
-[REGION:note]
-	A category (i.e. Game) must be supplied for the Config specifier. This determines which configuration files the class's variables are read from and saved to. All possible categories are 
+[REGION:tip]
+	Note that a category (i.e. Game) must be supplied for the `Config` specifier. This determines which configuration files the class's variables are read from and saved to. All possible categories are 
 	defined in FConfigCache.ini. For a list of all configuration file categories, see [Configuration Categories](#configurationcategories).
 [/REGION]
 
-Decorating a class with the Config specifier just indicates that that class can have variables read in from configuration files, and specifies which files the configurations are read from. To
-specify a particular variable as read in and saved to a configuration file, the Config specifier must also be supplied to a UPROPERTY() macro.
+Decorating a class with the `Config` specifier just indicates that that class can have variables read in from configuration files, and specifies which files the configurations are read from. To
+specify a particular variable as read in and saved to a configuration file, the `Config` specifier must also be supplied to a `UPROPERTY()` macro.
 
 	UCLASS(Config=Game)
 	class AExampleClass : public AActor
@@ -46,7 +46,7 @@ specify a particular variable as read in and saved to a configuration file, the
 
 	};
 
-No category is supplied to the property's Config specifier. The ExampleVariable property can now be read from any Game configuration file in the [configuration file hierarchy](#filehierarchy),
+No category is supplied to the property's `Config` specifier. The `ExampleVariable` property can now be read from any Game configuration file in the [configuration file hierarchy](#filehierarchy),
 as long as the information is specified with the following syntax:
 
 	[/Script/ModuleName.ExampleClass]
@@ -54,9 +54,9 @@ as long as the information is specified with the following syntax:
 
 ### Config Files and Inheritance
 
-The Config, UCLASS, and UPROPERTY specifiers are inherited.  This means that a child class can either read in or save out all variables specified as Config in the parent class, and they
-will be in the same configuration file category. The variables will all be under a section title with the child class's name. For example, the configuration file information for a ChildExampleClass
-which inherits from the ExampleClass above would look like the following lines, and be saved in the same Game configuration files.
+The `Config` `UCLASS` and `UPROPERTY` specifiers are inherited.  This means that a child class can either read in or save out all variables specified as `Config` in the parent class, and they
+will be in the same configuration file category. The variables will all be under a section title with the child class's name. For example, the configuration file information for a `ChildExampleClass`
+which inherits from the `ExampleClass` above would look like the following lines, and be saved in the same Game configuration files.
 
 	[/Script/ModuleName.ChildExampleClass]
 	ExampleVariable=0.0f
@@ -64,8 +64,8 @@ which inherits from the ExampleClass above would look like the following lines,
 
 ### Per-Instance Configurations
 
-Unreal Engine 4 has the ability to save the configuration of an object to any desired configuration file. If the **PerObjectConfig** specifier is used in the UCLASS macro, 
-configuration information for this class will be stored per instance, where each instance has a section in the .ini file named after the object in the format [ObjectName ClassName]. 
+Unreal Engine 4 has the ability to save the configuration of an object to any desired configuration file. If the `PerObjectConfig` specifier is used in the `UCLASS` macro, 
+configuration information for this class will be stored per instance, where each instance has a section in the .ini file named after the object in the format `[ObjectName ClassName]`. 
 This keyword is propagated to child classes.
 
 
@@ -95,7 +95,7 @@ to all projects, while project-specific settings should be in files in the proje
 The below file hierarchy example is for the Engine category of configuration files.
 
 1. Engine/Config/Base.ini
-	[REGION:note]
+	[REGION:tip]
 	Base.ini is usually empty.
 	[/REGION]
 1. Engine/Config/BaseEngine.ini
@@ -103,7 +103,7 @@ The below file hierarchy example is for the Engine category of configuration fil
 1. [ProjectDirectory]/Config/DefaultEngine.ini
 1. [ProjectDirectory]/Config/[Platform]/[Platform]Engine.ini
 1. [ProjectDirectory]/Saved/Config/[Platform]/Engine.ini
-	[REGION:note]
+	[REGION:tip]
 	The configuration file in the Saved directory only stores the project-specific and platform-specific differences in the stack of configuration files.
 	[/REGION]
 	
@@ -130,7 +130,7 @@ Typical configuration files consist of sections of key-value pairs, arranged as
 * *!* - Removes a property; but you do not have to have an exact match, just the name of the property.
 
 [REGION:note]
-. is like + except it will potentially add a duplicate line. This is useful for the bindings (as seen in DefaultInput.ini), for instance, where the bottom-most binding takes effect, so if you add something like:
+**Note**: . is like + except it will potentially add a duplicate line. This is useful for the bindings (as seen in DefaultInput.ini), for instance, where the bottom-most binding takes effect, so if you add something like:
     
     [/Script/Engine.PlayerInput]
     Bindings=(Name="Q",Command="Foo")
@@ -142,7 +142,7 @@ It will work appropriately. Using a *+* there would fail to add the last line, a
 
 #### Comments
 
-Most people seem to be under the impression that the semicolon denotes comments in configuration files, but they are not (FConfigFile::ProcessInputFileContents does not actually treat them, or any other string, as a comment delimiter). This behavior is intentional. Technically any character can represent a different key-value pair. Typically, a semicolon is placed at the beginning of a new line. It works like a comment, but it is not actually.
+Most people seem to be under the impression that the semicolon denotes comments in configuration files, but they are not (`FConfigFile::ProcessInputFileContents` does not actually treat them, or any other string, as a comment delimiter). This behavior is intentional. Technically any character can represent a different key-value pair. Typically, a semicolon is placed at the beginning of a new line. It works like a comment, but it is not actually.
      
     ; This is a Comment
 	; So is this!
diff --git a/Engine/Documentation/Source/Programming/BlueprintFunctionLibraries/BlueprintFunctionLibraries.INT.udn b/Engine/Documentation/Source/Programming/BlueprintFunctionLibraries/BlueprintFunctionLibraries.INT.udn
--- a/Engine/Documentation/Source/Programming/BlueprintFunctionLibraries/BlueprintFunctionLibraries.INT.udn
+++ b/Engine/Documentation/Source/Programming/BlueprintFunctionLibraries/BlueprintFunctionLibraries.INT.udn
@@ -19,7 +19,8 @@ Creating a Blueprint Function Library is very similar to exposing functions to B
 		UFUNCTION(BlueprintCallable, Category="Analytics")
 		static bool StartSession();
 
-As you can see in the example above, a Blueprint Function Library is indirectly UObject derived so requires the standard UCLASS() and GENERATED_UCLASS_BODY() macros. It also decorates the functions that are to be callable from Blueprints with the UFUNCTION() macro. Functions in a Blueprint Function Library can be designated as BlueprintCallable or BlueprintPure depending on whether the calls have side effects or not. See /UE4/Engine/Plugins/Runtime/Analytics/AnalyticsBlueprintLibrary for the full source code. Here’s the implementation of the StartSession() function:
+As you can see in the example above, a Blueprint Function Library is indirectly UObject derived so requires the standard UCLASS() and GENERATED_UCLASS_BODY() macros. It also decorates the functions that are to be callable from Blueprints with the UFUNCTION() macro. NOTE: functions in a Blueprint Function Library can be designated as BlueprintCallable or BlueprintPure depending on whether the calls have side effects or not. See /UE4/Engine/Plugins/Runtime/Analytics/AnalyticsBlueprintLibrary for the full source code. Here’s the implementation of the StartSession() function:
+
 	bool UAnalyticsBlueprintLibrary::StartSession()
 	{
 		TSharedPtr<IAnalyticsProvider> Provider = FAnalytics::Get().GetDefaultConfiguredProvider();
diff --git a/Engine/Documentation/Source/Programming/Development/BuildingUnrealEngine/BuildingUnrealEngine.INT.udn b/Engine/Documentation/Source/Programming/Development/BuildingUnrealEngine/BuildingUnrealEngine.INT.udn
--- a/Engine/Documentation/Source/Programming/Development/BuildingUnrealEngine/BuildingUnrealEngine.INT.udn
+++ b/Engine/Documentation/Source/Programming/Development/BuildingUnrealEngine/BuildingUnrealEngine.INT.udn
@@ -8,7 +8,7 @@ Related: Engine/Content/QuickStart
 Related: Programming/QuickStart
 Related: Engine/QuickStart
 Related: Support/Builds
-Version: 4.9
+Version: 4.8
 SkillLevel: Advanced
 
 %Globals:osselect_linux%
@@ -39,7 +39,7 @@ SkillLevel: Advanced
 			All project files are intermediate (`[UE4Root]\Engine\Intermediate\ProjectFiles`). You must generate project files each time you sync a new build to ensure they are up to date. If you delete your `Intermediate` folder, you must regenerate project files using the 'GenerateProjectFiles' batch file.
 			[/REGION]
 
-		1.	Load the project into Visual Studio by **Double-clicking UE4.sln**.
+		1.	Load the project into Visual Studio by **double-clicking UE4.sln**.
 
 		1.	Set your solution configuration to **Development Editor**.
 
@@ -90,7 +90,7 @@ SkillLevel: Advanced
 		
 		1.	Inside the root directory, run **GenerateProjectFiles.command** to set-up your project files.
 		
-		1.	Load the project into XCode by **Double-clicking UE4.xcodeproj**.
+		1.	Load the project into XCode by **double-clicking UE4.xcodeproj**.
 
 		1.	To set your build target, select **UE4Editor - Mac** for **My Mac** in the title bar.
 
@@ -165,7 +165,7 @@ If you're looking to quickly get started with UE4, refer to the following tutori
 *	[](Engine/QuickStart)
 
 
-[REGION:note]
+[REGION:tip]
 UE4's in-editor help features are a great way to get your questions answered.
 [/REGION]
 
diff --git a/Engine/Documentation/Source/Programming/Development/CodingStandard/CodingStandard.INT.udn b/Engine/Documentation/Source/Programming/Development/CodingStandard/CodingStandard.INT.udn
--- a/Engine/Documentation/Source/Programming/Development/CodingStandard/CodingStandard.INT.udn
+++ b/Engine/Documentation/Source/Programming/Development/CodingStandard/CodingStandard.INT.udn
@@ -111,7 +111,7 @@ Comments are communication; communication is vital. Some things to keep in mind
 * Write self-documenting code:
 
 		// Bad:
-		t = s + l - b;
+		t = s + l + b;
 		
 		// Good:
 		TotalLeaves = SmallLeaves + LargeLeaves - SmallAndLargeLeaves;
@@ -134,7 +134,7 @@ Comments are communication; communication is vital. Some things to keep in mind
 		// total number of leaves is sum of
 		// small and large leaves less the
 		// number of leaves that are both
-		t = s + l - b;
+		t = s + l + b;
 		
 		// Good:
 		TotalLeaves = SmallLeaves + LargeLeaves - SmallAndLargeLeaves;
diff --git a/Engine/Documentation/Source/Programming/Development/CompilingProjects/CompilingProjects.INT.udn b/Engine/Documentation/Source/Programming/Development/CompilingProjects/CompilingProjects.INT.udn
--- a/Engine/Documentation/Source/Programming/Development/CompilingProjects/CompilingProjects.INT.udn
+++ b/Engine/Documentation/Source/Programming/Development/CompilingProjects/CompilingProjects.INT.udn
@@ -135,7 +135,7 @@ When running Unreal, it is important to add the `-game` flag if you rebuilt your
 
 ## Building with Xcode
 
-[REGION:note]
+[REGION:tip]
 When building in Xcode, you are only compiling the game project, not the Editor.  
 [/REGION]
 
diff --git a/Engine/Documentation/Source/Programming/Development/Development.INT.udn b/Engine/Documentation/Source/Programming/Development/Development.INT.udn
--- a/Engine/Documentation/Source/Programming/Development/Development.INT.udn
+++ b/Engine/Documentation/Source/Programming/Development/Development.INT.udn
@@ -2,7 +2,6 @@ Availability:Public
 Title:Development Setup
 Crumbs: %ROOT%, Programming
 Description:Setting up your code environment, creating code projects, and working with Unreal Engine source.
-version: 4.9
 
 [OBJECT:TopicButtonList]
 	[PARAM:icon]
diff --git a/Engine/Documentation/Source/Programming/Development/ManagingGameCode/CppClassWizard/CppClassWizard.INT.udn b/Engine/Documentation/Source/Programming/Development/ManagingGameCode/CppClassWizard/CppClassWizard.INT.udn
--- a/Engine/Documentation/Source/Programming/Development/ManagingGameCode/CppClassWizard/CppClassWizard.INT.udn
+++ b/Engine/Documentation/Source/Programming/Development/ManagingGameCode/CppClassWizard/CppClassWizard.INT.udn
@@ -9,7 +9,7 @@ Version: 4.9
 [EXCERPT:main]
 
 
-The **C++ Class Wizard** provides a fast and easy way to add native C++ code classes into your project for you to extend 
+The _C++ Class Wizard_ provides a fast and easy way to add native C++ code classes into your project for you to extend 
 with your own functionality, if you wish. This converts a content-only project into a code project. You can access the C++ Class Wizard like so:
 
 [REGION:note]
@@ -17,11 +17,11 @@ On Windows, please make sure you have Visual Studio 2013 Express, Community, or
 [/REGION:note]
 
 
-1. In the main editor, select **File > New C++ Class...**
+1. In the main editor, select **File > Add Code to Project...**
 
 	![](FileCPP.png)
 
-1. The **C++ Class Wizard** will appear, showing commonly added classes. If you do not see the class you are looking for, then you will need to view the entire Class hierarchy listing. Click the **Show All Classes** checkbox in the upper-right corner of the window. 
+1. The **C++ Class Wizard** will appear, showing commonly added classes. If you do not see the class you are looking for, then you will need to view the entire Class hierarchy listing. Click the ![](button_ShowAllClasses.png)(h:25) checkbox in the upper-right corner of the window. 
 
 	[REGION:imagetable]
 		| [REGION:lightbox][![](CppClassWizard.png)(w:430)](CppClassWizard.png)[/REGION] | [REGION:lightbox][![](AllClasses.png)(w:430)](AllClasses.png)[/REGION] |
@@ -29,9 +29,9 @@ On Windows, please make sure you have Visual Studio 2013 Express, Community, or
 		| Default Classes | All Classes |
 	[/REGION]
 
-	Choose the Class you wish to add and click the **Next >** button.
+	Choose the Class you wish to add and click the ![](button_Next.png) button.
 
-1. You will then be prompted to enter a **Name** for your new Class. Do so and click the **Create Class** button. This will create the header (.h) and source (.cpp) files.
+1. You will then be prompted to enter a **name** for your new Class. Do so and click the ![](button_Create.png) button. This will create the header (.h) and source (.cpp) files.
 	
 	[REGION:lightbox]
 		[![](NameNewActor.png)(w:600)](NameNewActor.png)
@@ -51,7 +51,7 @@ On Windows, please make sure you have Visual Studio 2013 Express, Community, or
 		active_button_content
 		[/PARAMLITERAL]
 		[PARAM:content]
-		The code will immediately open in Visual Studio 2013, ready for editing.
+		The code will immediately open in _Visual Studio 2013_, ready for editing.
 	
 		![](CodeEditing.png)
 
@@ -67,7 +67,7 @@ On Windows, please make sure you have Visual Studio 2013 Express, Community, or
 		[PARAMLITERAL:active]
 		[/PARAMLITERAL]
 		[PARAM:content]
-		The code will immediately open in Xcode, ready for editing.
+		The code will immediately open in _Xcode_, ready for editing.
 
 		![](CodeEditing_xcode.png)
 
diff --git a/Engine/Documentation/Source/Programming/Development/ManagingGameCode/ManagingGameCode.INT.udn b/Engine/Documentation/Source/Programming/Development/ManagingGameCode/ManagingGameCode.INT.udn
--- a/Engine/Documentation/Source/Programming/Development/ManagingGameCode/ManagingGameCode.INT.udn
+++ b/Engine/Documentation/Source/Programming/Development/ManagingGameCode/ManagingGameCode.INT.udn
@@ -73,12 +73,12 @@ Version: 4.9
 	[PUBLISH:Licensee]
 	**If your project is inside the UE4 root directory:**
 
-	* Open the UE4.sln Visual Studio solution located in the UE4 root directory.
+	* Open the `UE4.sln` Visual Studio solution located in the UE4 root directory.
 
 	**If your project is outside of the UE4 root directory:**
 	[/PUBLISH:Licensee]
 
-	* Open the [ProjectName].sln Visual Studio solution located in the root of the project's directory.
+	* Open the `[ProjectName].sln` Visual Studio solution located in the root of the project's directory.
 	[/PARAM]
 [/OBJECT]
 [OBJECT:ToggleButtonContent]
@@ -97,7 +97,7 @@ Version: 4.9
 
 	You can also open the project in Xcode through Finder or Xcode's **File > Open**.
 
-	* Open the [ProjectName].xcodeproj Xcode project located in the root of the project's directory. 
+	* Open the `[ProjectName].xcodeproj` Xcode project located in the root of the project's directory. 
 	[/PARAM]
 [/OBJECT]
 
@@ -127,8 +127,8 @@ Version: 4.9
 
 **If your project is inside the UE4 root directory:**
 
-1. Run the GenerateProjectFiles.bat batch file located in the root `UE4` directory.
-1. Open the UE4.sln Visual Studio solution located in the root `UE4` directory.
+1. Run the `GenerateProjectFiles.bat` batch file located in the root `UE4` directory.
+1. Open the `UE4.sln` Visual Studio solution located in the root `UE4` directory.
 
 There are advanced options available for [](Programming/UnrealBuildSystem/ProjectFileGenerator), which can be used to generate project files for shipping, additional platforms, or Visual Studio 2013. 
 
@@ -152,9 +152,9 @@ To get the **Generate Visual Studio Files** menu item to show up in the **right-
 	active_button_content
 	[/PARAMLITERAL]
 	[PARAM:content]
-	1. Navigate to the location of [ProjectName].uproject in Windows Explorer.
+	1. Navigate to the location of `[ProjectName].uproject` in Windows Explorer.
 	
-	1. **Right-click** on the [ProjectName].uproject file and select **Generate Visual Studio Files**.
+	1. **Right-click** on the `[ProjectName].uproject` file and select **Generate Visual Studio Files**.
 	
 		![](Generate_Visual_Studio_Files.png)
 
@@ -162,7 +162,7 @@ To get the **Generate Visual Studio Files** menu item to show up in the **right-
 
 		![](generate_project_files.png)
 
-	1. Open the [ProjectName].sln Visual Studio solution located in the root of the project's directory to view the game project in Visual Studio.
+	1. Open the `[ProjectName].sln` Visual Studio solution located in the root of the project's directory to view the game project in Visual Studio.
 	[/PARAM]
 [/OBJECT]
 [OBJECT:ToggleButtonContent]
@@ -191,8 +191,8 @@ To get the **Generate Visual Studio Files** menu item to show up in the **right-
 
 [PUBLISH:Licensee]
 [REGION:note]
-For advanced project file generation options for projects outside the UE4 root directory, edit the RegisterShellCommands.bat batch file. The Generate Visual Studio projects (all platforms) entry
-shows an example of an additional menu entry, in this case with the `-AllPlatforms` flag. The [available options](Programming/UnrealBuildSystem/ProjectFileGenerator) are the same as for the main GenerateProjectFiles.bat file.  
+For advanced project file generation options for projects outside the UE4 root directory, edit the `RegisterShellCommands.bat` batch file. The `Generate Visual Studio projects (all platforms)` entry
+shows an example of an additional menu entry, in this case with the `-AllPlatforms` flag. The [available options](Programming/UnrealBuildSystem/ProjectFileGenerator) are the same as for the main `GenerateProjectFiles.bat` file.  
 [/REGION:note]
 [/PUBLISH:Licensee]
 
diff --git a/Engine/Documentation/Source/Programming/Development/VisualStudioSetup/UnrealVS/UnrealVS.INT.udn b/Engine/Documentation/Source/Programming/Development/VisualStudioSetup/UnrealVS/UnrealVS.INT.udn
--- a/Engine/Documentation/Source/Programming/Development/VisualStudioSetup/UnrealVS/UnrealVS.INT.udn
+++ b/Engine/Documentation/Source/Programming/Development/VisualStudioSetup/UnrealVS/UnrealVS.INT.udn
@@ -31,12 +31,12 @@ The UnrealVS extension does not work with the Visual Studio Express editions. It
 
 		[UE4RootLocation]/Engine/Extras/UnrealVS/<VS-Version>/UnrealVS.vsix
 
-1. Run the `UnrealVS.vsix` file by **Double-clicking** it to begin the installation.
+1. Run the `UnrealVS.vsix` file by **double-clicking** it to begin the installation.
 
-1. UnrealVS extensions target specific Visual Studio versions.  Check that the installer is targeting your version and that the box is ticked then click the **Install** button to proceed.
+1. UnrealVS extensions target specific Visual Studio versions.  Check that the installer is targeting your version and that the box is ticked then click the ![Install](button_install.png) button to proceed.
 
 	![](install_dialog.png)
-1. Click the **Close** button once the installation is complete.
+1. Click the ![Close](button_close.png) button once the installation is complete.
 
 	![Install Complete](install_complete_dialog.png)
 
@@ -51,7 +51,7 @@ The UnrealVS extension does not work with the Visual Studio Express editions. It
 	The extension is also shown in the Installed Products list in the **Help->About Microsoft Visual Studio** dialog.
 	[/REGION]
 
-1. Go to **View->Toolbars** or **Right-click** the Visual Studio toolbar area and select **UnrealVS** to display the extension's toolbar.
+1. Go to **View->Toolbars** or **right-click** the Visual Studio toolbar area and select **UnrealVS** to display the extension's toolbar.
 
 	![UnrealVS Toolbar](unrealvs_toolbar_cmd.png)
 
@@ -85,10 +85,10 @@ hotkey or other means of running custom commands.
 
 1. Go to **Tools->Options...** and select **Keyboard** under the **Environment** tree.
 
-1. Select the **UnrealVS.BuildStartupProject** command in the commands list. 
+1. Select the **_UnrealVS.BuildStartupProject_** command in the commands list. 
 
 	[REGION:tip]
-		Type "UnrealVS" in the search box to filter the list.
+		Type '`UnrealVS`' in the search box to filter the list.
 	[/REGION]
 
 	![Filtered Command List](build_startup_search.png)
@@ -103,12 +103,12 @@ hotkey or other means of running custom commands.
 		combination is already in use.
 	[/REGION]
 
-1. Press the **Assign** button to bind the keys to the command. The keys are 
+1. Press the ![Assign](button_assign.png) button to bind the keys to the command. The keys are 
 	displayed in the **Shortcuts for selected command** box.
 
 	![Assigned Keys](build_startup_assigned.png)
 
-1. Press the **OK** button to save the changes. Now, when you use the shortcut, the project set as the startup project will automatically build.
+1. Press the ![OK](button_ok.png) button to save the changes. Now, when you use the shortcut, the project set as the startup project will automatically build.
 
 
 ## Command Line Arguments
@@ -138,9 +138,7 @@ file and run it manually.
 
 **To refresh project files:**
 
-1. Press the Refresh Projects Refresh button on the **UnrealVS** toolbar.
-
-![Refresh Projects](button_refresh.png)
+1. Press the Refresh Projects ![Refresh Projects](button_refresh.png) button on the **UnrealVS** toolbar.
 
 1. The project generation process will be displayed in the **Output** window as it occurs.
 
@@ -154,11 +152,9 @@ file and run it manually.
 
 **To open the Batch Builder window:**
 
-1. Press the Batch Builder button on the **UnrealVS** toolbar or use a keyboard shortcut that you have assigned to the command UnrealVS.BatchBuilder 
+1. Press the Batch Builder ![Batch Builder](button_batchbuilder.png) button on the **UnrealVS** toolbar or use a keyboard shortcut that you have assigned to the command _UnrealVS.BatchBuilder_ 
 (see **Build Startup Project** above for instructions on setting up keyboard shortcuts with **UnrealVS** commands).
 
-![Batch Builder](button_batchbuilder.png)
-
 1. The **UnrealVS Batch Builder** window appears.
 
 	![Batch Builder Window](batchbuilder_window.png)
@@ -167,11 +163,11 @@ file and run it manually.
 	* Add/remove jobs with the arrow buttons.
 	* Move selected jobs up/down the list with the arrow buttons.
 	* The controls edit the current **Job Set** displayed in the combo.
-	* To create a new Job Set enter a new name into the combo.
-	* You can select existing Job Sets using the dropdown list.
-	* The **Delete** button deletes the selected Job Set from the list.
-	* Start/Stop the batch build task on the current Job Set with the **Start** button.
-	* Job Sets are stored for next time in the .suo file for the loaded solution.
+	* To create a new **Job Set** enter a new name into the combo.
+	* You can select existing **Job Sets** using the dropdown list.
+	* The **Delete** button deletes the selected **Job Set** from the list.
+	* Start/Stop the batch build task on the current **Job Set** with the **Start** button.
+	* **Job Sets** are stored for next time in the .suo file for the loaded solution.
 
 1. When you click **Start**, the **Output** tab appears showing the progress of the batch build.
 
@@ -185,13 +181,13 @@ file and run it manually.
 		* The current executing build job is shown in bold
 	* Busy animation and stop button are shown when a batch build is running.
 	
-1. You can view the output from individual **Build Jobs** by **Double-clicking** an item in the **Output tab** or **Right-clicking** and selecting **Show Output** from the context menu.
+1. You can view the output from individual **Build Jobs** by **double-clicking** an item in the **Output tab** or **right-clicking** and selecting **Show Output** from the context menu.
 
 	![BatchBuild Output Pane](batchbuilder_outputpane.png)
 
 	[REGION:note]
-		Showing **Batch Builder** output for individual items is shown in the Visual Studio **Output** window in a pane named UnrealVS - BatchBuild.
-		This is not to be confused with the Build pane that shows the output of the current/last build job.
+		Showing **Batch Builder** output for individual items is shown in the Visual Studio **Output** window in a pane named _UnrealVS - BatchBuild_.
+		This is not to be confused with the _Build_ pane that shows the output of the current/last build job.
 	[/REGION]	
 	
 ## Quick Build Menu
@@ -202,7 +198,7 @@ The **Quick Build** menu provides an easy way to build a project with any config
 
 1. **Right-click** a Project node in the **Solution Explorer**.
 
-1. The **UnrealVS Quick Build** menu contains sub-menus for the available platforms and build configs in the Unreal Engine solution.
+1. The *UnrealVS Quick Build* menu contains sub-menus for the available platforms and build configs in the Unreal Engine solution.
 
 1. Select an item to build the selected project, platform, and config. This is much quicker than changing the Solution Configuration and Solution Platform in the IDE, starting a build, then switching
 the config and platform back again.
diff --git a/Engine/Documentation/Source/Programming/Development/VisualStudioSetup/VisualStudioSetup.INT.udn b/Engine/Documentation/Source/Programming/Development/VisualStudioSetup/VisualStudioSetup.INT.udn
--- a/Engine/Documentation/Source/Programming/Development/VisualStudioSetup/VisualStudioSetup.INT.udn
+++ b/Engine/Documentation/Source/Programming/Development/VisualStudioSetup/VisualStudioSetup.INT.udn
@@ -46,14 +46,14 @@ The following are recommended settings for developers using Visual Studio with U
 
 Other configuration steps you may find useful are:
 
-* Turn off **Show Inactive Blocks**.  If you do not, many chunks of code may appear grayed out in the text editor.  (**Tools > Options > Text Editor > C/C++ > View**).
-* Set **Disable External Dependencies Folders** to True to hide unneeded folders in the **Solution Explorer**.  (**Disable External Dependencies Folder** in **Tools > Options > Text Editor > C/C++ > Advanced**)
-* Turn off **Edit & Continue** features, you do not need them.  (**Tools > Options > Debugging > Edit and Continue**)
+* Turn off **Show Inactive Blocks**.  If you do not, many chunks of code may appear grayed out in the text editor.  (_Tools > Options > Text Editor > C/C++ > View)_.
+* Set **Disable External Dependencies Folders** to _True_ to hide unneeded folders in the **Solution Explorer**.  (**Disable External Dependencies Folder** in _Tools > Options > Text Editor > C/C++ > Advanced_)
+* Turn off **Edit & Continue** features, you do not need them.  (_Tools > Options > Debugging > Edit and Continue_)
 
 * Turn on IntelliSense. (See [Intellisense, Live Errors, and Squiggles](#Intellisense))
 
 **If you have Visual Assist X installed:**
-* Turn off **Format After Paste**.  It does weird things sometimes.  (**VAssistX > Visual Assist X Options > Advanced > Corrections**)
+* Turn off **Format After Paste**.  It does weird things sometimes.  (_VAssistX > Visual Assist X Options > Advanced > Corrections_)
 
 
 (#Intellisense)
diff --git a/Engine/Documentation/Source/Programming/Introduction/IntroToProgramming.INT.udn b/Engine/Documentation/Source/Programming/Introduction/IntroToProgramming.INT.udn
--- a/Engine/Documentation/Source/Programming/Introduction/IntroToProgramming.INT.udn
+++ b/Engine/Documentation/Source/Programming/Introduction/IntroToProgramming.INT.udn
@@ -92,7 +92,7 @@ Now let us expose that same property to Blueprints.
 	UPROPERTY(EditAnywhere, BlueprintReadWrite, Category="Damage")
 	int32 TotalDamage;
 
-As you can see, there is a Blueprint specific parameter to make a property available for reading and writing. There's a separate option, BlueprintReadOnly, you can use if you want the property to be treated as const in Blueprints. There are quite a few options available for controlling how a property is exposed to the engine. To see more options, follow this [link.](Programming/UnrealArchitecture/Reference/Properties/Specifiers)
+As you can see, there is a Blueprint specific parameter to make a property available for reading and writing. There's a separate option, BlueprintReadOnly, you can use if you want the property to be treated as *const* in Blueprints. There are quite a few options available for controlling how a property is exposed to the engine. To see more options, follow this [link.](Programming/UnrealArchitecture/Reference/Properties/Specifiers)
 
 Before continuing to the section below, let us add a couple of properties to this sample class. There is already a property to control the total amount of damage this actor will deal out, but let us take that further and make that damage happen over time. The code below adds one designer settable property and one that is visible to the designer but not changeable by them.
 
@@ -113,7 +113,7 @@ Before continuing to the section below, let us add a couple of properties to thi
 		...
 	};
 
-DamageTimeInSeconds is a property the designer can modify. The DamagePerSecond property is a calculated value using the designer's settings (see the next section). The VisibleAnywhere flag marks that property as viewable, but not editable in the Unreal Editor. The Transient flag means that it won't be saved or loaded from disk; it is meant to be a derived, non-persistent value. The image below shows the properties as part of the class defaults.
+	DamageTimeInSeconds is a property the designer can modify. The DamagePerSecond property is a calculated value using the designer's settings (see the next section). The VisibleAnywhere flag marks that property as viewable, but not editable in the Unreal Editor. The Transient flag means that it won't be saved or loaded from disk; it is meant to be a derived, non-persistent value. The image below shows the properties as part of the class defaults.
 
 ![image alt text](image_3.png)
 
@@ -155,17 +155,17 @@ Here again is the same view of the properties after we have added the **PostInit
 
 Here is a cool feature of Unreal that you might be surprised about if you are used to programming C++ in other projects.  You can compile your C++ changes without shutting down the editor!  There are two ways to do this:
 
-1. With the editor still running, go ahead and build from Visual Studio or Xcode like you normally would.  The editor will detect the newly compiled DLLs and reload your changes instantly!  
+1. With the editor still running, go ahead and **Build from Visual Studio or Xcode** like you normally would.  The editor will detect the newly compiled DLLs and reload your changes instantly!  
 
 	![image alt text](image_6.png)
-	[REGION:note]
-	If you are attached with the debugger, you'll need to detach first so that Visual Studio will allow you to Build.
-	[/REGION]
+
+	(Note that if you are attached with the debugger, you'll need to detach first so that Visual Studio will allow you to Build.)
+
 2. Or, simply click the **Compile **button on the editor's main toolbar.
 
 	![image alt text](image_7.png)
 
-You can use this feature in the sections below as we advance through the tutorial.
+You can use this feature in the sections below as we advance through the tutorial. What a time saver!
 
 ### Extending a C++ Class via Blueprints
 
@@ -175,7 +175,7 @@ First thing we are going to do is create a new Blueprint class from our AMyActor
 
 ![image alt text](image_8.png)
 
-Once you choose **Select**, a new, default named Blueprint class is created for you. In this case, I set the name to CustomActor1 as you can see in the snapshot of the **Content Browser** below.
+Once you choose Select, a new, default named Blueprint class is created for you. In this case, I set the name to CustomActor1 as you can see in the snapshot of the **Content Browser** below.
 
 ![image alt text](image_9.png)
 
@@ -244,7 +244,7 @@ This version still generates the thunking method to call into the Blueprint VM.
 		// Do something cool here
 	}
 
-Now this version of the function is called when the Blueprint in question does not override the method. One thing to note, is that in previous versions of the build tools the _Implementation() declaration was auto generated. In any version 4.8 or higher, you'll be expected to explicitly add that to the header.
+Now this version of the function is called when the Blueprint in question does not override the method. One thing to note, is that in future versions of the build tools the auto generated _Implementation() declaration will go away and you'll be expected to explicitly add that to the header. As of version 4.7, the auto generation of that declaration still occurs.
 
 Now that we have walked through the common gameplay programmer workflow and methods to work with designers to build out gameplay features, it is time for you to choose your own adventure. You can either continue with this document to read more about how we use C++ in the engine or you can jump right into one of our samples that we include in the launcher to get a more hands on experience. 
 
@@ -389,9 +389,9 @@ You can limit the scope of the search by providing a more specific type to the i
 	{
 		// ...
 	}
-[REGION:warning]
-Using object iterators in PIE (Play In Editor) can lead to unexpected results.  Since the editor is loaded, the object iterator will return all UObjects created for your game world instance, in addition to those that are just being used by the editor.
-[/REGION]
+
+Warning: Using object iterators in PIE (Play In Editor) can lead to unexpected results.  Since the editor is loaded, the object iterator will return all UObjects created for your game world instance, in addition to those that are just being used by the editor.
+
 Actor iterators work in much the same way as object iterators, but only work for objects that derive from AActor.  Actor iterators do not have the problem noted below, and will only return objects being used by the current game world instance.
 
 When creating an actor iterator, you need to give it a pointer to a **UWorld** instance.  Many UObject classes, such as **APlayerController**, provide a **GetWorld** method to help you.  If you are not sure, you can check the **ImplementsGetWorld** method on a UObject to see if it implements the GetWorld method.
@@ -426,7 +426,7 @@ UE4 uses the reflection system to implement a garbage collection system.  With g
 		GENERATED_BODY()
 	};
 
-In the garbage collector, there is this concept called the root set.  This root set is basically a list of objects that the collector knows about will never be garbage collected.  An object will not be garbage collected as long as there is a path of references from an object in the root set to the object in question.  If no such path to the root set exists for an object, it is called unreachable and will be collected (deleted) the next time the garbage collector is ran.  The engine runs the garbage collector at certain intervals.
+In the garbage collector, there is this concept called the root set.  This root set is basically a list of objects that the collector knows about will never be garbage collected.  An object will not be garbage collected as long as there is a path of references from an object in the root set to the object in question.  If no such path to the root set exists for an object, it is called *unreachable *and will be collected (deleted) the next time the garbage collector is ran.  The engine runs the garbage collector at certain intervals.
 
 What counts as a "reference"?  Any UObject pointer stored in a UPROPERTY.  Let us start with a simple example.
 
@@ -439,7 +439,7 @@ When we call the above function, we create a new UObject, but we do not store a
 
 ### Actors and Garbage collection
 
-Actors are not usually garbage collected.  Once spawned, you must manually call **Destroy()** on them.  They will not be deleted immediately, and instead will be cleaned up during the next garbage collection phase.
+Actors are not usually garbage collected.  Once spawned, you must manually call Destroy() on them.  They will not be deleted immediately, and instead will be cleaned up during the next garbage collection phase.
 
 This is a more common case, where you have actors with UObject properties.
 
diff --git a/Engine/Documentation/Source/Programming/ProgrammingGuide.INT.udn b/Engine/Documentation/Source/Programming/ProgrammingGuide.INT.udn
--- a/Engine/Documentation/Source/Programming/ProgrammingGuide.INT.udn
+++ b/Engine/Documentation/Source/Programming/ProgrammingGuide.INT.udn
@@ -2,7 +2,6 @@ Availability:Public
 Title:Programming Guide
 Crumbs: %ROOT%
 Description:Information for programmers developing with Unreal Engine.
-version: 4.9
 
 [VAR:Params]
 [PARAM:image]
diff --git a/Engine/Documentation/Source/Resources/ContentExamples/NavMesh/1_1/1_1.INT.udn b/Engine/Documentation/Source/Resources/ContentExamples/NavMesh/1_1/1_1.INT.udn
--- a/Engine/Documentation/Source/Resources/ContentExamples/NavMesh/1_1/1_1.INT.udn
+++ b/Engine/Documentation/Source/Resources/ContentExamples/NavMesh/1_1/1_1.INT.udn
@@ -6,7 +6,6 @@ Related: Engine/AI/BehaviorTrees/QuickStart
 Related: Gameplay/AI
 Related: Engine/AI/BehaviorTrees/NodeReference
 Related: Engine/Actors/Volumes
-Version: 4.9
 
 [VAR:Steps]
 [OBJECT:Navigation]
@@ -68,7 +67,7 @@ Version: 4.9
 ![](Resources\ContentExamples\NavMesh\NavMesh_banner.png)
 [/REGION]
 
-This section of the **NavMesh** content example shows how a **Pawn** can navigate through a set of obstacles. You can try adding and moving more blocks around to make more complicated paths. The navigation system will often list things as "simple" because it is not taking in much external data, like if an Actor is walking on ice or through tar, but can quite handily navigate in three dimensions, accounting for bridges and overhangs along the way, to any desired location.
+This section of the **NavMesh** content example shows how a **Pawn** can navigate through a set of obstacles. You can try adding and moving more blocks around to make more complicated paths. The navigation system will often list things as "simple" because it is not taking in much external data, like if an Actor is walking on ice or through tar, but can quite handily navigate in 3 dimensions, accounting for bridges and overhangs along the way, to any desired location.
 
 [REGION:lightbox]
 [![](NM_1_1_001.png)(w:300)](NM_1_1_001.png)
diff --git a/Engine/Documentation/Source/Resources/ContentExamples/NavMesh/1_2/1_2.INT.udn b/Engine/Documentation/Source/Resources/ContentExamples/NavMesh/1_2/1_2.INT.udn
--- a/Engine/Documentation/Source/Resources/ContentExamples/NavMesh/1_2/1_2.INT.udn
+++ b/Engine/Documentation/Source/Resources/ContentExamples/NavMesh/1_2/1_2.INT.udn
@@ -6,7 +6,6 @@ Related: Engine/AI/BehaviorTrees/QuickStart
 Related: Gameplay/AI
 Related: Engine/AI/BehaviorTrees/NodeReference
 Related: Engine/Actors/Volumes
-Version: 4.9
 
 [VAR:Steps]
 [OBJECT:Navigation]
diff --git a/Engine/Documentation/Source/Resources/ContentExamples/NavMesh/1_3/1_3.INT.udn b/Engine/Documentation/Source/Resources/ContentExamples/NavMesh/1_3/1_3.INT.udn
--- a/Engine/Documentation/Source/Resources/ContentExamples/NavMesh/1_3/1_3.INT.udn
+++ b/Engine/Documentation/Source/Resources/ContentExamples/NavMesh/1_3/1_3.INT.udn
@@ -6,7 +6,6 @@ Related: Engine/AI/BehaviorTrees/QuickStart
 Related: Gameplay/AI
 Related: Engine/AI/BehaviorTrees/NodeReference
 Related: Engine/Actors/Volumes
-Version: 4.9
 
 [VAR:Steps]
 [OBJECT:Navigation]
@@ -70,7 +69,6 @@ Version: 4.9
 
 The **NavMesh** automatically updates when Actors like **Static Meshes** are moved in the editor. In this example, you can move the block Static Mesh and position it so that it bridges the gap between the ramps. The NavMesh will automatically update and build across the newly-formed bridge, allowing the Pawn to navigate across it when you **Play** or **Simulate**.
 
-
 [REGION:lightbox]
 [![](NavMesh_1_3_001.png)(w:225)](NavMesh_1_3_001.png)
 [![](NavMesh_1_3_002.png)(w:225)](NavMesh_1_3_002.png)
diff --git a/Engine/Documentation/Source/Resources/ContentExamples/Networking/1_1/1_1.INT.udn b/Engine/Documentation/Source/Resources/ContentExamples/Networking/1_1/1_1.INT.udn
--- a/Engine/Documentation/Source/Resources/ContentExamples/Networking/1_1/1_1.INT.udn
+++ b/Engine/Documentation/Source/Resources/ContentExamples/Networking/1_1/1_1.INT.udn
@@ -7,7 +7,6 @@ Related: Gameplay/Networking/Actors/Components
 Related: Gameplay/Networking/Actors
 Related: Gameplay/Networking/Actors/RPCs
 Related: Engine/UI/LevelEditor/InEditorTesting/NetworkedPlayInEditor
-Version: 4.9
 
 [VAR:Steps]
 [OBJECT:Navigation]
diff --git a/Engine/Documentation/Source/Resources/ContentExamples/Networking/1_2/1_2.INT.udn b/Engine/Documentation/Source/Resources/ContentExamples/Networking/1_2/1_2.INT.udn
--- a/Engine/Documentation/Source/Resources/ContentExamples/Networking/1_2/1_2.INT.udn
+++ b/Engine/Documentation/Source/Resources/ContentExamples/Networking/1_2/1_2.INT.udn
@@ -7,7 +7,6 @@ Related: Gameplay/Networking/Actors/Components
 Related: Gameplay/Networking/Actors
 Related: Gameplay/Networking/Actors/RPCs
 Related: Engine/UI/LevelEditor/InEditorTesting/NetworkedPlayInEditor
-Version: 4.9
 
 [VAR:Steps]
 [OBJECT:Navigation]
diff --git a/Engine/Documentation/Source/Resources/ContentExamples/Networking/1_3/1_3.INT.udn b/Engine/Documentation/Source/Resources/ContentExamples/Networking/1_3/1_3.INT.udn
--- a/Engine/Documentation/Source/Resources/ContentExamples/Networking/1_3/1_3.INT.udn
+++ b/Engine/Documentation/Source/Resources/ContentExamples/Networking/1_3/1_3.INT.udn
@@ -7,7 +7,6 @@ Related: Gameplay/Networking/Actors/Components
 Related: Gameplay/Networking/Actors
 Related: Gameplay/Networking/Actors/RPCs
 Related: Engine/UI/LevelEditor/InEditorTesting/NetworkedPlayInEditor
-Version: 4.9
 
 [VAR:Steps]
 [OBJECT:Navigation]
diff --git a/Engine/Documentation/Source/Resources/ContentExamples/Networking/1_4/1_4.INT.udn b/Engine/Documentation/Source/Resources/ContentExamples/Networking/1_4/1_4.INT.udn
--- a/Engine/Documentation/Source/Resources/ContentExamples/Networking/1_4/1_4.INT.udn
+++ b/Engine/Documentation/Source/Resources/ContentExamples/Networking/1_4/1_4.INT.udn
@@ -7,7 +7,6 @@ Related: Gameplay/Networking/Actors/Components
 Related: Gameplay/Networking/Actors
 Related: Gameplay/Networking/Actors/RPCs
 Related: Engine/UI/LevelEditor/InEditorTesting/NetworkedPlayInEditor
-Version: 4.9
 
 [VAR:Steps]
 [OBJECT:Navigation]
diff --git a/Engine/Documentation/Source/Resources/ContentExamples/Networking/1_5/1_5.INT.udn b/Engine/Documentation/Source/Resources/ContentExamples/Networking/1_5/1_5.INT.udn
--- a/Engine/Documentation/Source/Resources/ContentExamples/Networking/1_5/1_5.INT.udn
+++ b/Engine/Documentation/Source/Resources/ContentExamples/Networking/1_5/1_5.INT.udn
@@ -7,7 +7,6 @@ Related: Gameplay/Networking/Actors/Components
 Related: Gameplay/Networking/Actors
 Related: Gameplay/Networking/Actors/RPCs
 Related: Engine/UI/LevelEditor/InEditorTesting/NetworkedPlayInEditor
-Version: 4.9
 
 [VAR:Steps]
 [OBJECT:Navigation]
diff --git a/Engine/Documentation/Source/Resources/ContentExamples/Networking/2_1/2_1.INT.udn b/Engine/Documentation/Source/Resources/ContentExamples/Networking/2_1/2_1.INT.udn
--- a/Engine/Documentation/Source/Resources/ContentExamples/Networking/2_1/2_1.INT.udn
+++ b/Engine/Documentation/Source/Resources/ContentExamples/Networking/2_1/2_1.INT.udn
@@ -7,7 +7,6 @@ Related: Gameplay/Networking/Actors/Components
 Related: Gameplay/Networking/Actors
 Related: Gameplay/Networking/Actors/RPCs
 Related: Engine/UI/LevelEditor/InEditorTesting/NetworkedPlayInEditor
-Version:4.9
 
 [VAR:Steps]
 [OBJECT:Navigation]
diff --git a/Engine/Documentation/Source/Resources/ContentExamples/Networking/2_2/2_2.INT.udn b/Engine/Documentation/Source/Resources/ContentExamples/Networking/2_2/2_2.INT.udn
--- a/Engine/Documentation/Source/Resources/ContentExamples/Networking/2_2/2_2.INT.udn
+++ b/Engine/Documentation/Source/Resources/ContentExamples/Networking/2_2/2_2.INT.udn
@@ -7,7 +7,6 @@ Related: Gameplay/Networking/Actors/Components
 Related: Gameplay/Networking/Actors
 Related: Gameplay/Networking/Actors/RPCs
 Related: Engine/UI/LevelEditor/InEditorTesting/NetworkedPlayInEditor
-Version: 4.9
 
 [VAR:Steps]
 [OBJECT:Navigation]
diff --git a/Engine/Documentation/Source/Resources/ContentExamples/Networking/2_3/2_3.INT.udn b/Engine/Documentation/Source/Resources/ContentExamples/Networking/2_3/2_3.INT.udn
--- a/Engine/Documentation/Source/Resources/ContentExamples/Networking/2_3/2_3.INT.udn
+++ b/Engine/Documentation/Source/Resources/ContentExamples/Networking/2_3/2_3.INT.udn
@@ -7,7 +7,6 @@ Related: Gameplay/Networking/Actors/Components
 Related: Gameplay/Networking/Actors
 Related: Gameplay/Networking/Actors/RPCs
 Related: Engine/UI/LevelEditor/InEditorTesting/NetworkedPlayInEditor
-Version: 4.9
 
 [VAR:Steps]
 [OBJECT:Navigation]
diff --git a/Engine/Documentation/Source/Resources/ContentExamples/Networking/2_4/2_4.INT.udn b/Engine/Documentation/Source/Resources/ContentExamples/Networking/2_4/2_4.INT.udn
--- a/Engine/Documentation/Source/Resources/ContentExamples/Networking/2_4/2_4.INT.udn
+++ b/Engine/Documentation/Source/Resources/ContentExamples/Networking/2_4/2_4.INT.udn
@@ -7,7 +7,6 @@ Related: Gameplay/Networking/Actors/Components
 Related: Gameplay/Networking/Actors
 Related: Gameplay/Networking/Actors/RPCs
 Related: Engine/UI/LevelEditor/InEditorTesting/NetworkedPlayInEditor
-Version: 4.9
 
 [VAR:Steps]
 [OBJECT:Navigation]
diff --git a/Engine/Documentation/Source/Resources/ContentExamples/Networking/NetworkFeatures.INT.udn b/Engine/Documentation/Source/Resources/ContentExamples/Networking/NetworkFeatures.INT.udn
--- a/Engine/Documentation/Source/Resources/ContentExamples/Networking/NetworkFeatures.INT.udn
+++ b/Engine/Documentation/Source/Resources/ContentExamples/Networking/NetworkFeatures.INT.udn
@@ -7,7 +7,6 @@ Related: Gameplay/Networking/Actors/Components
 Related: Gameplay/Networking/Actors
 Related: Gameplay/Networking/Actors/RPCs
 Related: Engine/UI/LevelEditor/InEditorTesting/NetworkedPlayInEditor
-Version: 4.9
 
 [REGION:banner]
 ![](NetworkFeatures_Header.png)(convert:false)
@@ -48,7 +47,7 @@ Version: 4.9
 
 
 
-The **Network_Features** map inside the Content Examples project illustrates how to implement network functionality entirely through Blueprints. It provides key examples in **Replication** including the **Replicates** boolean which can be set on Actors, detecting **Network Authority** and executing different logic based on that authority, usage of **Variable Replication** and how it is used during gameplay, variable replication using **RepNotify** to execute logic based on updates made to a variable that is replicated, as well as the use of **Function Replication** to handle non-critical gameplay elements or "one-off" events.
+The **Network_Features** map inside the **Content Examples** project illustrates how to implement network functionality entirely through Blueprints. It provides key examples in **Replication** including the **Replicates** boolean which can be set on Actors, detecting **Network Authority** and executing different logic based on that authority, usage of **Variable Replication** and how it is used during gameplay, variable replication using **RepNotify** to execute logic based on updates made to a variable that is replicated, as well as the use of **Function Replication** to handle non-critical gameplay elements or "one-off" events.
 
 The other section in the map provides insight on **Relevancy** and how to properly handle communicating information to and from the Server and Clients. Often times, it is not critical that every bit of information be passed from the server to every client (e.g. a player that is on one side of the map does not need to know about another player opening a chest on the other side of the map). There are situations when this would become relevant to that player however (e.g. if that player moved from the other side of the map to the chest that was opened), the Content Examples in this map will show how to handle situations like this as well as how to handle a player that joins an in progress session. 
 
@@ -77,10 +76,10 @@ The examples below provide possible solutions to the problem stated above and wh
 
 | Example | What is Demonstrated |
 | ------- | ------------ |
-| **[2.1 Network Relevancy - Not Replicated At All](Resources\ContentExamples\Networking\2_1)** | An example of relevancy where No Replication is used (e.g. Server opens a chest, Client does not see the chest as opened).  |
-| **[2.2 Network Relevancy - Function Replication Solution Attempt](Resources\ContentExamples\Networking\2_2)** | An example of relevancy where Function Replication is used (e.g. Server opens a chest, Client does not see the chest open and does not see it open when moving close to it). |
-| **[2.3 Network Relevancy - Variable Replication Solution Attempt](Resources\ContentExamples\Networking\2_3)** | An example of relevancy where Variable Replication is used. e.g. Server opens a chest, Client does not see the chest open but sees it open when moving close to it. Particle effects also play, which is not desired for this example as the chest was previously opened by another player. |
-| **[2.4 Network Relevancy - Combination Solution](Resources\ContentExamples\Networking\2_4)** | An example of relevancy where a combination of Function Replication is used to play the effects associated with opening the chest and a Replicated Variable is used to preserve the state of the chest and if it has already been opened or not. |
+| **[2.1 Network Relevancy - Not Replicated At All](Resources\ContentExamples\Networking\2_1)** | An example of relevancy where **No Replication** is used (e.g. Server opens a chest, Client does not see the chest as opened).  |
+| **[2.2 Network Relevancy - Function Replication Solution Attempt](Resources\ContentExamples\Networking\2_2)** | An example of relevancy where **Function Replication** is used (e.g. Server opens a chest, Client does not see the chest open and does not see it open when moving close to it). |
+| **[2.3 Network Relevancy - Variable Replication Solution Attempt](Resources\ContentExamples\Networking\2_3)** | An example of relevancy where **Variable Replication** is used. e.g. Server opens a chest, Client does not see the chest open but sees it open when moving close to it. Particle effects also play, which is not desired for this example as the chest was previously opened by another player. |
+| **[2.4 Network Relevancy - Combination Solution](Resources\ContentExamples\Networking\2_4)** | An example of relevancy where a combination of **Function Replication** is used to play the effects associated with opening the chest and a **Replicated Variable** is used to preserve the state of the chest and if it has already been opened or not. |
 
 
 
diff --git a/Engine/Documentation/Source/Resources/ContentExamples/Paper2D/Paper2D.INT.udn b/Engine/Documentation/Source/Resources/ContentExamples/Paper2D/Paper2D.INT.udn
--- a/Engine/Documentation/Source/Resources/ContentExamples/Paper2D/Paper2D.INT.udn
+++ b/Engine/Documentation/Source/Resources/ContentExamples/Paper2D/Paper2D.INT.udn
@@ -8,7 +8,6 @@ Related: Engine/Content/Tools/RenderToTextureTools/5
 Related: Engine/Content/Tools/RenderToTextureTools/3
 Related: Engine/Rendering/ParticleSystems/Reference/Modules/SubUV
 Related: Engine/Rendering/Materials/Functions/Reference/Misc
-Version: 4.9
 
 [REGION:banner]
 ![](paper2d_header.png)(convert:false)
@@ -45,10 +44,12 @@ Version: 4.9
 [/OBJECT]
 [/VAR]
 
+[TOC(start:2 end:2)]
 
 
+## Overview
 
-The **Paper 2D** map in Content Examples cover several topics related to 2D game creation including **Sprite** creation and editing via the **Sprite Editor**, **Flipbooks** which are in essence a method in which you can create animated 2D assets using Sprites and keyframing, how to control Flipbooks with Blueprints, adding physics to 2D assets as well as how to create the classic low resolution pixel look for your Sprites. 
+The **Paper 2D Content Examples** cover several topics related to 2D game creation including **Sprite** creation and editing via the **Sprite Editor**, **Flipbooks** which are in essence a method in which you can create animated 2D assets using **Sprites** and keyframing, how to control **Flipbooks** with **Blueprints**, adding physics to 2D assets as well as how to create the classic low resolution pixel look for your **Sprites**. 
 
 ## Examples
 
@@ -56,10 +57,10 @@ The **Paper 2D** map in Content Examples cover several topics related to 2D game
 
 | Example | What is Demonstrated |
 | ------- | ------------ |
-| **1.1 Sprites** | A quick and easy way to draw 2D images in UE4, this example shows how to create Sprites in the Content Browser (e.g. a simple static 2D sprite and an animated 2D sprite is shown). |
-| **1.2 Sprite Editor** | The Paper 2D Sprite Editor enables setup and editing of individual Sprite assets (e.g. custom collision is shown for 2D assets which can be edited in the Sprite Editor). |
-| **1.3 Flipbooks** | Flipbooks consist of a series of keyframes, each of which contains a Sprite to be displayed and a duration (in frames) to display it (e.g. two examples of 2D characters are shown running using different frames for their animations). |
-| **1.4 Flipbooks (Blueprint Control)** | Shows how Blueprint Functions can be used to control animations in a Flipbook or to switch Flipbooks (e.g. a 2D character performs different animations when pressing different buttons). |
+| **1.1 Sprites** | A quick and easy way to draw 2D images in UE4, this example shows how to create **Sprites** in the **Content Browser** (e.g. a simple static 2D sprite and an animated 2D sprite is shown). |
+| **1.2 Sprite Editor** | The **Paper 2D Sprite Editor** enables setup and editing of individual **Sprite** assets (e.g. custom collision is shown for 2D assets which can be edited in the **Sprite Editor**). |
+| **1.3 Flipbooks** | **Flipbooks** consist of a series of keyframes, each of which contains a **Sprite** to be displayed and a duration (in frames) to display it (e.g. two examples of 2D characters are shown running using different frames for their animations). |
+| **1.4 Flipbooks (Blueprint Control)** | Shows how **Blueprint Functions** can be used to control animations in a **Flipbook** or to switch **Flipbooks** (e.g. a 2D character performs different animations when pressing different buttons). |
 | **1.5 Physics (3D Physics Limited to a Plane)** | This example illustrates how physics can be applied to 2D assets by using locked 3D physics (e.g. several coins are dropped and physics are applied, but are restricted to the Y axis). |
-| **1.6 Classic Pixel Sprites** | How to achieve the classic pixelated look for Sprites by changing texture settings (e.g. imported default texture is shown alongside an edited version which uses a pixel look). |
+| **1.6 Classic Pixel Sprites** | How to achieve the classic pixelated look for **Sprites** by changing texture settings (e.g. imported default texture is shown alongside an edited version which uses a pixel look). |
 
diff --git a/Engine/Documentation/Source/Resources/ContentExamples/Physics/1_1/1_1.INT.udn b/Engine/Documentation/Source/Resources/ContentExamples/Physics/1_1/1_1.INT.udn
--- a/Engine/Documentation/Source/Resources/ContentExamples/Physics/1_1/1_1.INT.udn
+++ b/Engine/Documentation/Source/Resources/ContentExamples/Physics/1_1/1_1.INT.udn
@@ -8,7 +8,6 @@ Related: Engine/Physics/Constraints/ConstraintsReference
 Related: Engine/Physics/Constraints/ConstraintsUserGuide
 Related: Engine/Physics/Collision
 Related: Engine/Physics/PhysicalMaterials/Reference
-Version: 4.9
 
 [VAR:Steps]
 [OBJECT:Navigation]
diff --git a/Engine/Documentation/Source/Resources/ContentExamples/Physics/1_2/1_2.INT.udn b/Engine/Documentation/Source/Resources/ContentExamples/Physics/1_2/1_2.INT.udn
--- a/Engine/Documentation/Source/Resources/ContentExamples/Physics/1_2/1_2.INT.udn
+++ b/Engine/Documentation/Source/Resources/ContentExamples/Physics/1_2/1_2.INT.udn
@@ -8,7 +8,6 @@ Related: Engine/Physics/Constraints/ConstraintsReference
 Related: Engine/Physics/Constraints/ConstraintsUserGuide
 Related: Engine/Physics/Collision
 Related: Engine/Physics/PhysicalMaterials/Reference
-Version: 4.9
 
 [VAR:Steps]
 [OBJECT:Navigation]
diff --git a/Engine/Documentation/Source/Resources/ContentExamples/Physics/1_3/1_3.INT.udn b/Engine/Documentation/Source/Resources/ContentExamples/Physics/1_3/1_3.INT.udn
--- a/Engine/Documentation/Source/Resources/ContentExamples/Physics/1_3/1_3.INT.udn
+++ b/Engine/Documentation/Source/Resources/ContentExamples/Physics/1_3/1_3.INT.udn
@@ -8,7 +8,6 @@ Related: Engine/Physics/Constraints/ConstraintsReference
 Related: Engine/Physics/Constraints/ConstraintsUserGuide
 Related: Engine/Physics/Collision
 Related: Engine/Physics/PhysicalMaterials/Reference
-Version: 4.9
 
 [VAR:Steps]
 [OBJECT:Navigation]
@@ -72,11 +71,11 @@ Version: 4.9
 [/REGION]
 
 
-Using a **Radial Force Actor**, you can apply a **Constant Force** (listed as _Force Strength_ in the editor) or an **Impulse**. Impulses need to be fired from a **Blueprint**, and apply their entire force in only one frame, as compared to every frame from a Constant Force. In the content example, only a Constant Force is being applied.
+Using a **Radial Force Actor**, you can apply a **Constant Force** (listed as _Force Strength_ in the editor) or an **Impulse**. Impulses need to be fired from a **Blueprint**, and apply their entire force in only one frame, as compared to every frame from a Constant Force. In the content example, only a Constant Forceis being applied.
 
 ![](physics_1_3_002.png)(w:450) ![](physics_1_3_003.png)(w:450)
 
-_This force is applied to everything that simulates physics._
+This force is applied to everything that simulates physics.
 
 
 
diff --git a/Engine/Documentation/Source/Resources/ContentExamples/Physics/1_4/1_4.INT.udn b/Engine/Documentation/Source/Resources/ContentExamples/Physics/1_4/1_4.INT.udn
--- a/Engine/Documentation/Source/Resources/ContentExamples/Physics/1_4/1_4.INT.udn
+++ b/Engine/Documentation/Source/Resources/ContentExamples/Physics/1_4/1_4.INT.udn
@@ -8,7 +8,6 @@ Related: Engine/Physics/Constraints/ConstraintsReference
 Related: Engine/Physics/Constraints/ConstraintsUserGuide
 Related: Engine/Physics/Collision
 Related: Engine/Physics/PhysicalMaterials/Reference
-Version: 4.9
 
 [VAR:Steps]
 [OBJECT:Navigation]
diff --git a/Engine/Documentation/Source/Resources/ContentExamples/Physics/1_5/1_5.INT.udn b/Engine/Documentation/Source/Resources/ContentExamples/Physics/1_5/1_5.INT.udn
--- a/Engine/Documentation/Source/Resources/ContentExamples/Physics/1_5/1_5.INT.udn
+++ b/Engine/Documentation/Source/Resources/ContentExamples/Physics/1_5/1_5.INT.udn
@@ -8,7 +8,6 @@ Related: Engine/Physics/Constraints/ConstraintsReference
 Related: Engine/Physics/Constraints/ConstraintsUserGuide
 Related: Engine/Physics/Collision
 Related: Engine/Physics/PhysicalMaterials/Reference
-Version: 4.9
 
 [VAR:Steps]
 [OBJECT:Navigation]
@@ -74,7 +73,7 @@ Version: 4.9
 [/REGION]
 
 
-A **Physics Constraint** is basically a type of joint. It allows you to connect two Actors together (presumably one physically simulating), and also apply limits or forces. Unreal Engine 4 
+A constraint is basically a type of joint. It allows you to connect two Actors together (presumably one physically simulating), and also apply limits or forces. Unreal Engine 4 
 has a very flexible and data-driven constraint system that allows designers to make many different types of joints simply by changing some options. The engine has a few default 
 joint types (ball-and-socket, hinge, prismatic) but they only differ in their settings. You should feel free to start with any joint type and experiment with making your own 
 configurations.
diff --git a/Engine/Documentation/Source/Resources/ContentExamples/Physics/1_6/1_6.INT.udn b/Engine/Documentation/Source/Resources/ContentExamples/Physics/1_6/1_6.INT.udn
--- a/Engine/Documentation/Source/Resources/ContentExamples/Physics/1_6/1_6.INT.udn
+++ b/Engine/Documentation/Source/Resources/ContentExamples/Physics/1_6/1_6.INT.udn
@@ -8,7 +8,6 @@ Related: Engine/Physics/Constraints/ConstraintsReference
 Related: Engine/Physics/Constraints/ConstraintsUserGuide
 Related: Engine/Physics/Collision
 Related: Engine/Physics/PhysicalMaterials/Reference
-Version: 4.9
 
 [VAR:Steps]
 [OBJECT:Navigation]
@@ -72,7 +71,7 @@ Version: 4.9
 [/REGION]
 
 
-A Physics Constraint with an **Angular Motor** will essentially spin, but the detail is in **how** it spins. An angular motor **Driven by Velocity** will apply a constant force to essentially "spin," where as a **Driven by Position** will attempt to obtain an angle based on the size of the spring and dampening applied to said spring.
+A constraint with an angular motor will essentially spin, but the detail is in **how** it spins. An angular motor **Driven by Velocity** will apply a constant force to essentially "spin," where as a **Driven by Position** will attempt to obtain an angle based on the size of the spring and dampening applied to said spring.
 
 <div align="center">
 ![](physics_1_6_002.png)(h:300) ![](physics_1_6_003.png)(h:300)
diff --git a/Engine/Documentation/Source/Resources/ContentExamples/Physics/1_7/1_7.INT.udn b/Engine/Documentation/Source/Resources/ContentExamples/Physics/1_7/1_7.INT.udn
--- a/Engine/Documentation/Source/Resources/ContentExamples/Physics/1_7/1_7.INT.udn
+++ b/Engine/Documentation/Source/Resources/ContentExamples/Physics/1_7/1_7.INT.udn
@@ -8,7 +8,6 @@ Related: Engine/Physics/Constraints/ConstraintsReference
 Related: Engine/Physics/Constraints/ConstraintsUserGuide
 Related: Engine/Physics/Collision
 Related: Engine/Physics/PhysicalMaterials/Reference
-Version: 4.9
 
 [VAR:Steps]
 [OBJECT:Navigation]
diff --git a/Engine/Documentation/Source/Resources/ContentExamples/Physics/1_8/1_8.INT.udn b/Engine/Documentation/Source/Resources/ContentExamples/Physics/1_8/1_8.INT.udn
--- a/Engine/Documentation/Source/Resources/ContentExamples/Physics/1_8/1_8.INT.udn
+++ b/Engine/Documentation/Source/Resources/ContentExamples/Physics/1_8/1_8.INT.udn
@@ -8,7 +8,6 @@ Related: Engine/Physics/Constraints/ConstraintsReference
 Related: Engine/Physics/Constraints/ConstraintsUserGuide
 Related: Engine/Physics/Collision
 Related: Engine/Physics/PhysicalMaterials/Reference
-Version: 4.9
 
 [VAR:Steps]
 [OBJECT:Navigation]
@@ -67,16 +66,15 @@ Version: 4.9
 [/VAR]
 
 
-
 <div align="center">
-	![](physics_1_8_002.png)(w:960) <!---![](physics_1_8_001.png)(w:450)--->
+	![](physics_1_8_002.png)(w:450) ![](physics_1_8_001.png)(w:450)
 </div>
 
 
 A constraint can be made to break once a certain amount of angular or linear force has been applied to the constraint through its constrained objects. An **Angular Force** would be considered like snapping a twig or a torsion force. A **Linear Force** is a "pulling" force along the "length" of the constraint, such as applying too much weight to a crane's lift line.
 
-<div align="center">
-	![](physics_1_8_003.png)(w:500)
-</div>
+[REGION:fullwidth]
+	![](physics_1_8_003.png)
+[/REGION]
 
 
diff --git a/Engine/Documentation/Source/Resources/ContentExamples/Physics/Physics.INT.udn b/Engine/Documentation/Source/Resources/ContentExamples/Physics/Physics.INT.udn
--- a/Engine/Documentation/Source/Resources/ContentExamples/Physics/Physics.INT.udn
+++ b/Engine/Documentation/Source/Resources/ContentExamples/Physics/Physics.INT.udn
@@ -8,7 +8,6 @@ Related: Engine/Physics/Constraints/ConstraintsReference
 Related: Engine/Physics/Constraints/ConstraintsUserGuide
 Related: Engine/Physics/Collision
 Related: Engine/Physics/PhysicalMaterials/Reference
-Version: 4.9
 
 [REGION:banner]
 ![](Physics_Header.png)(convert:false)
@@ -45,9 +44,10 @@ Version: 4.9
 [/OBJECT]
 [/VAR]
 
+[TOC(start:2 end:2)]
 
 
-Inside the **Physics** map, you will find several examples of how physics can be used in Unreal Engine 4. Creating Physics Actors out of Static Meshes so that they respond to player interaction, applying a Physics Asset to a Skeletal Mesh so that it responds to physics, and several examples of how Force, Impulse, Motors, and Physical Constraints are on display in this map with details of how they can aid in your level design or gameplay scenarios. 
+Inside the **Physics** map, you will find several examples of how physics can be used in Unreal Engine 4. Creating **Physics Actors** out of **Static Meshes** so that they respond to player interaction, applying a **Physics Asset** to a **Skeletal Mesh** so that it responds to physics, and several examples of how **Force**, **Impulse**, **Motors**, and **Physical Constraints** are on display in this map with details of how they can aid in your level design or gameplay scenarios. 
 
 
 
diff --git a/Engine/Documentation/Source/Resources/ContentExamples/PivotPainter/PivotPainter.INT.udn b/Engine/Documentation/Source/Resources/ContentExamples/PivotPainter/PivotPainter.INT.udn
--- a/Engine/Documentation/Source/Resources/ContentExamples/PivotPainter/PivotPainter.INT.udn
+++ b/Engine/Documentation/Source/Resources/ContentExamples/PivotPainter/PivotPainter.INT.udn
@@ -6,7 +6,6 @@ Related: Engine/Rendering/Materials
 Related: Engine/Rendering/Materials/ExpressionReference
 Related: Engine/Rendering/Materials/Functions/Reference/PivotPainter
 Related: Engine/Rendering/Materials/Functions/Creating
-Version: 4.9
 
 [REGION:banner]
 ![](PivotPainter_Header.png)(convert:false)
@@ -43,11 +42,13 @@ Version: 4.9
 [/OBJECT]
 [/VAR]
 
+[TOC(start:2 end:2)]
 
 
-The [Pivot Painter Tool](Engine\Content\Tools\PivotPainter) is a MAXScript that stores model pivot and rotation information in the model's vertex data. That information can then be referenced inside of Unreal Engine 4's shader system to create interactive effects. 
 
-The **PivotPainter** map inside the Content Examples project provides several examples which use **Pivot Painter Material Functions** to produce models that achieve this effect and demonstrates how the models can be manipulated through several forces such as wind or player interaction. 
+The [**Pivot Painter Tool**](Engine\Content\Tools\PivotPainter) is a MAXScript that stores model pivot and rotation information in the model's vertex data. That information can then be referenced inside of Unreal Engine 4's shader system to create interactive effects. 
+
+The **PivotPainter** map inside the **Content Examples** project provides several examples which use **Pivot Painter Material Functions** to produce models that achieve this effect and demonstrates how the models can be manipulated through several forces such as wind or player interaction. 
 
 
  
diff --git a/Engine/Documentation/Source/Resources/ContentExamples/PostProcessing/1_1/1_1.INT.udn b/Engine/Documentation/Source/Resources/ContentExamples/PostProcessing/1_1/1_1.INT.udn
--- a/Engine/Documentation/Source/Resources/ContentExamples/PostProcessing/1_1/1_1.INT.udn
+++ b/Engine/Documentation/Source/Resources/ContentExamples/PostProcessing/1_1/1_1.INT.udn
@@ -5,7 +5,6 @@ Description:An overview of the Post Processing example level, example 1.1: Defau
 Related: Engine/Rendering/Overview
 Related: Engine/Rendering/PostProcessEffects
 Related: Engine/Rendering/PostProcessEffects/PostProcessMaterials
-Version: 4.9
 
 [VAR:Steps]
 [OBJECT:Navigation]
diff --git a/Engine/Documentation/Source/Resources/ContentExamples/PostProcessing/1_10/1_10.INT.udn b/Engine/Documentation/Source/Resources/ContentExamples/PostProcessing/1_10/1_10.INT.udn
--- a/Engine/Documentation/Source/Resources/ContentExamples/PostProcessing/1_10/1_10.INT.udn
+++ b/Engine/Documentation/Source/Resources/ContentExamples/PostProcessing/1_10/1_10.INT.udn
@@ -5,7 +5,6 @@ Description:An overview of the Post Processing example level, example 1.10: Glob
 Related: Engine/Rendering/Overview
 Related: Engine/Rendering/PostProcessEffects
 Related: Engine/Rendering/PostProcessEffects/PostProcessMaterials
-Version: 4.9
 
 [VAR:Steps]
 [OBJECT:Navigation]
diff --git a/Engine/Documentation/Source/Resources/ContentExamples/PostProcessing/1_11/1_11.INT.udn b/Engine/Documentation/Source/Resources/ContentExamples/PostProcessing/1_11/1_11.INT.udn
--- a/Engine/Documentation/Source/Resources/ContentExamples/PostProcessing/1_11/1_11.INT.udn
+++ b/Engine/Documentation/Source/Resources/ContentExamples/PostProcessing/1_11/1_11.INT.udn
@@ -6,7 +6,6 @@ Related: Engine/Rendering/PostProcessEffects/DepthOfField
 Related: Engine/Rendering/Overview
 Related: Engine/Rendering/PostProcessEffects
 Related: Engine/Rendering/PostProcessEffects/PostProcessMaterials
-Version: 4.9
 
 [VAR:Steps]
 [OBJECT:Navigation]
diff --git a/Engine/Documentation/Source/Resources/ContentExamples/PostProcessing/1_12/1_12.INT.udn b/Engine/Documentation/Source/Resources/ContentExamples/PostProcessing/1_12/1_12.INT.udn
--- a/Engine/Documentation/Source/Resources/ContentExamples/PostProcessing/1_12/1_12.INT.udn
+++ b/Engine/Documentation/Source/Resources/ContentExamples/PostProcessing/1_12/1_12.INT.udn
@@ -5,7 +5,6 @@ Description:An overview of the Post Processing example level, example 1.12: Moti
 Related: Engine/Rendering/Overview
 Related: Engine/Rendering/PostProcessEffects
 Related: Engine/Rendering/PostProcessEffects/PostProcessMaterials
-Version: 4.9
 
 [VAR:Steps]
 [OBJECT:Navigation]
diff --git a/Engine/Documentation/Source/Resources/ContentExamples/PostProcessing/1_13/1_13.INT.udn b/Engine/Documentation/Source/Resources/ContentExamples/PostProcessing/1_13/1_13.INT.udn
--- a/Engine/Documentation/Source/Resources/ContentExamples/PostProcessing/1_13/1_13.INT.udn
+++ b/Engine/Documentation/Source/Resources/ContentExamples/PostProcessing/1_13/1_13.INT.udn
@@ -5,7 +5,6 @@ Description:An overview of the Post Processing example level, example 1.13: Scre
 Related: Engine/Rendering/Overview
 Related: Engine/Rendering/PostProcessEffects
 Related: Engine/Rendering/PostProcessEffects/PostProcessMaterials
-Version: 4.9
 
 [VAR:Steps]
 [OBJECT:Navigation]
@@ -64,24 +63,11 @@ Version: 4.9
 [/OBJECT]
 [/VAR]
 
-<div align="center">
-[OBJECT:ComparisonSlider]
-	[PARAM:before]
-	![Screen Percentage - 100](1_13_01.png)(w:600)
-	[/PARAM]
-	[PARAM:after]
-	![Screen Percentage - 25](1_13_02.png)(w:600)
-	[/PARAM]
-[/OBJECT]
-</div>
-
-<!---
 <div align="center">
 [REGION:lightbox]
 [![](1_13_01.png)(w:440)](1_13_01.png) [![](1_13_02.png)(w:440)](1_13_02.png)
 [/REGION:lightbox]
 </div>
---->
 
 **Screen Percentage** renders the scene at a lower resolution and then scales it up. The value represents a percentage of the overall resolution of the scene.
 
diff --git a/Engine/Documentation/Source/Resources/ContentExamples/PostProcessing/1_14/1_14.INT.udn b/Engine/Documentation/Source/Resources/ContentExamples/PostProcessing/1_14/1_14.INT.udn
--- a/Engine/Documentation/Source/Resources/ContentExamples/PostProcessing/1_14/1_14.INT.udn
+++ b/Engine/Documentation/Source/Resources/ContentExamples/PostProcessing/1_14/1_14.INT.udn
@@ -5,7 +5,6 @@ Description:An overview of the Post Processing example level, example 1.14: AA M
 Related: Engine/Rendering/Overview
 Related: Engine/Rendering/PostProcessEffects
 Related: Engine/Rendering/PostProcessEffects/PostProcessMaterials
-Version: 4.9
 
 [VAR:Steps]
 [OBJECT:Navigation]
@@ -64,26 +63,9 @@ Version: 4.9
 [/OBJECT]
 [/VAR]
 
-
-[OBJECT:ComparisonSlider]
-	[PARAM:before]
-	![Anti-Aliasing: None](1_14_01.png)(w:960)
-	[/PARAM]
-	[PARAM:after]
-	![Anti-Aliasing: Temporal](1_14_03.png)(w:960)
-	[/PARAM]
-[/OBJECT]
-
-
-[Region:lightbox]
-[Region:imagetable]
-| [![](1_14_01.png)(w:302)](1_14_01.png) | [![](1_14_02.png)(w:302)](1_14_02.png) | [![](1_14_03.png)(w:302)](1_14_03.png) |
-| ------ | ------ | ------ |
-| None | FXAA | TXAA |
-[/Region]
-[/Region]
-
-
+[REGION:lightbox]
+[![](1_14_01.png)(w:302)](1_14_01.png) [![](1_14_02.png)(w:302)](1_14_02.png) [![](1_14_03.png)(w:302)](1_14_03.png)
+[/REGION:lightbox]
 
 Anti-Aliasing improves image quality by removing the stepping effect from the edges that appear angled on screen. **Temporal AA** is the default and gives the best result. Although it is on by default, it only shows in the editor when Real-time is activated in viewport or while in Game Mode **(G-key)**. 
 
diff --git a/Engine/Documentation/Source/Resources/ContentExamples/PostProcessing/1_15/1_15.INT.udn b/Engine/Documentation/Source/Resources/ContentExamples/PostProcessing/1_15/1_15.INT.udn
--- a/Engine/Documentation/Source/Resources/ContentExamples/PostProcessing/1_15/1_15.INT.udn
+++ b/Engine/Documentation/Source/Resources/ContentExamples/PostProcessing/1_15/1_15.INT.udn
@@ -5,7 +5,6 @@ Description:An overview of the Post Processing example level, example 1.15: Blur
 Related: Engine/Rendering/Overview
 Related: Engine/Rendering/PostProcessEffects
 Related: Engine/Rendering/PostProcessEffects/PostProcessMaterials
-Version: 4.9
 
 [VAR:Steps]
 [OBJECT:Navigation]
diff --git a/Engine/Documentation/Source/Resources/ContentExamples/PostProcessing/1_16/1_16.INT.udn b/Engine/Documentation/Source/Resources/ContentExamples/PostProcessing/1_16/1_16.INT.udn
--- a/Engine/Documentation/Source/Resources/ContentExamples/PostProcessing/1_16/1_16.INT.udn
+++ b/Engine/Documentation/Source/Resources/ContentExamples/PostProcessing/1_16/1_16.INT.udn
@@ -5,7 +5,6 @@ Description:An overview of the Post Processing example level, example 1.16: Blen
 Related: Engine/Rendering/Overview
 Related: Engine/Rendering/PostProcessEffects
 Related: Engine/Rendering/PostProcessEffects/PostProcessMaterials
-Version: 4.9
 
 [VAR:Steps]
 [OBJECT:Navigation]
diff --git a/Engine/Documentation/Source/Resources/ContentExamples/PostProcessing/1_17/1_17.INT.udn b/Engine/Documentation/Source/Resources/ContentExamples/PostProcessing/1_17/1_17.INT.udn
--- a/Engine/Documentation/Source/Resources/ContentExamples/PostProcessing/1_17/1_17.INT.udn
+++ b/Engine/Documentation/Source/Resources/ContentExamples/PostProcessing/1_17/1_17.INT.udn
@@ -5,7 +5,6 @@ Description:An overview of the Post Processing example level, example 1.17: Prio
 Related: Engine/Rendering/Overview
 Related: Engine/Rendering/PostProcessEffects
 Related: Engine/Rendering/PostProcessEffects/PostProcessMaterials
-Version: 4.9
 
 [VAR:Steps]
 [OBJECT:Navigation]
@@ -71,5 +70,5 @@ Version: 4.9
 [/REGION:lightbox]
 </div>
 
-**Volume Priority** allows you to control which Post Process Volume will override another if one is nested inside the other. If you are inside two volumes, the volume with the higher priority will be used and lower priority volumes will be ignored.
+Volume Priority allows you to control which Post Process Volume will override another if one is nested inside the other. If you are inside two volumes, the volume with the higher priority will be used and lower priority volumes will be ignored.
 
diff --git a/Engine/Documentation/Source/Resources/ContentExamples/PostProcessing/1_18/1_18.INT.udn b/Engine/Documentation/Source/Resources/ContentExamples/PostProcessing/1_18/1_18.INT.udn
--- a/Engine/Documentation/Source/Resources/ContentExamples/PostProcessing/1_18/1_18.INT.udn
+++ b/Engine/Documentation/Source/Resources/ContentExamples/PostProcessing/1_18/1_18.INT.udn
@@ -5,7 +5,6 @@ Description:An overview of the Post Processing example level, example 1.18: Blen
 Related: Engine/Rendering/Overview
 Related: Engine/Rendering/PostProcessEffects
 Related: Engine/Rendering/PostProcessEffects/PostProcessMaterials
-Version: 4.9
 
 [VAR:Steps]
 [OBJECT:Navigation]
@@ -72,5 +71,5 @@ Version: 4.9
 [![](1_18_01.png)(w:302)](1_18_01.png) [![](1_18_02.png)(w:302)](1_18_02.png) [![](1_18_03.png)(w:302)](1_18_03.png)
 [/REGION:lightbox]
 
-**Blend Radius** and **Weight**  is used to change the transition from one PostProcessVolume to another over a distance. Volumes can just take control of the post process pipeline or, using these settings, it can smoothing transition between two volumes over the distance set by the radius.
+This is used to change the transition from one PostProcessVolume to another over a distance. Volumes can just take control of the post process pipeline or, using these settings, it can smoothing transition between 2 volumes over the distance set by the radius.
 
diff --git a/Engine/Documentation/Source/Resources/ContentExamples/PostProcessing/1_19/1_19.INT.udn b/Engine/Documentation/Source/Resources/ContentExamples/PostProcessing/1_19/1_19.INT.udn
--- a/Engine/Documentation/Source/Resources/ContentExamples/PostProcessing/1_19/1_19.INT.udn
+++ b/Engine/Documentation/Source/Resources/ContentExamples/PostProcessing/1_19/1_19.INT.udn
@@ -5,7 +5,6 @@ Description:An overview of the Post Processing example level, example 1.19: Unbo
 Related: Engine/Rendering/Overview
 Related: Engine/Rendering/PostProcessEffects
 Related: Engine/Rendering/PostProcessEffects/PostProcessMaterials
-Version: 4.9
 
 [VAR:Steps]
 [OBJECT:Navigation]
diff --git a/Engine/Documentation/Source/Resources/ContentExamples/PostProcessing/1_2/1_2.INT.udn b/Engine/Documentation/Source/Resources/ContentExamples/PostProcessing/1_2/1_2.INT.udn
--- a/Engine/Documentation/Source/Resources/ContentExamples/PostProcessing/1_2/1_2.INT.udn
+++ b/Engine/Documentation/Source/Resources/ContentExamples/PostProcessing/1_2/1_2.INT.udn
@@ -5,7 +5,6 @@ Description:An overview of the Post Processing example level, example 1.2: Film
 Related: Engine/Rendering/Overview
 Related: Engine/Rendering/PostProcessEffects
 Related: Engine/Rendering/PostProcessEffects/PostProcessMaterials
-Version: 4.9
 
 [VAR:Steps]
 [OBJECT:Navigation]
@@ -69,5 +68,5 @@ Version: 4.9
 [![](1_2_01.png)(w:302)](1_2_01.png) [![](1_2_02.png)(w:302)](1_2_02.png) [![](1_2_03.png)(w:302)](1_2_03.png)
 [/REGION:lightbox]
 
-The **Film** category contains various properties for adjusting the color of your scenes. Includes tinting shadows, color mixing channels, saturation, contrast, crushing shadows, crushing highlights, and dynamic range.
+The Film category contains various properties for adjusting the color of your scenes. Includes tinting shadows, color mixing channels, saturation, contrast, crushing shadows, crushing highlights, and dynamic range.
 
diff --git a/Engine/Documentation/Source/Resources/ContentExamples/PostProcessing/1_3/1_3.INT.udn b/Engine/Documentation/Source/Resources/ContentExamples/PostProcessing/1_3/1_3.INT.udn
--- a/Engine/Documentation/Source/Resources/ContentExamples/PostProcessing/1_3/1_3.INT.udn
+++ b/Engine/Documentation/Source/Resources/ContentExamples/PostProcessing/1_3/1_3.INT.udn
@@ -2,7 +2,6 @@ Availability: Public
 Title:1.3 - Scene Color
 Crumbs:%ROOT%, Resources, Resources/ContentExamples, Resources/ContentExamples/PostProcessing
 Description:An overview of the Post Processing example level, example 1.3: Scene Color
-Version: 4.9
 
 [VAR:Steps]
 [OBJECT:Navigation]
diff --git a/Engine/Documentation/Source/Resources/ContentExamples/PostProcessing/1_4/1_4.INT.udn b/Engine/Documentation/Source/Resources/ContentExamples/PostProcessing/1_4/1_4.INT.udn
--- a/Engine/Documentation/Source/Resources/ContentExamples/PostProcessing/1_4/1_4.INT.udn
+++ b/Engine/Documentation/Source/Resources/ContentExamples/PostProcessing/1_4/1_4.INT.udn
@@ -6,7 +6,6 @@ Related: Engine/Rendering/PostProcessEffects/Bloom
 Related: Engine/Rendering/Overview
 Related: Engine/Rendering/PostProcessEffects
 Related: Engine/Rendering/PostProcessEffects/PostProcessMaterials
-Version: 4.9
 
 [VAR:Steps]
 [OBJECT:Navigation]
diff --git a/Engine/Documentation/Source/Resources/ContentExamples/PostProcessing/1_5/1_5.INT.udn b/Engine/Documentation/Source/Resources/ContentExamples/PostProcessing/1_5/1_5.INT.udn
--- a/Engine/Documentation/Source/Resources/ContentExamples/PostProcessing/1_5/1_5.INT.udn
+++ b/Engine/Documentation/Source/Resources/ContentExamples/PostProcessing/1_5/1_5.INT.udn
@@ -6,7 +6,6 @@ Related: Engine/Rendering/LightingAndShadows/AmbientCubemap
 Related: Engine/Rendering/Overview
 Related: Engine/Rendering/PostProcessEffects
 Related: Engine/Rendering/PostProcessEffects/PostProcessMaterials
-Version: 4.9
 
 [VAR:Steps]
 [OBJECT:Navigation]
diff --git a/Engine/Documentation/Source/Resources/ContentExamples/PostProcessing/1_6/1_6.INT.udn b/Engine/Documentation/Source/Resources/ContentExamples/PostProcessing/1_6/1_6.INT.udn
--- a/Engine/Documentation/Source/Resources/ContentExamples/PostProcessing/1_6/1_6.INT.udn
+++ b/Engine/Documentation/Source/Resources/ContentExamples/PostProcessing/1_6/1_6.INT.udn
@@ -6,7 +6,6 @@ Related: Engine/Rendering/PostProcessEffects/AutomaticExposure
 Related: Engine/Rendering/Overview
 Related: Engine/Rendering/PostProcessEffects
 Related: Engine/Rendering/PostProcessEffects/PostProcessMaterials
-Version: 4.9
 
 [VAR:Steps]
 [OBJECT:Navigation]
diff --git a/Engine/Documentation/Source/Resources/ContentExamples/PostProcessing/1_7/1_7.INT.udn b/Engine/Documentation/Source/Resources/ContentExamples/PostProcessing/1_7/1_7.INT.udn
--- a/Engine/Documentation/Source/Resources/ContentExamples/PostProcessing/1_7/1_7.INT.udn
+++ b/Engine/Documentation/Source/Resources/ContentExamples/PostProcessing/1_7/1_7.INT.udn
@@ -6,7 +6,6 @@ Related: Engine/Rendering/PostProcessEffects/LensFlare
 Related: Engine/Rendering/Overview
 Related: Engine/Rendering/PostProcessEffects
 Related: Engine/Rendering/PostProcessEffects/PostProcessMaterials
-Version: 4.9
 
 [VAR:Steps]
 [OBJECT:Navigation]
diff --git a/Engine/Documentation/Source/Resources/ContentExamples/PostProcessing/1_8/1_8.INT.udn b/Engine/Documentation/Source/Resources/ContentExamples/PostProcessing/1_8/1_8.INT.udn
--- a/Engine/Documentation/Source/Resources/ContentExamples/PostProcessing/1_8/1_8.INT.udn
+++ b/Engine/Documentation/Source/Resources/ContentExamples/PostProcessing/1_8/1_8.INT.udn
@@ -6,7 +6,6 @@ Related: Engine/Rendering/LightingAndShadows/AmbientOcclusion
 Related: Engine/Rendering/Overview
 Related: Engine/Rendering/PostProcessEffects
 Related: Engine/Rendering/PostProcessEffects/PostProcessMaterials
-Version: 4.9
 
 [VAR:Steps]
 [OBJECT:Navigation]
diff --git a/Engine/Documentation/Source/Resources/ContentExamples/PostProcessing/1_9/1_9.INT.udn b/Engine/Documentation/Source/Resources/ContentExamples/PostProcessing/1_9/1_9.INT.udn
--- a/Engine/Documentation/Source/Resources/ContentExamples/PostProcessing/1_9/1_9.INT.udn
+++ b/Engine/Documentation/Source/Resources/ContentExamples/PostProcessing/1_9/1_9.INT.udn
@@ -6,7 +6,6 @@ Related: Engine/Rendering/PostProcessEffects/ScreenSpaceReflection
 Related: Engine/Rendering/Overview
 Related: Engine/Rendering/PostProcessEffects
 Related: Engine/Rendering/PostProcessEffects/PostProcessMaterials
-Version: 4.9
 
 [VAR:Steps]
 [OBJECT:Navigation]
diff --git a/Engine/Documentation/Source/Resources/ContentExamples/PostProcessing/PostProcessing.INT.udn b/Engine/Documentation/Source/Resources/ContentExamples/PostProcessing/PostProcessing.INT.udn
--- a/Engine/Documentation/Source/Resources/ContentExamples/PostProcessing/PostProcessing.INT.udn
+++ b/Engine/Documentation/Source/Resources/ContentExamples/PostProcessing/PostProcessing.INT.udn
@@ -6,7 +6,6 @@ Related: Engine/Rendering/Overview
 Related: Engine/Rendering/PostProcessEffects
 Related: Engine/Rendering/PostProcessEffects/PostProcessMaterials
 Related: Resources/Showcases/Stylized
-Version: 4.9
 
 [REGION:banner]
 ![](PostProcessing_Header.png)(convert:false)
@@ -43,9 +42,12 @@ Version: 4.9
 [/OBJECT]
 [/VAR]
 
+[TOC(start:2 end:2)]
 
 
-The Content Examples provided in the **PostProcessing** map illustrate how to add screen effects inside and outside of a Post Processing Volume. By adjusting the properties of a Post Processing Volume you can achieve effects like lens flares, screen tinting or blurring and many others as well as define how volumes interact with one another or a player. 
+## Overview
+
+The Content Examples provided in the **PostProcessing** map illustrate how to add screen effects inside and outside of a **Post Processing Volume**. By adjusting the properties of a **Post Processing Volume** you can achieve effects like lens flares, screen tinting or blurring and many others as well as define how volumes interact with one another or a player. 
 
 
 [REGION:note]
@@ -61,25 +63,25 @@ Listed below are the examples provided inside the **PostProcessing** map:
 
 | Example | What is Demonstrated |
 | ------- | ------------ |
-| **[1.1 Default Settings](Resources\ContentExamples\PostProcessing\1_1)** | The default settings of a Post Process Volume (also the settings used when a camera is not within a Post Process Volume). |
-| **[1.2 Film](Resources\ContentExamples\PostProcessing\1_2)** | Film allows you to adjust the color of your scenes. |
-| **[1.3 SceneColor](Resources\ContentExamples\PostProcessing\1_3)** | Scene Color applies effects, instead of the color shifts seen in the Film section, to your scene (e.g. Vignette, Scene Fringe, and Tone Mapping effects are shown). |
-| **[1.4 Bloom](Resources\ContentExamples\PostProcessing\1_4)** | Bloom simulates the effects the eye perceives when viewing objects that appear very bright. |
-| **[1.5 Ambient Cubemap](Resources\ContentExamples\PostProcessing\1_5)** | The Ambient Cubemap lights the scene from a provided image. |
-| **[1.6 AutoExposure](Resources\ContentExamples\PostProcessing\1_6)** | Shows how to simulate the effect of the human eye adjusting to a bright environment when coming from a dark environment or vice versa with Auto Exposure.|
-| **[1.7 Lens Flares](Resources\ContentExamples\PostProcessing\1_7)** | Lens Flare simulates the scattering of light when viewing bright objects through imperfections found in camera lenses. |
+| **[1.1 Default Settings](Resources\ContentExamples\PostProcessing\1_1)** | The default settings of a **Post Process Volume** (also the settings used when a camera is not within a **Post Process Volume**). |
+| **[1.2 Film](Resources\ContentExamples\PostProcessing\1_2)** | **Film** allows you to adjust the color of your scenes. |
+| **[1.3 SceneColor](Resources\ContentExamples\PostProcessing\1_3)** | **Scene Color** applies effects, instead of the color shifts seen in the **Film** section, to your scene (e.g. **Vignette**, **Scene Fringe**, and **Tone Mapping** effects are shown). |
+| **[1.4 Bloom](Resources\ContentExamples\PostProcessing\1_4)** | **Bloom** simulates the effects the eye perceives when viewing objects that appear very bright. |
+| **[1.5 Ambient Cubemap](Resources\ContentExamples\PostProcessing\1_5)** | The **Ambient Cubemap** lights the scene from a provided image. |
+| **[1.6 AutoExposure](Resources\ContentExamples\PostProcessing\1_6)** | Shows how to simulate the effect of the human eye adjusting to a bright environment when coming from a dark environment or vice versa with **Auto Exposure**.|
+| **[1.7 Lens Flares](Resources\ContentExamples\PostProcessing\1_7)** | **Lens Flare** simulates the scattering of light when viewing bright objects through imperfections found in camera lenses. |
 | **[1.8 Ambient Occlusion](Resources\ContentExamples\PostProcessing\1_8)** | An effect that approximates the attenuation of light due to occlusion. Best used as a subtle effect that darkens corners or crevices to make them appear more natural.|
-| **[1.9 Screenspace Reflection](Resources\ContentExamples\PostProcessing\1_9)** | Alters the appearance of objects in reflections where the Intensity, Quality, or Max Roughness can be set.  |
-| **[1.10 Global Illumination](Resources\ContentExamples\PostProcessing\1_10)** | How to affect the indirect lighting contribution coming from Lightmass in order to alter a scene's brightness, tint, or color. |
-| **[1.11 Depth Of Field](Resources\ContentExamples\PostProcessing\1_11)** | Applying a blur effect to a scene based on distance in front or behind a focal point with Depth of Field. |
-| **[1.12 Motion Blur](Resources\ContentExamples\PostProcessing\1_12)** | Generating a Motion Blur which blurs objects based on its motion. |
-| **[1.13 Screen Percentage](Resources\ContentExamples\PostProcessing\1_13)** | Through Screen Percentage, a scene is rendered at a lower resolution and then scaled up. The Screen Percentage value represents a percentage of the overall resolution of the scene. |
-| **[1.14 AA Method](Resources\ContentExamples\PostProcessing\1_14)** | Demonstrates Anti-Aliasing (or AA) which improves image quality by smoothing out the edges that appear jagged on screen. |
-| **[1.15 Blur GBuffer intensity](Resources\ContentExamples\PostProcessing\1_15)** | Although hard to see, this effect helps with anti-aliasing by increasing the Material Roughness on the hard edges of surfaces, smoothing out the reflections on those hard edges. |
-| **[1.16 Blendables](Resources\ContentExamples\PostProcessing\1_16)** | Blendables are, in short, overlays that can be applied to the screen (e.g. screen warping or coloring). |
-| **[1.17 Priorities](Resources\ContentExamples\PostProcessing\1_17)** | Allows you to control which Post Process Volumes will override one another if they are overlapping. |
-| **[1.18 Blend Radius and Weight](Resources\ContentExamples\PostProcessing\1_18)** | Used to change the transition from one Post Process Volume to another over a distance (e.g. blend scene tint or color based on distance to an object). |
-| **[1.19 Unbound](Resources\ContentExamples\PostProcessing\1_19)** | This allows you to create a global Post Process effect without being constrained to being inside an actual volume (e.g. the effects are applied where ever you are in the world). |
+| **[1.9 Screenspace Reflection](Resources\ContentExamples\PostProcessing\1_9)** | Alters the appearance of objects in reflections where the **Intensity**, **Quality**, or **Max Roughness** can be set.  |
+| **[1.10 Global Illumination](Resources\ContentExamples\PostProcessing\1_10)** | How to affect the indirect lighting contribution coming from **Lightmass** in order to alter a scene's brightness, tint, or color. |
+| **[1.11 Depth Of Field](Resources\ContentExamples\PostProcessing\1_11)** | Applying a blur effect to a scene based on distance in front or behind a focal point with **Depth of Field**. |
+| **[1.12 Motion Blur](Resources\ContentExamples\PostProcessing\1_12)** | Generating a **Motion Blur** which blurs objects based on its motion. |
+| **[1.13 Screen Percentage](Resources\ContentExamples\PostProcessing\1_13)** | Through **Screen Percentage**, a scene is rendered at a lower resolution and then scaled up. The Screen Percentage value represents a percentage of the overall resolution of the scene. |
+| **[1.14 AA Method](Resources\ContentExamples\PostProcessing\1_14)** | Demonstrates **Anti-Aliasing** (or AA) which improves image quality by smoothing out the edges that appear jagged on screen. |
+| **[1.15 Blur GBuffer intensity](Resources\ContentExamples\PostProcessing\1_15)** | Although hard to see, this effect helps with anti-aliasing by increasing the **Material Roughness** on the hard edges of surfaces, smoothing out the reflections on those hard edges. |
+| **[1.16 Blendables](Resources\ContentExamples\PostProcessing\1_16)** | **Blendables** are, in short, overlays that can be applied to the screen (e.g. screen warping or coloring). |
+| **[1.17 Priorities](Resources\ContentExamples\PostProcessing\1_17)** | Allows you to control which **Post Process Volumes** will override one another if they are overlapping. |
+| **[1.18 Blend Radius and Weight](Resources\ContentExamples\PostProcessing\1_18)** | Used to change the transition from one **Post Process Volume** to another over a distance (e.g. blend scene tint or color based on distance to an object). |
+| **[1.19 Unbound](Resources\ContentExamples\PostProcessing\1_19)** | This allows you to create a global **Post Process** effect without being constrained to being inside an actual volume (e.g. the effects are applied where ever you are in the world). |
 
 
 
diff --git a/Engine/Documentation/Source/Resources/ContentExamples/Reflections/1_1/1_1.INT.udn b/Engine/Documentation/Source/Resources/ContentExamples/Reflections/1_1/1_1.INT.udn
--- a/Engine/Documentation/Source/Resources/ContentExamples/Reflections/1_1/1_1.INT.udn
+++ b/Engine/Documentation/Source/Resources/ContentExamples/Reflections/1_1/1_1.INT.udn
@@ -8,7 +8,6 @@ Related: Engine/Rendering/LightingAndShadows/LightTypes/SkyLight\
 Related: Engine/Rendering/Materials/PhysicallyBased
 Related: Resources/Showcases/RealisticRendering
 Related: Engine/Rendering/PostProcessEffects
-Version: 4.9
 
 [VAR:Steps]
 [OBJECT:Navigation]
diff --git a/Engine/Documentation/Source/Resources/ContentExamples/Reflections/1_2/1_2.INT.udn b/Engine/Documentation/Source/Resources/ContentExamples/Reflections/1_2/1_2.INT.udn
--- a/Engine/Documentation/Source/Resources/ContentExamples/Reflections/1_2/1_2.INT.udn
+++ b/Engine/Documentation/Source/Resources/ContentExamples/Reflections/1_2/1_2.INT.udn
@@ -8,7 +8,6 @@ Related: Engine/Rendering/LightingAndShadows/LightTypes/SkyLight\
 Related: Engine/Rendering/Materials/PhysicallyBased
 Related: Resources/Showcases/RealisticRendering
 Related: Engine/Rendering/PostProcessEffects
-Version: 4.9
 
 [VAR:Steps]
 [OBJECT:Navigation]
diff --git a/Engine/Documentation/Source/Resources/ContentExamples/Reflections/1_3/1_3.INT.udn b/Engine/Documentation/Source/Resources/ContentExamples/Reflections/1_3/1_3.INT.udn
--- a/Engine/Documentation/Source/Resources/ContentExamples/Reflections/1_3/1_3.INT.udn
+++ b/Engine/Documentation/Source/Resources/ContentExamples/Reflections/1_3/1_3.INT.udn
@@ -8,7 +8,6 @@ Related: Engine/Rendering/LightingAndShadows/LightTypes/SkyLight\
 Related: Engine/Rendering/Materials/PhysicallyBased
 Related: Resources/Showcases/RealisticRendering
 Related: Engine/Rendering/PostProcessEffects
-Version: 4.9
 
 [VAR:Steps]
 [OBJECT:Navigation]
@@ -71,7 +70,7 @@ Version: 4.9
 
 ![](Reflections_1_3_Header.png)
 
-The **Box Reflection Capture** Actor is one of two Reflection Capture Actors available in Unreal Engine 4. Many of these can be placed around the world with almost no performance impact as they are calculated before run time. While the box shape is a less commonly used in comparison to the sphere shape, it is very good for interiors that have a box shape themselves.
+The **Box Reflection Capture Actor** is one of two Reflection Capture Actors available in Unreal Engine 4. Many of these can be placed around the world with almost no performance impact as they are calculated before run time. While the box shape is a less commonly used in comparison to the sphere shape, it is very good for interiors that have a box shape themselves.
 
 One big drawback to consider for Box Reflection Captures is that if there are objects close to the Actor center (near the reflection sprite in the center of the box volume), they will appear very large in the reflection. This is due to the fact that the image used from the reflection is calculated from that sprite location, meaning that objects closer to that point will exhibit perspective-like depth from that point.
 
diff --git a/Engine/Documentation/Source/Resources/ContentExamples/Reflections/1_4/1_4.INT.udn b/Engine/Documentation/Source/Resources/ContentExamples/Reflections/1_4/1_4.INT.udn
--- a/Engine/Documentation/Source/Resources/ContentExamples/Reflections/1_4/1_4.INT.udn
+++ b/Engine/Documentation/Source/Resources/ContentExamples/Reflections/1_4/1_4.INT.udn
@@ -8,7 +8,6 @@ Related: Engine/Rendering/LightingAndShadows/LightTypes/SkyLight\
 Related: Engine/Rendering/Materials/PhysicallyBased
 Related: Resources/Showcases/RealisticRendering
 Related: Engine/Rendering/PostProcessEffects
-Version: 4.9
 
 [VAR:Steps]
 [OBJECT:Navigation]
diff --git a/Engine/Documentation/Source/Resources/ContentExamples/Reflections/1_5/1_5.INT.udn b/Engine/Documentation/Source/Resources/ContentExamples/Reflections/1_5/1_5.INT.udn
--- a/Engine/Documentation/Source/Resources/ContentExamples/Reflections/1_5/1_5.INT.udn
+++ b/Engine/Documentation/Source/Resources/ContentExamples/Reflections/1_5/1_5.INT.udn
@@ -8,7 +8,6 @@ Related: Engine/Rendering/LightingAndShadows/LightTypes/SkyLight\
 Related: Engine/Rendering/Materials/PhysicallyBased
 Related: Resources/Showcases/RealisticRendering
 Related: Engine/Rendering/PostProcessEffects
-Version: 4.9
 
 [VAR:Steps]
 [OBJECT:Navigation]
diff --git a/Engine/Documentation/Source/Resources/ContentExamples/Reflections/1_6/1_6.INT.udn b/Engine/Documentation/Source/Resources/ContentExamples/Reflections/1_6/1_6.INT.udn
--- a/Engine/Documentation/Source/Resources/ContentExamples/Reflections/1_6/1_6.INT.udn
+++ b/Engine/Documentation/Source/Resources/ContentExamples/Reflections/1_6/1_6.INT.udn
@@ -8,7 +8,6 @@ Related: Engine/Rendering/LightingAndShadows/LightTypes/SkyLight\
 Related: Engine/Rendering/Materials/PhysicallyBased
 Related: Resources/Showcases/RealisticRendering
 Related: Engine/Rendering/PostProcessEffects
-Version: 4.9
 
 [VAR:Steps]
 [OBJECT:Navigation]
diff --git a/Engine/Documentation/Source/Resources/ContentExamples/Reflections/1_7/1_7.INT.udn b/Engine/Documentation/Source/Resources/ContentExamples/Reflections/1_7/1_7.INT.udn
--- a/Engine/Documentation/Source/Resources/ContentExamples/Reflections/1_7/1_7.INT.udn
+++ b/Engine/Documentation/Source/Resources/ContentExamples/Reflections/1_7/1_7.INT.udn
@@ -8,7 +8,6 @@ Related: Engine/Rendering/LightingAndShadows/LightTypes/SkyLight\
 Related: Engine/Rendering/Materials/PhysicallyBased
 Related: Resources/Showcases/RealisticRendering
 Related: Engine/Rendering/PostProcessEffects
-Version: 4.9
 
 [VAR:Steps]
 [OBJECT:Navigation]
@@ -72,7 +71,7 @@ Version: 4.9
 ![](Reflections_1_7_Header.png)
 ## Overview
 
-Much like the [Scene Capture Cube](Resources\ContentExamples\Reflections\1_6), a **Scene Capture 2D** makes an image of the scene _on every frame_. In fact, Scene Capture 2D Actors do not work like real reflections; they work more like a TV screen receiving an image from a camera. The Scene Capture 2D captures the scene from its view frustum, stores that view as an image, which is then used within a Material. In fact, the Scene Capture 2D representation in Unreal Editor is a camera. As such, they are most commonly used for security camera effects in game. 
+Much like the [**Scene Capture Cube**](Resources\ContentExamples\Reflections\1_6), a **Scene Capture 2D** makes an image of the scene _on every frame_. In fact, Scene Capture 2D Actors do not work like real reflections; they work more like a TV screen receiving an image from a camera. The Scene Capture 2D captures the scene from its view frustum, stores that view as an image, which is then used within a Material. In fact, the **Scene Capture 2D** representation in Unreal Editor is a camera. As such, they are most commonly used for security camera effects in game. 
 
 [REGION:warning]
 While only 1/6 as expensive as the Scene Capture Cube, care should still be used with the Scene Capture 2D, as they still require extra processing.
diff --git a/Engine/Documentation/Source/Resources/ContentExamples/Reflections/Reflections.INT.udn b/Engine/Documentation/Source/Resources/ContentExamples/Reflections/Reflections.INT.udn
--- a/Engine/Documentation/Source/Resources/ContentExamples/Reflections/Reflections.INT.udn
+++ b/Engine/Documentation/Source/Resources/ContentExamples/Reflections/Reflections.INT.udn
@@ -8,7 +8,6 @@ Related: Engine/Rendering/LightingAndShadows/LightTypes/SkyLight\
 Related: Engine/Rendering/Materials/PhysicallyBased
 Related: Resources/Showcases/RealisticRendering
 Related: Engine/Rendering/PostProcessEffects
-Version: 4.9
 
 
 [REGION:banner]
@@ -46,11 +45,12 @@ Version: 4.9
 [/OBJECT]
 [/VAR]
 
+[TOC(start:2 end:2)]
 
-The **Reflections** Content Example map provides examples on how to adjust the way reflections appear within your levels by using a **Reflection Capture** Actor (which projects a reflection within the Actor's effect volume or radius). 
+The **Reflections** Content Example map provides examples on how to adjust the way reflections appear within your levels by using a Reflection Capture Actor (which projects a reflection within the Actor's effect volume or radius). 
 Examples are also shown on how the Roughness, Metallic, and Specular properties of a Material can be used to produce different reflection results. 
 
-In addition to tweaking reflection settings, the use of both a **Box Reflection** Actor and the more commonly used Sphere Reflection Actor are shown with sample effects that can be achieved.
+In addition to tweaking reflection settings, the use of both a Box Reflection Actor and the more commonly used*Sphere Reflection Actor are shown with sample effects that can be achieved.
 Lastly, Scene Capture methods are demonstrated which can be used to capture a fully dynamic cube map or capture an image of the scene on every frame. 
 
 
diff --git a/Engine/Documentation/Source/Resources/ContentExamples/StaticMeshes/StaticMeshes.INT.udn b/Engine/Documentation/Source/Resources/ContentExamples/StaticMeshes/StaticMeshes.INT.udn
--- a/Engine/Documentation/Source/Resources/ContentExamples/StaticMeshes/StaticMeshes.INT.udn
+++ b/Engine/Documentation/Source/Resources/ContentExamples/StaticMeshes/StaticMeshes.INT.udn
@@ -7,7 +7,6 @@ Related: Engine/Physics/Collision/Reference
 Related: Engine/Content/Types/StaticMeshes/Editor
 Related: Engine/Actors/StaticMeshActor
 Related: Engine/Content/FBX/BestPractices
-Version: 4.9
 
 
 
@@ -46,6 +45,8 @@ Version: 4.9
 [/OBJECT]
 [/VAR]
 
+[TOC(start:2 end:2)]
+
 Inside the **StaticMeshes** Content Example map, you will find essential examples using Static Meshes including how to set Mobility, assign UVs, and add or adjust Collision settings. Other examples include how to customize your Static Meshes with Material Elements (which allow you to specify the Material applied to parts of a Static Mesh) as well as using the MeshPaint Tool to manually paint onto a Static Mesh. 
 Last, you will see how LODs can improve performance by scaling back the quality of a mesh based on player distance.
 
@@ -54,7 +55,7 @@ Last, you will see how LODs can improve performance by scaling back the quality
 
 ![](StaticMeshes_Map.png)(w:640)
 
-Listed below are the examples provided inside the **StaticMeshes** map:
+Listed below are the examples provided inside the StaticMeshes map:
 
 | Example | What is Demonstrated |
 | ------- | ------------ |
diff --git a/Engine/Documentation/Source/Resources/ContentExamples/Volumes/Volumes.INT.udn b/Engine/Documentation/Source/Resources/ContentExamples/Volumes/Volumes.INT.udn
--- a/Engine/Documentation/Source/Resources/ContentExamples/Volumes/Volumes.INT.udn
+++ b/Engine/Documentation/Source/Resources/ContentExamples/Volumes/Volumes.INT.udn
@@ -7,7 +7,6 @@ Related: Engine/Rendering/PostProcessEffects
 Related: Engine/Rendering/LightingAndShadows/LightPropagationVolumes
 Related: Engine/AI/BehaviorTrees/QuickStart/2
 Related: Engine/Actors/Triggers
-Version: 4.9
 
 
 [REGION:banner]
@@ -45,7 +44,7 @@ Version: 4.9
 [/OBJECT]
 [/VAR]
 
-
+[TOC(start:2 end:2)]
 
 The **Volumes** Content Example map illustrates how you can use different types of Volumes to produce different gameplay effects, set boundaries, improve performance and (or) reduce lighting build times. 
 As seen in example 1.1 below, you can place a Volume by selecting any Volume from the Volumes Tab and dragging and dropping it into your level. Based on the type of Volume used, different settings may be available for editing. 
@@ -58,7 +57,7 @@ The most commonly used Volumes for gameplay purposes would be Blocking Volumes (
 
 ![](Volumes_Map.png)(w:640)
 
-Listed below are the examples provided inside the **Volumes** map:
+Listed below are the examples provided inside the Volumes map:
 
 | Example | What is Demonstrated |
 | ------- | ------------ |
diff --git a/Engine/Documentation/Source/Resources/Downloads/Downloads.INT.udn b/Engine/Documentation/Source/Resources/Downloads/Downloads.INT.udn
--- a/Engine/Documentation/Source/Resources/Downloads/Downloads.INT.udn
+++ b/Engine/Documentation/Source/Resources/Downloads/Downloads.INT.udn
@@ -2,7 +2,6 @@ Availability: Licensee
 Title:Resource Downloads
 Crumbs:%ROOT%, Resources
 Description:Repository of downloadable content and example materials for use in developing games with UE4.
-Version: 4.9
 
 [TOC]
 
diff --git a/Engine/Documentation/Source/Resources/SampleGames/CardGame/CardGame.INT.udn b/Engine/Documentation/Source/Resources/SampleGames/CardGame/CardGame.INT.udn
--- a/Engine/Documentation/Source/Resources/SampleGames/CardGame/CardGame.INT.udn
+++ b/Engine/Documentation/Source/Resources/SampleGames/CardGame/CardGame.INT.udn
@@ -8,7 +8,6 @@ Related: Engine/Animation/AnimBlueprints/EventGraph
 Related: Gameplay/Input
 Related: Gameplay/Framework
 Related: Engine/Content/Types/StaticMeshes
-Version: 4.9
 
 
 [VAR:Topic]
diff --git a/Engine/Documentation/Source/Resources/SampleGames/PlatformerGame/PlatformerGame.INT.udn b/Engine/Documentation/Source/Resources/SampleGames/PlatformerGame/PlatformerGame.INT.udn
--- a/Engine/Documentation/Source/Resources/SampleGames/PlatformerGame/PlatformerGame.INT.udn
+++ b/Engine/Documentation/Source/Resources/SampleGames/PlatformerGame/PlatformerGame.INT.udn
@@ -2,7 +2,6 @@ Availability: Public
 Title:Platformer Game Sample
 Crumbs:%ROOT%, Resources
 Description:Overview of the Platformer Sample Game.
-Version: 4.9
 
 
 [VAR:Topic]
@@ -53,7 +52,7 @@ Version: 4.9
 
 [TOC(start:2)]
 
-The **Platformer Game** sample is an example of a PC/mobile side-scroller game. It includes basic implementations of 
+The PlatformerGame sample is an example of a PC/mobile side-scroller game. It includes basic implementations of 
 different movement types along with a simple front end menu system.
 
 A complete list of the featured concepts:
diff --git a/Engine/Documentation/Source/Resources/SampleGames/SampleGames.INT.udn b/Engine/Documentation/Source/Resources/SampleGames/SampleGames.INT.udn
--- a/Engine/Documentation/Source/Resources/SampleGames/SampleGames.INT.udn
+++ b/Engine/Documentation/Source/Resources/SampleGames/SampleGames.INT.udn
@@ -2,7 +2,6 @@ Availability: Public
 Title:Example Game Projects
 Crumbs:%ROOT%, Resources
 Description:Full sample games demonstrating the implementation of various genres.
-Version: 4.9
 
 
 [VAR:Topic]
diff --git a/Engine/Documentation/Source/Resources/SampleGames/ShooterGame/ShooterGame.INT.udn b/Engine/Documentation/Source/Resources/SampleGames/ShooterGame/ShooterGame.INT.udn
--- a/Engine/Documentation/Source/Resources/SampleGames/ShooterGame/ShooterGame.INT.udn
+++ b/Engine/Documentation/Source/Resources/SampleGames/ShooterGame/ShooterGame.INT.udn
@@ -14,7 +14,6 @@ Related: Engine/Audio/Overview
 Related: Programming/Slate
 Related: Gameplay/Input
 Related: Engine/Rendering/ParticleSystems
-Version: 4.9
 
 
 
diff --git a/Engine/Documentation/Source/Resources/SampleGames/StrategyGame/StrategyGame.INT.udn b/Engine/Documentation/Source/Resources/SampleGames/StrategyGame/StrategyGame.INT.udn
--- a/Engine/Documentation/Source/Resources/SampleGames/StrategyGame/StrategyGame.INT.udn
+++ b/Engine/Documentation/Source/Resources/SampleGames/StrategyGame/StrategyGame.INT.udn
@@ -10,7 +10,6 @@ Related: Programming/Slate
 Related: Gameplay/Framework/UIAndHUD
 Related: Gameplay/Input
 Related: Engine/Rendering/ParticleSystems
-Version: 4.9
 
 
 
@@ -57,7 +56,7 @@ Version: 4.9
 
 [TOC(start:2 end:2)]
 
-The **Tower Defense** sample is an example of a RTS/Tower Defense game. 
+The Tower Defense sample is an example of a RTS/Tower Defense game. 
 
 A complete list of the featured concepts:
 
@@ -84,23 +83,23 @@ selected, and the current action can be stopped if there is an action with highe
 Both enemy and friendly Pawns operate with the AI logic, moving towards the opposing base and attacking Pawns of the other team if they encounter them.  While players cannot control the movement
 or behavior of their friendly Pawns, they are able to buy new units to spawn.  
 
-Blueprint are also used to add logic to Minion Pawns.  Both friendly and enemy Pawns can be equipped with shields; friendly Pawns get shields if the armory upgrade has been purchased for the brewery,
+_Blueprints_ are also used to add logic to Minion Pawns.  Both friendly and enemy Pawns can be equipped with shields; friendly Pawns get shields if the armory upgrade has been purchased for the brewery,
 and enemy Pawns get shields if they were spawned due to the SpawnHeavyFunction or SpawnEndBossFunction being called in the _Level Blueprint_. If a Pawn has a shield, projectiles from the auto-arbalest will be destroyed and do no damage. This logic is carried out using _Blueprint Interfaces_.  The
-**Minion** Blueprint also contains a network to slow down enemy Pawns when they are hit by a charged-up fire turret.  
+**Minion** _Blueprint_ also contains a network to slow down enemy Pawns when they are hit by a charged-up fire turret.  
 
 ## Building Construction
 
 There are two building classes in Tower Defense - `StrategyBuilding` and `StrategyBuilding_Brewery`.  All the turret types in Tower Defense, as well as the empty building slot, are designed using 
-Blueprint with `StrategyBuilding` as the parent class.  Players can click on an empty building slot to show the context menu and select a new building to construct.  When a building is constructed, the empty building slot is destroyed and the new building is spawned.  
+_Blueprints_ with `StrategyBuilding` as the parent class.  Players can click on an empty building slot to show the context menu and select a new building to construct.  When a building is constructed, the empty building slot is destroyed and the new building is spawned.  
 
 There is a mechanism in place for upgrading buildings as well. The `StrategyBuilding_Brewery` class implements this case so that upgrades are built in connected slots near the brewery base.
 
-Again, the code in Tower Defense simply creates the base building classes. All logic and design of the buildings in the Tower Defense were made by level designers in Blueprint.
+Again, the code in Tower Defense simply creates the base building classes. All logic and design of the buildings in the Tower Defense were made by level designers in _Blueprints_.
 
 ### Brewery
 
-The **Brewery** Blueprint has the parent class `StrategyBuilding_Brewery`, and also includes an AIDirector component.  There are two breweries placed in the TowerDefenseMap, one for the enemy where enemy
-Pawns spawn, and one friendly brewery where an armory and a smithy upgrade can be built, and friendly Pawns can be spawned.  There is no graph logic present in the **Brewery** Blueprint, just Defaults set for the 
+The **Brewery** _Blueprint_ has the parent class `StrategyBuilding_Brewery`, and also includes an AIDirector component.  There are two breweries placed in the TowerDefenseMap, one for the enemy where enemy
+Pawns spawn, and one friendly brewery where an armory and a smithy upgrade can be built, and friendly Pawns can be spawned.  There is no graph logic present in the **Brewery** _Blueprint_, just Defaults set for the 
 building properties and a Components list including the AIDirector, a TriggerBox, and a _Static Mesh_. 
 
 #### Upgrades
@@ -110,12 +109,12 @@ building properties and a Components list including the AIDirector, a TriggerBox
 There are two upgrade slots attached to the friendly brewery.  These slots are _Blueprint Classes_ which are also derived from the `StrategyBuilding` class.  If an upgrade is selected from the Brewery menu,
 one slot will be replaced by that upgrade. Only one smithy upgrade and one armory upgrade can be purchased. 
 
-Once a smithy upgrade is purchased, the build will be started, firing the OnBuildStarted event in the **Wall_Smithy** Blueprint.  This Blueprint also informs the system that the upgrade has been built,
-once the build completes.  After this point, any friendly Pawns that are spawned will be equipped with a shield attachment, a Blueprint derived from the `StrategyAttachment` class.  The network which assigns the 
+Once a smithy upgrade is purchased, the build will be started, firing the OnBuildStarted event in the **Wall_Smithy** _Blueprint_.  This _Blueprint_ also informs the system that the upgrade has been built,
+once the build completes.  After this point, any friendly Pawns that are spawned will be equipped with a shield attachment, a _Blueprint_ derived from the `StrategyAttachment` class.  The network which assigns the 
 "shield-attaching" behavior after the armory reports that it has been constructed is present in the PlayerBaseUpgrades collapsed graph in the TowerDefenseMap _Level Blueprint_. The `StrategyAttachment` 
-class simply contains a  `SkeletalMeshComponent`; the mesh and attachment point for the attachment are set in the Defaults for the **Attachment_Armorer** Blueprint.
+class simply contains a  `SkeletalMeshComponent`; the mesh and attachment point for the attachment are set in the Defaults for the **Attachment_Armorer** _Blueprint_.
 
-The armory Blueprint contains the same logic setup with OnBuildStarted and OnBuildFinished events.  After construction of the armory, any friendly Pawns spawned will be equipped with a hammer, also derived from the 
+The armory _Blueprint_ contains the same logic setup with OnBuildStarted and OnBuildFinished events.  After construction of the armory, any friendly Pawns spawned will be equipped with a hammer, also derived from the 
 `StrategyAttachment` class.  The network which assigns the "hammer-attaching" behavior after the smithy reports that it has been constructed is also present in the PlayerBaseUpgrades collapsed graph in the TowerDefenseMap _Level Blueprint_.
 
 ### Turrets
@@ -124,21 +123,21 @@ The armory Blueprint contains the same logic setup with OnBuildStarted and OnBui
 
 ![](building_empty_slot.png)
 
-The empty slot is also a Blueprint with the `StrategyBuilding` parent class, Wall_EmptySlot. No logic is present in the Blueprint graphs. This is a _Blueprint Class_ with Defaults set for building 
+The empty slot is also a _Blueprint_ with the `StrategyBuilding` parent class, Wall_EmptySlot. No logic is present in the _Blueprint_ graphs. This is a _Blueprint Class_ with Defaults set for building 
 properties and Static Meshes and a trigger box set as the Components.
 
-The possible turret upgrades are all set in the Defaults of the Wall_EmptySlot Blueprint, in the Upgrades section of the Building category.  
+The possible turret upgrades are all set in the Defaults of the Wall_EmptySlot _Blueprint_, in the Upgrades section of the Building category.  
 
 
 #### Arbalest
 
 ![](building_arbalest_shooting.png)
 
-The **Wall_arbalest** Blueprint contains the logic for the arbalest, the basic turret type. The arbalest shoots the closest enemy with a medium-strength projectile, automatically firing arrows in its default mode. The player
+The **Wall_arbalest** _Blueprint_ contains the logic for the arbalest, the basic turret type. The arbalest shoots the closest enemy with a medium-strength projectile, automatically firing arrows in its default mode. The player
 can manually fire the arbalest as well, by clicking on it and dragging in the direction they want the arrow to fire.  The longer the mouse drag is, the stronger the shot will be.  
 
 
-The arbalest projectile is stored in another Blueprint, Projectile_arbalest, derived from the Blueprint TestProjectile which has `StrategyProjectile`.  The Wall_arbalest Blueprint has a number of sub-networks, all contained within the EventGraph.  No Blueprint logic is present in the ConstructionScript.   
+The arbalest projectile is stored in another _Blueprint_, Projectile_arbalest, derived from the _Blueprint_ TestProjectile which has `StrategyProjectile`.  The Wall_arbalest _Blueprint_ has a number of sub-networks, all contained within the EventGraph.  No _Blueprint_ logic is present in the ConstructionScript.   
 
 
 
@@ -146,11 +145,11 @@ The arbalest projectile is stored in another Blueprint, Projectile_arbalest, der
 
 ![](building_auto-arbalest_shooting.png)
 
-The **Wall_arbalest_auto** Blueprint contains the logic for the auto-arbalest.  The auto-arbalest shoots projectiles in a straight line from the wall, doing a small amount of damage to every unit the projectiles
+The **Wall_arbalest_auto** _Blueprint_ contains the logic for the auto-arbalest.  The auto-arbalest shoots projectiles in a straight line from the wall, doing a small amount of damage to every unit the projectiles
 pass through.  The auto-arbalest arrows are not destroyed unless they hit a wall or an enemy with a shield.  It is possible to aim the auto-arbalest by clicking and dragging so that the auto-arbalest faces the desired direction;
 the auto-arbalest will continue to fire in the aimed direction while the mouse button is held down, and will return to its default firing position when the mouse button is released. 
 
-Like the arbalest, this turret shoots projectile arrows contained in a separate Blueprint.  The Projectile_arbalest_auto auto-arbalest arrows are not destroyed unless they hit an enemy Pawn with a shield or a wall, and this behavior
+Like the arbalest, this turret shoots projectile arrows contained in a separate _Blueprint_.  The Projectile_arbalest_auto auto-arbalest arrows are not destroyed unless they hit an enemy Pawn with a shield or a wall, and this behavior
 is carried out with the aid of _Blueprint Interfaces_, Interface_Auto_Arbalest and Interface_Auto_Projectile. 
 
 #### Flamethrower
@@ -240,11 +239,11 @@ The waves are constructed using three _Blueprint Macros_: spawn fast, spawn norm
 
 Each spawn macro takes two execution inputs, one for beginning the macro and one for opening the **Gate** after the OnWaveSpawned event fires, and an integer input for the number of Pawns to spawn.  
 Functions from the class `StrategyAIDirector` are called, with inputs specific to each type of Pawn wave.  The three functions are `SetDefaultWeapon`, `SetDefaultArmor`, and `SetBuffModifier`.  
-`SetDefaultWeapon` and `SetDefaultArmor` take Blueprint as inputs and assign those Blueprint as the new default weapon or default armor for the spawn.  For example, all enemy Pawns spawned by 
-the SpawnFastMacro have the Attachment_Smithy hammer Blueprint as their default weapon, and all enemy Pawns spawned by the SpawnHeavyMacro have the Attachment_Armorer shield Blueprint as their default armor. 
+`SetDefaultWeapon` and `SetDefaultArmor` take _Blueprints_ as inputs and assign those _Blueprints_ as the new default weapon or default armor for the spawn.  For example, all enemy Pawns spawned by 
+the SpawnFastMacro have the Attachment_Smithy hammer _Blueprint_ as their default weapon, and all enemy Pawns spawned by the SpawnHeavyMacro have the Attachment_Armorer shield _Blueprint_ as their default armor. 
 
-The last `StrategyAIDirector` function called by the spawning Blueprint functions is `SetBuffModifier`, which has a number of data inputs including the attack abilities, health bonus, speed, and size
-of the Pawns.  These inputs are all exposed to the Blueprint, so it is straightforward for a level designer to create a new class of enemy Pawn to spawn.  Finally, each spawning Blueprint function sets the `WaveSize`
+The last `StrategyAIDirector` function called by the spawning _Blueprint_ functions is `SetBuffModifier`, which has a number of data inputs including the attack abilities, health bonus, speed, and size
+of the Pawns.  These inputs are all exposed to the _Blueprint_, so it is straightforward for a level designer to create a new class of enemy Pawn to spawn.  Finally, each spawning _Blueprint_ function sets the `WaveSize`
 property of the enemy brewery's `StrategyAIDirector`.   
 
 There are five waves of enemies, each with different combinations of fast, normal, and heavy enemy Pawns.  At the beginning of a wave, the Show Wave Title node displays the wave number.  Then, the first enemy 
@@ -273,10 +272,10 @@ NumberOfLives of the friendly brewery, decreasing the value by one each time. Th
 
 ## Resource Node - Gold
 
-The gold nodes are Blueprint with the parent class `StrategyResourceNode`.  This class contains the public functions `GetAvailableResources` and `GetInitialResources`, the protected function `OnDepleted`, 
+The gold nodes are _Blueprints_ with the parent class `StrategyResourceNode`.  This class contains the public functions `GetAvailableResources` and `GetInitialResources`, the protected function `OnDepleted`, 
 a `BlueprintImplementableEvent` to report when the resource has been depleted, and the protected property `NumResources` to set the amount of the resource present in the node.  
 
-The Blueprint for a gold node contains sub-networks to make the node appear and disappear on a timer.  The **ConstructionScript** in the Blueprint sets the node to automatically be hidden when it is placed in the level.
+The _Blueprint_ for a gold node contains sub-networks to make the node appear and disappear on a timer.  The **ConstructionScript** in the _Blueprint_ sets the node to automatically be hidden when it is placed in the level.
 When the gold node appears, AppearFX particle effects are played along with an apparition noise. If the OnDepleted event fires because the node is successfully collected, the amount of gold present in the node
 is added to the player's total gold pool.  The CollectFX particle effects and the CoinSound is played, and then the node is hidden again.  If the player fails to collect the node in time by clicking on it,
 FadeFX particle effects and an appropriate sound is played.  
diff --git a/Engine/Documentation/Source/Resources/Showcases/BlueprintExamples/BlueprintExamples.INT.udn b/Engine/Documentation/Source/Resources/Showcases/BlueprintExamples/BlueprintExamples.INT.udn
--- a/Engine/Documentation/Source/Resources/Showcases/BlueprintExamples/BlueprintExamples.INT.udn
+++ b/Engine/Documentation/Source/Resources/Showcases/BlueprintExamples/BlueprintExamples.INT.udn
@@ -8,7 +8,6 @@ Related: Resources/ContentExamples/Blueprints
 Related: Gameplay/Framework/Camera
 Related: Engine/Rendering/ParticleSystems
 Related: Engine/Rendering/Materials
-Version: 4.9
 
 
 [VAR:Topic]
diff --git a/Engine/Documentation/Source/Resources/Showcases/BlueprintExamples/Butterfly/Butterfly.INT.udn b/Engine/Documentation/Source/Resources/Showcases/BlueprintExamples/Butterfly/Butterfly.INT.udn
--- a/Engine/Documentation/Source/Resources/Showcases/BlueprintExamples/Butterfly/Butterfly.INT.udn
+++ b/Engine/Documentation/Source/Resources/Showcases/BlueprintExamples/Butterfly/Butterfly.INT.udn
@@ -8,7 +8,6 @@ Related: Engine/Animation/AnimBlueprints/EventGraph
 Related: Engine/Blueprints/UserGuide/UserConstructionScript
 Related: Engine/Animation/AnimBlueprints/EventGraph
 Related: Gameplay/AI
-Version: 4.9
 
 [REGION:fullwidth]
 ![](ButterflyHeader.png)
diff --git a/Engine/Documentation/Source/Resources/Showcases/BlueprintExamples/CommentaryBox/CommentaryBox.INT.udn b/Engine/Documentation/Source/Resources/Showcases/BlueprintExamples/CommentaryBox/CommentaryBox.INT.udn
--- a/Engine/Documentation/Source/Resources/Showcases/BlueprintExamples/CommentaryBox/CommentaryBox.INT.udn
+++ b/Engine/Documentation/Source/Resources/Showcases/BlueprintExamples/CommentaryBox/CommentaryBox.INT.udn
@@ -7,7 +7,6 @@ Related: Engine/Animation/AnimBlueprints/EventGraph
 Related: Engine/Actors/Triggers
 Related: Engine/Blueprints
 Related: Engine/Content/Types/SkeletalMeshes/Sockets
-Version: 4.9
 
 [REGION:fullwidth]
 ![](BlueprintCommentary.png)
diff --git a/Engine/Documentation/Source/Resources/Showcases/BlueprintExamples/FogSheet/FogSheet.INT.udn b/Engine/Documentation/Source/Resources/Showcases/BlueprintExamples/FogSheet/FogSheet.INT.udn
--- a/Engine/Documentation/Source/Resources/Showcases/BlueprintExamples/FogSheet/FogSheet.INT.udn
+++ b/Engine/Documentation/Source/Resources/Showcases/BlueprintExamples/FogSheet/FogSheet.INT.udn
@@ -6,7 +6,6 @@ Related: Engine/Blueprints/Overview
 Related: Engine/Blueprints/UserGuide/Variables
 Related: Engine/Blueprints/UserGuide/UserConstructionScript
 Related: Engine/Rendering/Materials/MaterialInstances
-Version: 4.9
 
 
 
@@ -34,7 +33,7 @@ The Fog Sheet Blueprint, shown above, consists of a Static Mesh plane with a clo
 [![](FogSheetCS.png)](FogSheetCS.png)
 [/REGION]
 [/REGION]
-Each of the three aspects of the Blueprint (i.e. the _Components List_, the _Construction Script_, and the _Event Graph_) are responsible for a different part of the overall effect. The following table breaks down the jobs for each aspect and how they work together to create the final result.
+Each of the 3 aspects of the Blueprint (i.e. the _Components List_, the _Construction Script_, and the _Event Graph_) are responsible for a different part of the overall effect. The following table breaks down the jobs for each aspect and how they work together to create the final result.
 
 
 
diff --git a/Engine/Documentation/Source/Resources/Showcases/BlueprintExamples/Foliage/Foliage.INT.udn b/Engine/Documentation/Source/Resources/Showcases/BlueprintExamples/Foliage/Foliage.INT.udn
--- a/Engine/Documentation/Source/Resources/Showcases/BlueprintExamples/Foliage/Foliage.INT.udn
+++ b/Engine/Documentation/Source/Resources/Showcases/BlueprintExamples/Foliage/Foliage.INT.udn
@@ -6,7 +6,6 @@ Related: Engine/Blueprints/Overview
 Related: Engine/Blueprints/UserGuide/Variables
 Related: Engine/Animation/AnimBlueprints/EventGraph
 Related: Engine/Blueprints/UserGuide/UserConstructionScript
-Version: 4.9
 
 
 [REGION:fullwidth]
@@ -14,7 +13,7 @@ Version: 4.9
 [/REGION]
 
 [EXCERPT:Overview]
-The **Foliage** Blueprint shows how to use a Blueprint to create randomly-generated foliage (or potentially other structures) that are automatically placed such that they conform to a surface. The user has control over the extent of the effect and how many meshes the system will try to spawn. As the Blueprint is moved, the meshes update to re-conform to the new environment.
+The Foliage Blueprint shows how to use a Blueprint to create randomly-generated foliage (or potentially other structures) that are automatically placed such that they conform to a surface. The user has control over the extent of the effect and how many meshes the system will try to spawn. As the Blueprint is moved, the meshes update to re-conform to the new environment.
 [/EXCERPT:Overview]
 
 
@@ -29,7 +28,7 @@ The Components List, Construction Script, and Event Graph for this Blueprint per
 |---|---|
 | **Components List** | Adds a BillboardComponent to represent the center and a BoxComponent to help visualize the construction radius. |
 | **Construction Script** | [INCLUDE:#FoliageCS] |
-| **Event Graph** | This Blueprint has no Event Graph, meaning that once its properties are set, it remains the same throughout gameplay. |
+| **Event Graph** | This Blueprint has no **Event Graph**, meaning that once its properties are set, it remains the same throughout gameplay. |
 
 [REGION:lightbox]
 ![](FoliageCS.png)(w:920)
@@ -47,12 +46,12 @@ As you might expect, most of the work of this Blueprint is done in the **Constru
 
 | Property | Purpose |
 | --- | --- |
-| **Number of Meshes** | This integer value is the number of meshes the Blueprint will attempt to spawn. In the Construction Script, this value determines the number of times a _ForLoop_ node will trigger. Each execution coming out of the _ForLoop_ node passes through the section of script responsible for creating the Static Meshes.  |
+| **Number of Meshes** | This integer value is the number of meshes the Blueprint will attempt to spawn. In the **Construction Script**, this value determines the number of times a _ForLoop_ node will trigger. Each execution coming out of the _ForLoop_ node passes through the section of script responsible for creating the Static Meshes.  |
 | **Radius** | Controls the radius of the effect, within which the foliage will be randomly scattered. |
 | **Collision Trace Range** | This is the vertical distance (how far down the system will look) to find a valid surface on which to spawn foliage. |
-| **Meshes To Spawn** | This is an array variable, which meshes from the Content Browser can be added to. After the Construction Script spawns a mesh component, it is changed to a mesh picked at random from this array. |
+| **Meshes To Spawn** | This is an array variable, which meshes from the **Content Browser** can be added to. After the **Construction Script** spawns a mesh component, it is changed to a mesh picked at random from this array. |
 | **Vertical Offset** | Controls how far _into_ a surface a mesh will be spawned. This can help meshes appear more grounded. |
-| **Align To Surface** | This Boolean variable causes meshes spawned by the Blueprint will orient themselves such that they match the angle of the surface upon which they are resting. Internally, this variable affects a branch in the Construction Script that determines the method used for positioning the foliage meshes. When enabled, the script draws a line trace before attempting to spawn each mesh. When the trace hits a surface, the angle and height of the impact point is retrieved and used to determine the angle and height of the mesh being spawned. |
+| **Align To Surface** | This Boolean variable causes meshes spawned by the Blueprint will orient themselves such that they match the angle of the surface upon which they are resting. Internally, this variable affects a branch in the **Construction Script** that determines the method used for positioning the foliage meshes. When enabled, the script draws a line trace before attempting to spawn each mesh. When the trace hits a surface, the angle and height of the impact point is retrieved and used to determine the angle and height of the mesh being spawned. |
 | **Scale Max** | Maximum scale allowed for the spawned meshes. This is a 0.0-1.0 percentage value, with 1.0 representing 100%. |
 | **Scale Min** | Minimum scale allowed for the spawned meshes. This is a 0.0-1.0 percentage value, with 1.0 representing 100%. |
 | **Cluster Around Center** | This Boolean variable determines whether or not the randomized locations for the meshes will be primarily located around the center of the radius. When this option is enabled, the distance meshes can spawn from the center point is reduced by a random amount, causing fewer meshes to spawn the farther away the system looks from the radius. |
diff --git a/Engine/Documentation/Source/Resources/Showcases/BlueprintExamples/SecurityCameraAndAlarm/SecurityCameraAndAlarm.INT.udn b/Engine/Documentation/Source/Resources/Showcases/BlueprintExamples/SecurityCameraAndAlarm/SecurityCameraAndAlarm.INT.udn
--- a/Engine/Documentation/Source/Resources/Showcases/BlueprintExamples/SecurityCameraAndAlarm/SecurityCameraAndAlarm.INT.udn
+++ b/Engine/Documentation/Source/Resources/Showcases/BlueprintExamples/SecurityCameraAndAlarm/SecurityCameraAndAlarm.INT.udn
@@ -7,7 +7,6 @@ Related: Engine/Blueprints/UserGuide/Variables
 Related: Engine/Animation/AnimBlueprints/EventGraph
 Related: Engine/Blueprints/UserGuide/UserConstructionScript
 Related: Engine/Blueprints/UserGuide/Components
-Version: 4.9
 
 
 [EXCERPT:Overview]
@@ -19,7 +18,7 @@ Tied into the existing camera Blueprints is a security control room. This room c
 ![](ScreenBlueprint.png)(w:600)
 
 
-The **Screen** Blueprint is - at its heart - an optimized mechanism to switch between different Scene Capture Actors. The effect of the security cameras is created because of Scene Capture Actors located within each of the [Camera Blueprints](Resources\Showcases\BlueprintExamples\SecurityDoor). The purpose of a Scene Capture Actor is to function very much like a camera, capturing an image of the pixels it can see. It then transmits those to a Render Target texture, which functions like an internal video screen. By applying this Render Target texture to a Material, you effectively create a Material that can show what the camera sees.
+The Screen Blueprint is - at its heart - an optimized mechanism to switch between different Scene Capture Actors. The effect of the security cameras is created because of Scene Capture Actors located within each of the [Camera Blueprints](Resources\Showcases\BlueprintExamples\SecurityDoor). The purpose of a Scene Capture Actor is to function very much like a camera, capturing an image of the pixels it can see. It then transmits those to a Render Target texture, which functions like an internal video screen. By applying this Render Target texture to a Material, you effectively create a Material that can show what the camera sees.
 
 The catch? With many cameras running at the same time, it is possible to run into performance issues, since they would all have to be capturing and transmitting their current views. To alleviate this problem, the screen Blueprint will activate only one Scene Capture at a time, and deactivate the one that was used previously.
 
diff --git a/Engine/Documentation/Source/Resources/Showcases/BlueprintExamples/SecurityDoor/SecurityDoor.INT.udn b/Engine/Documentation/Source/Resources/Showcases/BlueprintExamples/SecurityDoor/SecurityDoor.INT.udn
--- a/Engine/Documentation/Source/Resources/Showcases/BlueprintExamples/SecurityDoor/SecurityDoor.INT.udn
+++ b/Engine/Documentation/Source/Resources/Showcases/BlueprintExamples/SecurityDoor/SecurityDoor.INT.udn
@@ -9,11 +9,8 @@ Related: Engine/Blueprints/UserGuide/UserConstructionScript
 Related: Engine/Blueprints/UserGuide/Components
 Related: Engine/Blueprints/UserGuide/Timelines
 Related: Gameplay/Framework/Pawn
-Version: 4.9
 
-<div align= "center">
 ![](SecurityCameraAndDoorHeader.png)(w:600)
-</div>
 
 [EXCERPT:Overview]
 The security camera Blueprints seen throughout the level are designed to interact with the player when in Play mode. During play, they will pan back and forth in a search pattern. When the green search cone overlaps a pawn, the light turns red and the camera begins to track player movement. It continues to track the player until the player moves out of range. Then, after a few seconds, it returns to searching.
@@ -82,9 +79,8 @@ The camera has the following exposed properties:
 
 
 ## Alarm Blueprint
-<div align= "center">
+
 ![](AlarmImage.png)
-</div>
 
 The Alarm Blueprint serves as a visual cue that the player has tripped the alarm by stepping in front of the security camera. When activated, it seems to change from green to red, and a beacon light begins spinning. It continues to do this until the alarm deactivates.
 
@@ -120,9 +116,8 @@ The alarm has the following exposed properties:
 
 
 ## Door Blueprint
-<div align= "center">
+
 ![](Door.png)
-</div>
 
 The Door is set up such that it remains open so long as the alarm is not engaged. The moment the alarm is triggered, however, the doors are animated such that they slam shut and will not reopen until the alarm ends. It is important to note that Timelines are used to animate the doors. In fact, there are 2 separate Timelines - one to open the doors and another to close them. This allows the doors to open at one speed and close at another. 
 
diff --git a/Engine/Documentation/Source/Resources/Showcases/Showcases.INT.udn b/Engine/Documentation/Source/Resources/Showcases/Showcases.INT.udn
--- a/Engine/Documentation/Source/Resources/Showcases/Showcases.INT.udn
+++ b/Engine/Documentation/Source/Resources/Showcases/Showcases.INT.udn
@@ -2,7 +2,6 @@ Availability: Public
 Title:Engine Feature Examples
 Crumbs:%ROOT%, Resources
 Description:Full scenes demonstrating particular features or approaches to specific challenges.
-Version: 4.9
 
 [VAR:Topic]
 [OBJECT:Topic]
diff --git a/Engine/Shaders/Common.usf b/Engine/Shaders/Common.usf
--- a/Engine/Shaders/Common.usf
+++ b/Engine/Shaders/Common.usf
@@ -1317,7 +1317,7 @@ MaterialFloat GetBoxPushout(MaterialFloat3 Normal,MaterialFloat3 Extent)
 /** Output of the screen vertex shader. */
 struct FScreenVertexOutput
 {
-	MaterialFloat2 UV : TEXCOORD0;
+	noperspective MaterialFloat2 UV : TEXCOORD0;
 	float4 Position : SV_POSITION;
 };
 
diff --git a/Engine/Shaders/DistortApplyScreenPS.usf b/Engine/Shaders/DistortApplyScreenPS.usf
--- a/Engine/Shaders/DistortApplyScreenPS.usf
+++ b/Engine/Shaders/DistortApplyScreenPS.usf
@@ -29,7 +29,7 @@ float4		SceneColorRect;
 
 /** distorts screen texture using accumulated distortion offsets */
 void Main(
-	in float4 TexCoord: TEXCOORD0,
+	in noperspective float4 TexCoord: TEXCOORD0,
 	out float4 OutColor : SV_Target0
 	)
 {
@@ -61,7 +61,7 @@ void Main(
 }
 
 void Merge(
-	in float4 TexCoord: TEXCOORD0,
+	in noperspective float4 TexCoord: TEXCOORD0,
 	out float4 OutColor : SV_Target0
 	)
 {
diff --git a/Engine/Shaders/DownsampleDepthPixelShader.usf b/Engine/Shaders/DownsampleDepthPixelShader.usf
--- a/Engine/Shaders/DownsampleDepthPixelShader.usf
+++ b/Engine/Shaders/DownsampleDepthPixelShader.usf
@@ -11,7 +11,7 @@ float4 SourceTexelOffsets01;
 float4 SourceTexelOffsets23;
 
 void Main(
-	float2 InUV : TEXCOORD0,
+	noperspective float2 InUV : TEXCOORD0,
 	out float OutDepth : SV_DEPTH,
 	out float4 OutColor : SV_Target0
 	)
diff --git a/Engine/Shaders/FXAAShader.usf b/Engine/Shaders/FXAAShader.usf
--- a/Engine/Shaders/FXAAShader.usf
+++ b/Engine/Shaders/FXAAShader.usf
@@ -46,14 +46,13 @@ float2 fxaaQualityRcpFrame;
 void FxaaVS(
 	in float4 InPosition : ATTRIBUTE0,
 	in float2 UV : ATTRIBUTE1,
-	out float2 OutCenterUV : TEXCOORD0,
-	out float4 OutCornersUV : TEXCOORD1,
+	out noperspective float2 OutCenterUV : TEXCOORD0,
+	out noperspective float4 OutCornersUV : TEXCOORD1,
 	out float4 OutPosition : SV_POSITION
 	)
 {
 	DrawRectangle(InPosition, UV, OutPosition, OutCenterUV);
 	OutCornersUV = float4(OutCenterUV.xy - 0.5f * fxaaQualityRcpFrame, OutCenterUV.xy + 0.5f * fxaaQualityRcpFrame);
-
 }
 
 
@@ -86,7 +85,7 @@ void FxaaVS(
 	// include NVIDIA FXAA source
 	#include "Fxaa3_11.usf"
 
-	void FxaaPS(float2 TexCenter : TEXCOORD0, float4 TexCorners : TEXCOORD1, out float4 OutColor : SV_Target0)
+	void FxaaPS(noperspective float2 TexCenter : TEXCOORD0, noperspective float4 TexCorners : TEXCOORD1, out float4 OutColor : SV_Target0)
 	{
 		FxaaTex TextureAndSampler;
 		TextureAndSampler.tex = PostprocessInput0;
diff --git a/Engine/Shaders/FilterVertexShader.usf b/Engine/Shaders/FilterVertexShader.usf
--- a/Engine/Shaders/FilterVertexShader.usf
+++ b/Engine/Shaders/FilterVertexShader.usf
@@ -13,8 +13,8 @@ float4 SampleOffsets[(NUM_SAMPLES + 1) / 2];
 void Main(
 	float4 InPosition : ATTRIBUTE0,
 	float2 UV : ATTRIBUTE1,
-	out float2 OutUV : TEXCOORD0,
-	out float2 OutOffsetUVs[NUM_SAMPLES] : TEXCOORD1,
+	out noperspective float2 OutUV : TEXCOORD0,
+	out noperspective float2 OutOffsetUVs[NUM_SAMPLES] : TEXCOORD1,
 	out float4 OutPosition : SV_POSITION
 	)
 {
@@ -41,8 +41,8 @@ void Main(
 void Main(
 	float4 InPosition : ATTRIBUTE0,
 	float2 UV : ATTRIBUTE1,
-	out float2 OutUV : TEXCOORD0,
-	out float4 OutOffsetUVs[(NUM_SAMPLES + 1) / 2] : TEXCOORD1,
+	out noperspective float2 OutUV : TEXCOORD0,
+	out noperspective float4 OutOffsetUVs[(NUM_SAMPLES + 1) / 2] : TEXCOORD1,
 	out float4 OutPosition : SV_POSITION
 	)
 {
diff --git a/Engine/Shaders/HZBOcclusion.usf b/Engine/Shaders/HZBOcclusion.usf
--- a/Engine/Shaders/HZBOcclusion.usf
+++ b/Engine/Shaders/HZBOcclusion.usf
@@ -9,7 +9,7 @@ float2 InputViewportMaxBound;
 Texture2D		Texture;
 SamplerState	TextureSampler;
 
-void HZBBuildPS(float4 RenderTargetUVAndScreenPos : TEXCOORD0, out float4 OutColor : SV_Target0)
+void HZBBuildPS(noperspective float4 RenderTargetUVAndScreenPos : TEXCOORD0, out float4 OutColor : SV_Target0)
 {
 	float4 Depth;
 	float2 UV[4];
diff --git a/Engine/Shaders/PostProcessBloom.usf b/Engine/Shaders/PostProcessBloom.usf
--- a/Engine/Shaders/PostProcessBloom.usf
+++ b/Engine/Shaders/PostProcessBloom.usf
@@ -16,12 +16,19 @@ float4 ColorScale1;
 void MainPostprocessCommonVS(
 	in float4 InPosition : ATTRIBUTE0,
 	in float2 InTexCoord : ATTRIBUTE1,
-	out float4 OutTexCoord : TEXCOORD0,
+	out noperspective float4 OutTexCoord : TEXCOORD0,
 	out float4 OutPosition : SV_POSITION
 	)
 {
+	// hlslcc creates temporaries with noperspective qualifier, which is a bad thing
+#if COMPILER_GLSL || COMPILER_GLSL_ES2
+	float2 Temp;
+	DrawRectangle(InPosition, InTexCoord, OutPosition, Temp);
+	OutTexCoord = float4(Temp.x, Temp.y, OutPosition.x, OutPosition.y);
+#else
 	DrawRectangle(InPosition, InTexCoord, OutPosition, OutTexCoord.xy);
-	OutTexCoord.zw =  OutPosition.xy;
+	OutTexCoord.zw = OutPosition.xy;
+#endif
 }
 
 void VisualizeBloomSetupPS(
@@ -130,10 +137,10 @@ void VisualizeBloomOverlayPS(
 
 // vertex shader entry point
 void MainVS(
-	in float4 InPosition : ATTRIBUTE0,
-	in float2 InTexCoord : ATTRIBUTE1,
-	out float4 OutTexCoord : TEXCOORD0,
-	out float OutExposureScale : TEXCOORD1,
+	float4 InPosition : ATTRIBUTE0,
+	float2 InTexCoord : ATTRIBUTE1,
+	out noperspective float4 OutTexCoord : TEXCOORD0,
+	out nointerpolation float OutExposureScale : TEXCOORD1,
 	out float4 OutPosition : SV_POSITION
 	)
 {
diff --git a/Engine/Shaders/PostProcessBokehDOF.usf b/Engine/Shaders/PostProcessBokehDOF.usf
--- a/Engine/Shaders/PostProcessBokehDOF.usf
+++ b/Engine/Shaders/PostProcessBokehDOF.usf
@@ -47,7 +47,7 @@ float2 PixelToScreenPos(float2 PixelPos)
 void MainVS(
 	uint VId : SV_VertexID,
 	uint IId : SV_InstanceID,
-	out float2 OutTexCoord : TEXCOORD0,
+	out noperspective float2 OutTexCoord : TEXCOORD0,
 	out nointerpolation float4 OutColor : TEXCOORD1,
 	out float4 OutPosition : SV_POSITION
 	)
diff --git a/Engine/Shaders/PostProcessCompositeEditorPrimitives.usf b/Engine/Shaders/PostProcessCompositeEditorPrimitives.usf
--- a/Engine/Shaders/PostProcessCompositeEditorPrimitives.usf
+++ b/Engine/Shaders/PostProcessCompositeEditorPrimitives.usf
@@ -95,7 +95,7 @@ float4 GetScrColorAt(int2 PixelPos, out float OutDeviceZ)
 	return Ret;
 }
 
-void MainPS(float4 UVAndScreenPos : TEXCOORD0, out float4 OutColor : SV_Target0)
+void MainPS(noperspective float4 UVAndScreenPos : TEXCOORD0, out float4 OutColor : SV_Target0)
 {
 	float2 UV = UVAndScreenPos.xy;
 	float2 PixelPosF = UVAndScreenPos.zw * ScreenPosToPixel.xy + ScreenPosToPixel.zw + 0.5f;
diff --git a/Engine/Shaders/PostProcessDownsample.usf b/Engine/Shaders/PostProcessDownsample.usf
--- a/Engine/Shaders/PostProcessDownsample.usf
+++ b/Engine/Shaders/PostProcessDownsample.usf
@@ -11,7 +11,7 @@
 void MainDownsampleVS(
 	in float4 InPosition : ATTRIBUTE0,
 	in float2 UV : ATTRIBUTE1,
-	out float2 OutUV : TEXCOORD0,
+	out noperspective float2 OutUV : TEXCOORD0,
 	out float4 OutPosition : SV_POSITION
 	) 
 {
diff --git a/Engine/Shaders/PostProcessLensBlur.usf b/Engine/Shaders/PostProcessLensBlur.usf
--- a/Engine/Shaders/PostProcessLensBlur.usf
+++ b/Engine/Shaders/PostProcessLensBlur.usf
@@ -27,8 +27,8 @@ float2 PixelToScreenPos(float2 PixelPos)
 void MainVS(
 	uint VId : SV_VertexID,
 	uint IId : SV_InstanceID,
-	out float2 OutTexCoord : TEXCOORD0,
-	out float4 OutColor : TEXCOORD1,
+	out noperspective float2 OutTexCoord : TEXCOORD0,
+	out noperspective float4 OutColor : TEXCOORD1,
 	out float4 OutPosition : SV_POSITION
 	)
 {
diff --git a/Engine/Shaders/PostProcessMotionBlur.usf b/Engine/Shaders/PostProcessMotionBlur.usf
--- a/Engine/Shaders/PostProcessMotionBlur.usf
+++ b/Engine/Shaders/PostProcessMotionBlur.usf
@@ -165,8 +165,8 @@ void OverrideWithTestChart(float2 ScreenPos, inout float2 ObjectVelocity, inout
 void SetupVS(
 	in float4 InPosition : ATTRIBUTE0,
 	in float2 InTexCoord : ATTRIBUTE1,
-	out float2 OutUV: TEXCOORD0,
-	out float2 OutUVs[4]: TEXCOORD1,
+	out noperspective float2 OutUV: TEXCOORD0,
+	out noperspective float2 OutUVs[4]: TEXCOORD1,
 	out float4 OutPosition : SV_POSITION
 	)
 {
@@ -187,8 +187,8 @@ void SetupVS(
 // MRT0: velocity in normalized_motionblur_velocity and mask in 0:background..1:object (needed for soft masked)
 // MRT1: color and depth in alpha
 void SetupPS(
-	float2 UV : TEXCOORD0,				// UV into the full resolution source RT
-	float2 UVs[4] : TEXCOORD1,			// UV for high quality 4 samples with blurring
+	noperspective float2 UV : TEXCOORD0,				// UV into the full resolution source RT
+	noperspective float2 UVs[4] : TEXCOORD1,			// UV for high quality 4 samples with blurring
 	out float4 OutVelocity : SV_Target0,
 	out float4 OutColor : SV_Target1)
 {
@@ -645,7 +645,7 @@ void MainPS(noperspective float4 UVAndScreenPos : TEXCOORD0, float4 SvPosition :
 }
 
 
-void MainRecombinePS(float4 UVAndScreenPos : TEXCOORD0, out float4 OutColor : SV_Target0)
+void MainRecombinePS(noperspective float4 UVAndScreenPos : TEXCOORD0, out float4 OutColor : SV_Target0)
 {
 	float2 UV = UVAndScreenPos.xy;
 
diff --git a/Engine/Shaders/PostProcessTonemap.usf b/Engine/Shaders/PostProcessTonemap.usf
--- a/Engine/Shaders/PostProcessTonemap.usf
+++ b/Engine/Shaders/PostProcessTonemap.usf
@@ -94,10 +94,10 @@ float4 FringeUVParams;
 void MainVS(
 	in float4 InPosition : ATTRIBUTE0,
 	in float2 InTexCoord : ATTRIBUTE1,
-	out float2 OutTexCoord : TEXCOORD0,
-	out float3 OutExposureScaleVignette : TEXCOORD1,
-	out float4 OutGrainUV : TEXCOORD2,
-	out float4 OutFringe : TEXCOORD3,
+	out noperspective float2 OutTexCoord : TEXCOORD0,
+	out noperspective float3 OutExposureScaleVignette : TEXCOORD1,
+	out noperspective float4 OutGrainUV : TEXCOORD2,
+	out noperspective float4 OutFringe : TEXCOORD3,
 	out float4 OutPosition : SV_POSITION
 	)
 {
diff --git a/Engine/Shaders/PostProcessUpscale.usf b/Engine/Shaders/PostProcessUpscale.usf
--- a/Engine/Shaders/PostProcessUpscale.usf
+++ b/Engine/Shaders/PostProcessUpscale.usf
@@ -62,7 +62,7 @@ float2 PaniniProjectionScreenPos(float2 InScreenPos)
 // vertex shader entry point
 void MainVS(
 	in uint VertexId : SV_VertexID,
-	out float4 OutTexCoord : TEXCOORD0,
+	out noperspective float4 OutTexCoord : TEXCOORD0,
 	out float4 OutPosition : SV_POSITION
 	)
 {
diff --git a/Engine/Shaders/ReflectionEnvironmentShaders.usf b/Engine/Shaders/ReflectionEnvironmentShaders.usf
--- a/Engine/Shaders/ReflectionEnvironmentShaders.usf
+++ b/Engine/Shaders/ReflectionEnvironmentShaders.usf
@@ -444,7 +444,7 @@ SamplerState ScreenSpaceReflectionsSampler;
 
 /** Used to apply reflection capture contribution along with SSR to scene color. */
 void ReflectionApplyPS(
-	in float4 UVAndScreenPos : TEXCOORD0, 
+	in noperspective float4 UVAndScreenPos : TEXCOORD0, 
 	out float4 OutColor : SV_Target0)
 {
 	float2 UV = UVAndScreenPos.xy;
diff --git a/Engine/Source/Runtime/Engine/Private/Animation/AnimMontage.cpp b/Engine/Source/Runtime/Engine/Private/Animation/AnimMontage.cpp
--- a/Engine/Source/Runtime/Engine/Private/Animation/AnimMontage.cpp
+++ b/Engine/Source/Runtime/Engine/Private/Animation/AnimMontage.cpp
@@ -877,8 +877,8 @@ void FAnimMontageInstance::Stop(const FAlphaBlend& InBlendOut, bool bInterrupt)
 		bInterrupted = bInterrupt;
 	}
 
-	// if it hasn't stopped, stop now
-	if (IsStopped() == false)
+	// if desired weight is > 0.f, turn that off
+	if (Blend.GetDesiredValue() > 0.f)
 	{
 		// do not use default Montage->BlendOut 
 		// depending on situation, the BlendOut time can change 
@@ -1073,7 +1073,7 @@ bool FAnimMontageInstance::SetNextSectionID(int32 const & SectionID, int32 const
 
 void FAnimMontageInstance::OnMontagePositionChanged(FName const & ToSectionName) 
 {
-	if (bPlaying && IsStopped())
+	if (bPlaying && (Blend.GetDesiredValue() == 0.f))
 	{
 		UE_LOG(LogAnimation, Warning, TEXT("Changing section on Montage (%s) to '%s' during blend out. This can cause incorrect visuals!"),
 			*GetNameSafe(Montage), *ToSectionName.ToString());
@@ -1437,8 +1437,8 @@ void FAnimMontageInstance::Advance(float DeltaTime, struct FRootMotionMovementPa
 						}
 					}
 
-					// If current section is last one, check to trigger a blend out and if it hasn't stopped yet, see if we should stop
-					if (NextSectionIndex == INDEX_NONE && !IsStopped())
+					// If current section is last one, check to trigger a blend out.
+					if( NextSectionIndex == INDEX_NONE )
 					{
 						const float DeltaPosToEnd = bPlayingForward ? (CurrentSectionLength - PosInSection) : PosInSection;
 						const float DeltaTimeToEnd = DeltaPosToEnd / FMath::Abs(CombinedPlayRate);
@@ -1504,7 +1504,7 @@ void FAnimMontageInstance::Advance(float DeltaTime, struct FRootMotionMovementPa
 	}
 
 	// If this Montage has no weight, it should be terminated.
-	if (IsStopped() && (Blend.IsComplete()))
+	if ((Blend.GetDesiredValue()<= ZERO_ANIMWEIGHT_THRESH) && (Blend.IsComplete()))
 	{
 		// nothing else to do
 		Terminate();
diff --git a/Engine/Source/Runtime/WebBrowser/Private/WebBrowserHandler.cpp b/Engine/Source/Runtime/WebBrowser/Private/WebBrowserHandler.cpp
--- a/Engine/Source/Runtime/WebBrowser/Private/WebBrowserHandler.cpp
+++ b/Engine/Source/Runtime/WebBrowser/Private/WebBrowserHandler.cpp
@@ -300,11 +300,6 @@ bool FWebBrowserHandler::OnBeforeResourceLoad(CefRefPtr<CefBrowser> Browser, Cef
 
 void FWebBrowserHandler::OnRenderProcessTerminated(CefRefPtr<CefBrowser> Browser, TerminationStatus Status)
 {
-	TSharedPtr<FWebBrowserWindow> BrowserWindow = BrowserWindowPtr.Pin();
-	if (BrowserWindow.IsValid())
-	{
-		BrowserWindow->OnRenderProcessTerminated(Status);
-	}
 }
 
 bool FWebBrowserHandler::OnBeforeBrowse(CefRefPtr<CefBrowser> Browser,
diff --git a/Engine/Source/Runtime/WebBrowser/Private/WebBrowserSingleton.cpp b/Engine/Source/Runtime/WebBrowser/Private/WebBrowserSingleton.cpp
--- a/Engine/Source/Runtime/WebBrowser/Private/WebBrowserSingleton.cpp
+++ b/Engine/Source/Runtime/WebBrowser/Private/WebBrowserSingleton.cpp
@@ -139,10 +139,6 @@ FWebBrowserSingleton::FWebBrowserSingleton()
 	CefSettings Settings;
 	Settings.no_sandbox = true;
 	Settings.command_line_args_disabled = true;
-
-	FString CefLogFile(FPaths::Combine(*FPaths::GameLogDir(), TEXT("cef3.log")));
-	CefLogFile = FPaths::ConvertRelativePathToFull(CefLogFile);
-	CefString(&Settings.log_file) = *CefLogFile;
 	Settings.log_severity = LOGSEVERITY_WARNING;
 
 	// Specify locale from our settings
diff --git a/Engine/Source/Runtime/WebBrowser/Private/WebBrowserWindow.cpp b/Engine/Source/Runtime/WebBrowser/Private/WebBrowserWindow.cpp
--- a/Engine/Source/Runtime/WebBrowser/Private/WebBrowserWindow.cpp
+++ b/Engine/Source/Runtime/WebBrowser/Private/WebBrowserWindow.cpp
@@ -547,17 +547,6 @@ void FWebBrowserWindow::OnResetDialogState()
 	OnDismissAllDialogs().ExecuteIfBound();
 }
 
-void FWebBrowserWindow::OnRenderProcessTerminated(CefRequestHandler::TerminationStatus Status)
-{
-	if(bRecoverFromRenderProcessCrash)
-	{
-		bRecoverFromRenderProcessCrash = false;
-		NotifyDocumentError(); // Only attempt a single recovery at a time
-	}
-
-	bRecoverFromRenderProcessCrash = true;
-	Reload();
-}
 
 FReply FWebBrowserWindow::OnMouseButtonDown(const FGeometry& MyGeometry, const FPointerEvent& MouseEvent, bool bIsPopup)
 {
@@ -840,14 +829,6 @@ void FWebBrowserWindow::NotifyDocumentLoadingStateChange(bool IsLoading)
 	if (! IsLoading)
 	{
 		bIsInitialized = true;
-
-		if (bRecoverFromRenderProcessCrash)
-		{
-			bRecoverFromRenderProcessCrash = false;
-			// Toggle hidden/visible state to get OnPaint calls from CEF.
-			SetIsHidden(true);
-			SetIsHidden(false);
-		}
 	}
 	
 	EWebBrowserDocumentState NewState = IsLoading
diff --git a/Engine/Source/Runtime/WebBrowser/Private/WebBrowserWindow.h b/Engine/Source/Runtime/WebBrowser/Private/WebBrowserWindow.h
--- a/Engine/Source/Runtime/WebBrowser/Private/WebBrowserWindow.h
+++ b/Engine/Source/Runtime/WebBrowser/Private/WebBrowserWindow.h
@@ -298,11 +298,6 @@ private:
 	 */
 	void OnResetDialogState();
 
-	/**
-	 * Called when render process was terminated abnormally.
-	 */
-	void OnRenderProcessTerminated(CefRequestHandler::TerminationStatus Status);
-
 
 	/** Specifies if window creation functionality is available. 
 	 *
@@ -508,9 +503,6 @@ private:
 	FIntPoint PopupPosition;
 	bool bShowPopupRequested;
 
-	/** This is set to true when reloading after render process crash. */
-	bool bRecoverFromRenderProcessCrash;
-
 	/** Handling of passing and marshalling messages for JS integration is delegated to a helper class*/
 	TSharedPtr<FWebJSScripting> Scripting;
 };
