Update to coding standards.

diff --git a/Engine/Documentation/Source/Programming/Development/CodingStandard/CodingStandard.INT.udn b/Engine/Documentation/Source/Programming/Development/CodingStandard/CodingStandard.INT.udn
--- a/Engine/Documentation/Source/Programming/Development/CodingStandard/CodingStandard.INT.udn
+++ b/Engine/Documentation/Source/Programming/Development/CodingStandard/CodingStandard.INT.udn
@@ -2,7 +2,7 @@ Availability:Public
 Title:Coding Standard
 Crumbs:%ROOT%, Programming, Programming\Development
 Description:Standards and conventions used by Epic Games in the Unreal Engine 4 codebase.
-Version: 4.9
+Version: 4.11
 
 [TOC(start:2)]
 
@@ -14,10 +14,12 @@ Code conventions are important to programmers for a number of reasons:
 
 * 80% of the lifetime cost of a piece of software goes to maintenance.
 * Hardly any software is maintained for its whole life by the original author.
-* Code conventions improve the readability of the software, allowing engineers to understand new code more quickly and thoroughly. We will certainly be hiring new engineers and interns over the life of this project, and we will likely be using engine modifications we make new on our next few projects.
+* Code conventions improve the readability of the software, allowing engineers to understand new code more quickly and thoroughly.
 * If we decide to expose source code to mod community developers, we want it to be easily understood.
 * Many of these conventions are actually required for cross-compiler compatibility.
 
+The coding standards below are C++-centric, but the spirit of the standards are expected to be followed no matter which language is used.  A section may provide equivalent rules or exceptions for specific languages where it's applicable.
+
 
 ## Class Organization
 
@@ -34,14 +36,22 @@ If this line is missing or not formatted properly, CIS will generate an error an
 
 ## Naming Conventions
 
-* The first letter of each word in a name (e.g. type or variable) is capitalized, and there is usually no underscore between words.  For example, Health and **UPrimitiveComponent**, but not lastMouseCoordinates or delta_coordinates.
+* The first letter of each word in a name (e.g. type or variable) is capitalized, and there is usually no underscore between words.  For example, Health and UPrimitiveComponent, but not lastMouseCoordinates or delta_coordinates.
 * Type names are prefixed with an additional upper-case letter to distinguish them from variable names.  For example, FSkin is a type name, and Skin is an instance of a FSkin.
     * Template classes are prefixed by T.
-    * Classes that inherit from **UObject** are prefixed by U.
-    * Classes that inherit from **AActor** are prefixed by A.
-    * Classes that inherit from **SWidget** are prefixed by S.
+    * Classes that inherit from UObject are prefixed by U.
+    * Classes that inherit from AActor are prefixed by A.
+    * Classes that inherit from SWidget are prefixed by S.
     * Classes that are abstract interfaces are prefixed by I.
+    * Enums are prefixed by E.
+    * Boolean variables must be prefixed by b (e.g. "bPendingDestruction", or "bHasFadedIn").
     * Most other classes are prefixed by F, though some subsystems use other letters.
+    * Typedefs should be prefixed by whatever is appropriate for that type: F if it's a typedef of a struct, U if it's a typedef of a UObject etc.
+        * A typedef of a particular template instantiation is no longer a template and should be prefixed accordingly, e.g.:
+
+			typedef TArray<FMyType> FArrayOfMyTypes;
+    * Prefixes are omitted in C#.
+    * UnrealHeaderTool requires the correct prefixes in most cases, so it's important to provide them.
 * Type and variable names are nouns.
 * Method names are verbs that describe the method's effect, or describe the return value of a method that has no effect.
 
@@ -50,12 +60,14 @@ Variable, method, and class names should be clear, unambiguous, and descriptive.
 
 All variables should be declared one at a time so that a comment on the meaning of the variable can be provided. Also, the JavaDocs style requires it. You can use multi-line or single line comments before a variable, and the blank line is optional for grouping variables.
 
-All functions that return a bool should ask a true/false question, such as IsVisible() or ShouldClearBuffer().  All boolean variables must be prefixed with a "b" character (e.g. bPendingDestruction, or bHasFadedIn).
+All functions that return a bool should ask a true/false question, such as IsVisible() or ShouldClearBuffer().
 
 A procedure (a function with no return value) should use a strong verb followed by an Object. An exception is if the Object of the method is the Object it is in; then the Object is understood from context. Names to avoid include those beginning with "Handle" and "Process"; the verbs are ambiguous.
 
 Though not required, we encourage you to prefix function parameter names with "Out" if they are passed by reference and the function is expected to write to that value.  This makes it obvious that the value passed in this argument will be replaced by the function.
 
+If an In or Out parameter is also a boolean, put the b before the In/Out prefix, e.g. "bOutResult".
+
 Functions that return a value should describe the return value; the name should make clear what value the function will return. This is particularly important for boolean functions. Consider the following two example methods:
 
 		bool CheckTea(FTea Tea) {...} // what does true mean?
@@ -98,7 +110,7 @@ Functions that return a value should describe the return value; the name should
 * PTRINT for an integer that may hold a pointer (NEVER assume the size of PTRINT).
 
 
-Do not use the C+\+ int type in portable code, since it is dependent on the compiler how large it is.
+Use of C+\+'s int and unsigned int types - whose size may vary across platforms - is acceptable in code where the integer width is unimportant.  Explicitly-sized types must still be used in serialized or replicated formats.
 
 
 ## Comments
@@ -120,7 +132,7 @@ Comments are communication; communication is vital. Some things to keep in mind
 * Write useful comments:
 
 		// Bad:
-		// increment iLeaves
+		// increment Leaves
 		++Leaves;
 
 		// Good:
@@ -143,7 +155,7 @@ Comments are communication; communication is vital. Some things to keep in mind
 * Do not contradict the code:
 
 		// Bad:
-		// never increment iLeaves!
+		// never increment Leaves!
 		++Leaves;
 		
 		// Good:
@@ -151,6 +163,57 @@ Comments are communication; communication is vital. Some things to keep in mind
 		++Leaves;
 
 
+### Const correctness
+
+const is documentation as much as it is a compiler directive, so all code should strive to be const-correct.
+
+This includes passing function arguments by const pointer or reference if those arguments are not intended to be modified by the function, flagging methods as const if they do not modify the object, and using const iteration over containers if the loop isn't intended to modify the container.
+
+		void SomeMutatingOperation(FThing& OutResult, const TArray<int32>& InArray); // InArray will not be modified by SomeMutatingOperation, but OutResult probably will be
+
+		void FThing::SomeNonMutatingOperation() const
+		{
+			// This code will not modify the FThing it is invoked on
+		}
+
+		TArray<FString> StringArray;
+		for (const FString& : StringArray)
+		{
+			// The body of this loop will not modify StringArray
+		}
+
+Const should also be preferred on by-value function parameters and locals. This tells a reader that the variable will not be modified in the body of the function, which helps understandability. If you do this, ensure that the declaration and the definition match, as this can affect the JavaDoc process:
+
+		void AddSomeThings(const int32 Count);
+
+		void AddSomeThings(const int32 Count)
+		{
+			const int32 CountPlusOne = Count + 1;
+
+			// Neither Count nor CountPlusOne can be changed during the body of the function
+		}
+
+One exception to this is pass-by-value parameters which will ultimately be moved into a container (see "Move semantics"); this should be rare.
+
+		void FBlah::SetMemberArray(TArray<FString> InNewArray)
+		{
+			MemberArray = MoveTemp(InNewArray);
+		}
+
+Put the const keyword on the end when making a pointer itself const (rather than what it points to). References cannot be 'reassigned' anyway, and so cannot be made const in the same way:
+
+		// Const pointer to non-const object - pointer cannot be resassigned, but T can still be modified
+		T* const Ptr = ...;
+
+		// Illegal
+		T& const Ref = ...;
+
+Never use const on a return type, as this inhibits move semantics for complex types and will give compile warnings for built-in types:
+
+		// Bad
+		const TArray<FString> GetSomeArray();
+
+
 ### Example Formatting
 
 We use a system based on JavaDoc to automatically extract comments from the code and build documentation, so there are some specific comment formatting rules to follow.
@@ -159,24 +222,24 @@ The following example demonstrates the format of class, state, method, and varia
 
 Note that two different parameter comment styles are supported, as embodied by the Steep and Sweeten methods.  The @param style used by Steep is the traditional style, but for simple functions it can be clearer to integrate the parameter documentation into the descriptive comment for the function, as in the Sweeten example.
 
-Unlike the UE3 coding standard, the method comments should only be included once, where the method is publically declared.  The method comments should only contain information relevant to callers of the method, including any information about overrides of the method that may be relevant to the caller.  Details about the implementation of the method and its overrides that are not relevant to callers should be commented within the method implementation.
+Method comments should only be included once, where the method is publically declared.  The method comments should only contain information relevant to callers of the method, including any information about overrides of the method that may be relevant to the caller.  Details about the implementation of the method and its overrides that are not relevant to callers should be commented within the method implementation.
 
 
 		class IDrinkable
 		{
 		public:
-
+			/**
 			 * Called when a player drinks this object.
 			 * @param OutFocusMultiplier - Upon return, will contain a multiplier to apply to the drinker's focus.
 			 * @param OutThirstQuenchingFraction - Upon return, will contain the fraction of the drinker's thirst to quench (0-1).
 			 */
-			virtual void Drink(float& OutFocusMultiplier,float& OutThirstQuenchingFraction) = 0;
+			virtual void Drink(float& OutFocusMultiplier, float& OutThirstQuenchingFraction) = 0;
 		};
 
 		class FTea : public IDrinkable
 		{
 		public:
-
+			/**
 			 * Calculate a delta-taste value for the tea given the volume and temperature of water used to steep.
 			 * @param VolumeOfWater - Amount of water used to brew in mL
 			 * @param TemperatureOfWater - Water temperature in Kelvins
@@ -184,19 +247,19 @@ Unlike the UE3 coding standard, the method comments should only be included once
 			 * @return the change in intensity of the tea in tea taste units (TTU) per minute
 			 */
 			float Steep(
-				float VolumeOfWater,
-				float TemperatureOfWater,
+				const float VolumeOfWater,
+				const float TemperatureOfWater,
 				float& OutNewPotency
 				);
 
-			void Sweeten(float EquivalentGramsOfSucrose);
+			void Sweeten(const float EquivalentGramsOfSucrose);
 
 			float GetPrice() const
 			{
 				return Price;
 			}
 
-			virtual void Drink(float& OutFocusMultiplier,float& OutThirstQuenchingFraction);
+			virtual void Drink(float& OutFocusMultiplier, float& OutThirstQuenchingFraction) override;
 
 		private:
 
@@ -205,17 +268,17 @@ Unlike the UE3 coding standard, the method comments should only be included once
 			float Sweetness;
 		};
 
-		float FTea::Steep(float VolumeOfWater,float TemperatureOfWater,float& OutNewPotency)
+		float FTea::Steep(const float VolumeOfWater, const float TemperatureOfWater, float& OutNewPotency)
 		{
 			...
 		}
 
-		void FTea::Sweeten(float EquivalentGramsOfSucrose)
+		void FTea::Sweeten(const float EquivalentGramsOfSucrose)
 		{
 			...
 		}
 
-		void FTea::Drink(float& OutFocusMultiplier,float& OutThirstQuenchingFraction)
+		void FTea::Drink(float& OutFocusMultiplier, float& OutThirstQuenchingFraction)
 		{
 			...
 		}
@@ -230,36 +293,50 @@ What do all those parts of the method comment mean?
 * Then comes the return comment; it documents the expected return value just as an output variable is documented.
 
 
-## C++ 11 and Modern Language Syntax
+## C++11 and Modern Language Syntax
 
 Unreal Engine is built to be massively portable to many C++ compilers, so we are careful to use features that are compatible with the compilers we can imagine supporting.  Sometimes features are so useful that we will wrap them up in macros and use them pervasively, but usually we will wait until all of the compilers we could imagine supporting are up to the latest standard. 
 
-We are utilizing certain C++ 11 language features that appear to be well-supported across modern compilers, such as the "auto" keyword, range-based-for and lambdas.  In some cases, we are able to wrap up usage of these features in preprocessor conditionals (such as rvalue references in containers.)  However, certain language features we may opt to avoid entirely until we are confident we will not be surprised by a new platform appearing that cannot digest the syntax.
+We are utilizing certain C++11 language features that appear to be well-supported across modern compilers, such as range-based-for, move semantics and lambdas.  In some cases, we are able to wrap up usage of these features in preprocessor conditionals (such as rvalue references in containers.)  However, certain language features we may opt to avoid entirely until we are confident we will not be surprised by a new platform appearing that cannot digest the syntax.
 
 Unless specified below as a modern C++ compiler feature we are supporting, you should refrain from using compiler-specific language features unless they are wrapped in preprocessor macros or conditionals and used sparingly.
 
 
-### New Keywords for Old Macros
+### static_assert
+
+This keyword is valid for use where you need a compile-time assertion.
+
+
+### override and final
+
+These keywords are valid for use, and their use is strongly encouraged.  There will likely be many places where these have been omitted but they will be fixed over time.
 
-The old macros 'checkAtCompileTime', 'OVERRIDE', 'FINAL', and 'NULL' can now be replaced with 'static_assert', 'override', 'final', and 'nullptr'.  Some of these macros may remain defined as their use is still prevalent.
+
+### nullptr
+
+nullptr should be used instead of the C-style NULL macro in all cases.
 
 One exception to this is that nullptr in C++/CX builds (e.g. Xbox One) is actually the managed null reference type.  It is mostly compatible with nullptr from native C++ except in its type and some template instantiation contexts, and so you should use the TYPE_OF_NULLPTR macro instead of the more usual decltype(nullptr) for compatibility.
 
 
 ### The 'auto' Keyword
 
-The 'auto' keyword is supported by all compilers UE4 targets and you are encouraged to use it in your code where it makes sense to do so.
+You should not use auto in C++ code, barring a few exceptions below.  You must always be explicit about the type you're initializing.  This means that the type must be plainly visible to the reader.  This rule also applies to the use of the ‘var’ keyword in C#.
+
+When is it acceptable to use auto?
 
-Remember you can and should use const, & or * with auto just like you would with the type name.  With auto, this will coerce the inferred type to be what you want.  
+* When you need to bind a lambda to a variable, as lambda types are not expressible in code.
+* For iterator variables, but only where the iterator's type is very verbose and would impair readability.
+* In template code, where the type of an expression cannot easily be discerned.  This is an advanced case.
 
-We encourage use of the auto keyword for iterator loops (eliminates boilerplate) - or better still: ranged-based for loops - and also when you are initializing a variable to a new instance (eliminates a redundant type name.)  Some of the other uses are more contentious, but feel free to use it however you want for now and we can learn and improve best practices as we go.
+It's very important that types are clearly visible to someone who is reading the code.  Even though some IDEs are able to infer the type, this relies on the code being in a compilable state.  It also will not assist users of merge/diff tools, or when viewing individual source files in isolation, such as on GitHub.
 
-Tip:  If you hover over a variable in Visual Studio, it will usually tell you the inferred type.
+If you're sure you are using auto in an acceptable way, always remember to correctly use const, & or * just like you would with the type name.  With auto, this will coerce the inferred type to be what you want.
 
 
 ### Range Based For
 
-This is allowed in all engine and editor code and encouraged where it can help to keep the code easier to understand and more maintainable.  When migrating code that uses old TMap iterators, be aware that the old Key() and Value() functions which were methods of the iterator type are now simply Key and Value fields of the underlying key-value TPair:
+This is preferred to keep the code easier to understand and more maintainable.  When migrating code that uses old TMap iterators, be aware that the old Key() and Value() functions which were methods of the iterator type are now simply Key and Value fields of the underlying key-value TPair:
 
 		TMap<FString, int32> MyMap;
 
@@ -270,7 +347,7 @@ This is allowed in all engine and editor code and encouraged where it can help t
 		}
 
 		// New style
-		for (auto& Kvp : MyMap)
+		for (TPair<FString, int32>& Kvp : MyMap)
 		{
 			UE_LOG(LogCategory, Log, TEXT("Key: %s, Value: %d"), Kvp.Key, *Kvp.Value);
 		}
@@ -293,9 +370,7 @@ We also have range replacements for some standalone iterator types:
 
 ### Lambdas and Anonymous Functions
 
-Lambdas are now allowed, however we are cautious about use of stateful lambdas that capture stack variables -- we are still learning about where those are appropriate.  Also, stateful lambdas cannot be assigned to function pointers which we tend to use a lot.
-
-Lambdas are best used in conjunction with predicates in generic algorithms, for example:
+Lambdas can now be used on all compilers, but their usage should be considered. The best lambdas should be no more than a couple of statements in length, particularly when used as part of a larger expression or statement, for example as a predicate in a generic algorithm:
 
 		// Find first Thing whose name contains the word "Hello"
 		Thing* HelloThing = ArrayOfThings.FindByPredicate([](const Thing& Th){ return Th.GetName().Contains(TEXT("Hello")); });
@@ -303,12 +378,14 @@ Lambdas are best used in conjunction with predicates in generic algorithms, for
 		// Sort array in reverse order of name
 		AnotherArray.Sort([](const Thing& Lhs, const Thing& Rhs){ return Lhs.GetName() > Rhs.GetName(); });
 
-We expect to update this documentation in the future as we establish best practices.
+Be mindful of stateful lambdas that capture stack variables and whether you capture by value or reference -- a by-value capture of a non-trivial type could be a performance problem. Also, stateful lambdas cannot be assigned to function pointers, which we tend to use a lot.
+
+Documentation of non-trivial lambdas and anonymous functions should be considered in the same way a regular function is documented. Don't be afraid to split them over a few more lines in order to comment them.
 
 
-## Strongly-Typed Enums
+### Strongly-Typed Enums
 
-Enum classes are supported by all compilers and encouraged as a replacement for old-style namespaced enums, both for regular enums and UENUMs.  For example:
+Enum classes should always be used as a replacement for old-style namespaced enums, both for regular enums and UENUMs.  For example:
 
 		// Old enum
 		UENUM()
@@ -343,6 +420,7 @@ Enum classes used as flags can take advantage of a new ENUM_CLASS_FLAGS(EnumType
 
 		enum class EFlags
 		{
+			None  = 0x00,
 			Flag1 = 0x01,
 			Flag2 = 0x02,
 			Flag3 = 0x04
@@ -350,18 +428,18 @@ Enum classes used as flags can take advantage of a new ENUM_CLASS_FLAGS(EnumType
 
 		ENUM_CLASS_FLAGS(EFlags)
 
-The one exception to this is the use of flags in a 'truth' context - this is a limitation of the language and can be worked around using the 'double bang' operator:
+The one exception to this is the use of flags in a 'truth' context - this is a limitation of the language.  Instead, all flags enums should have an enumerator called 'None' which is set to 0 for comparisons:
 
 		// Old
 		if (Flags & EFlags::Flag1)
 
 		// New
-		if (!!(Flags & EFlags::Flag1))
+		if ((Flags & EFlags::Flag1) != EFlags::None)
 
 
-## Move Semantics
+### Move Semantics
 
-All of the main container types - TArray, TMap, TSet, FString - have move constructors and move assignment operators.  These are often used automatically when passing/returning these types by value, but can be explicit invoked by using MoveTemp, which is UE4's equivalent of std::move.
+All of the main container types - TArray, TMap, TSet, FString - have move constructors and move assignment operators.  These are often used automatically when passing/returning these types by value, but can be explicitly invoked by using MoveTemp, which is UE4's equivalent of std::move.
 
 Returning containers or strings by value can be a win for expressivity without the usual cost of temporary copies.  Rules around pass-by-value and use of MoveTemp are still being established, but can already be found in some optimized areas of the codebase.
 
@@ -390,12 +468,19 @@ Any third party code included in the engine should be marked with comments forma
 
 Brace wars are foul. Epic has a long standing usage pattern of putting braces on a new line. Please continue to adhere to that.
 
+Always include braces in single-statement blocks, e.g.:
+
+		if (bThing)
+		{
+			return;
+		}
+
 
 ### If - Else
 
 Each block of execution in an if-else statement should be in braces. This is to prevent editing mistakes - when braces are not used, someone could unwittingly add another line to an if block. The line would not be controlled by the if expression, which would be bad. Worse yet is when conditionally compiled items cause if/else statements to break. So always use braces.
 
-		if(HaveUnrealLicense)
+		if (HaveUnrealLicense)
 		{
 			InsertYourGameHere();
 		}
@@ -406,17 +491,17 @@ Each block of execution in an if-else statement should be in braces. This is to
 
 A multi-way if statement should be indented with each else if indented the same amount as the first if; this makes the structure clear to a reader:
 
-		if(TannicAcid < 10)
+		if (TannicAcid < 10)
 		{
-			Log("Low Acid");
+			UE_LOG(LogCategory, Log, TEXT("Low Acid"));
 		}
-		else if(TannicAcid < 100)
+		else if (TannicAcid < 100)
 		{
-			Log("Medium Acid");
+			UE_LOG(LogCategory, Log, TEXT("Medium Acid"));
 		}
 		else
 		{
-			Log("High Acid");
+			UE_LOG(LogCategory, Log, TEXT("High Acid"));
 		}
 
 
@@ -467,56 +552,6 @@ You can use namespaces to organize your classes, functions and variables where a
 * Note that forward-declared types need to be declared within their respective namespace, otherwise you will get link errors.
 * If you declare a lot of classes/types within a namespace, it can make it difficult to use those types in other global-scoped classes. (e.g. function signatures will need to use explicit namespace when appearing in class declarations.)
 * You can use "using" to alias only specific variables within a namespace into your scope (e.g. using Foo::FBar), but we do not usually do that in Unreal code.
-* We require C++ enum declarations to be [wrapped in a namespace](#C++EnumsandNamespaceScoping) so that the enum value names are not at global scope.
-
-
-## C++ Enums and Namespace Scoping
-
-* We always prefix enum types with an "E" character in Unreal Engine code.
-
-* We require that all enums use namespaces (or empty structs) for scoping.  The reason for this is that in C++, enum values are scoped to the same scope as the enum type itself.  This can cause naming collisions, which results in a programmer having to create strange names or prefixes for enum values to make their values appear unique.  Instead, we always explicitly scope new enums using namespaces.  The actual enum type name within the namespace should always be declared as "Type".
-
-* Example of a namespace-scoped enum:
-
-		namespace EColorChannel
-		{
-			enum Type
-			{
-				Red,
-				Green,
-				Blue
-			};
-		}
-
-		FString GetNameForColorChannel(const EColorChannel::Type ColorChannel)
-		{
-			switch(ColorChannel)
-			{
-				case EColorChannel::Red:   return TEXT("Red");
-				case EColorChannel::Green: return TEXT("Green");
-				case EColorChannel::Blue:  return TEXT("Blue");
-				default:                   return TEXT("Unknown");
-			}
-		}
-
-
-* Note that for locally-declared enums, you will not be able to use a namespace for scoping.  In these cases, we opt to declare a local struct with no member variables, only a local enum type and use that struct for scoping.
-
-		class FObjectMover
-		{
-        public:
-
-            struct EMoveDirection
-            {
-			    enum Type
-			    {
-				    Forward,
-				    Reverse,
-			    };
-            };
-
-            FObjectMover( const EMoveDirection::Type Direction );
-		}
 
 
 ## Physical Dependencies
@@ -536,6 +571,18 @@ You can use namespaces to organize your classes, functions and variables where a
 * Do not rely on being included through another header; include everything you need.
 * Modules have Private and Public source directories.  Any definitions that are needed by other modules must be in headers in the Public directory, but everything else should be in the Private directory.  Note that in older Unreal modules, these directories may still be called Src and Inc, but those directories are meant to separate private and public code in the same way, and are not meant to separate header files from source files.
 * Do not worry about setting up your headers for precompiled header generation.  UnrealBuildTool can do a better job of this than you can.
+* Split up large functions into logical sub-functions. One area of compilers' optimizations is the elimination of common subexpressions, and the bigger your functions are, the more work the compiler has to do to identity them, leading to greatly inflated build times.
+* Use inline functions judiciously, as they force rebuilds even in files which don't use them. Inlining should only be used for trivial accessors and when profiling shows there is a benefit to doing so.
+* Be even more conservative in the use of FORCEINLINE. All code and local variables will be expanded out into the calling function and will cause the same build time problems caused by large functions.
+
+
+## Encapsulation
+
+Enforce encapsulation with the protection keywords. Class members should almost always be declared private unless they are part of the public/protected interface to the class. Use your best judgement, but always be aware that a lack of accessors makes it hard to refactor later without breaking plugins and existing projects.
+
+If particular fields are only intended to be usable by derived classes, make them private and provide protected accessors.
+
+Use final if your class is not designed to be derived from.
 
 
 ## General Style Issues
@@ -546,10 +593,11 @@ You can use namespaces to organize your classes, functions and variables where a
 * Address compiler warnings. Compiler warning messages mean something is not as it should be. Fix what the compiler is complaining about. If you absolutely cannot address it, use #pragma to suppress the warning; this is a remedy of last resort.
 * Leave a blank line at the end of the file. All .cpp and .h files should include a blank line to play nice with gcc.
 * Never allow float to implicit convert to int32. This is a slow operation, and does not compile on all compilers. Instead, always use the appTrunc() function to convert to int32. This will ensure cross-compiler compatibility as well as generate faster code.
-* Enforce encapsulation with the protection keywords. Class members should be declared private unless they are part of the public interface to the class.
 * Interface classes (prefixed with "I") should always be abstract and must not have member variables.  Interfaces are allowed to contain methods that are not pure-virtual, and even methods that are non-virtual or static, as long as they are implemented inline.
-* Use const wherever possible. Particularly on reference parameters and class methods. const is documentation as much as it is a compiler directive.
 * Debug code should either be generally useful and polished, or not checked in. Debug code intermixed with other code makes the other code much harder to read.
+* Always use the TEXT() macro around string literals. Without it, code which constructs FStrings from literals will cause an undesirable string conversion process.
+* Avoid repeating the same operation redundantly in loops. Hoist common subexpressions out of loops to avoid redundant calculations. Make use of statics in some cases to avoid globally-redundant operations across function calls, e.g. constructing an FName from a string literal.
+* Be mindful of hot reload. Minimize dependencies to cut down on iteration time. Don't use inlining or templates for functions which are likely to change over a reload. Only use statics for things which are expected to remain constant over a reload.
 * Use intermediate variables to simplify complicated expressions. If you have a complicated expression, it can be easier to understand if you split it into sub-expressions that are assigned to intermediate variables with names describing the meaning of the sub-expression within the parent expression.  For example:
 
 		if ((Blah->BlahP->WindowExists->Etc && Stuff) &&
@@ -563,13 +611,13 @@ You can use namespaces to organize your classes, functions and variables where a
 
 		const bool bIsLegalWindow = Blah->BlahP->WindowExists->Etc && Stuff;
 		const bool bIsPlayerDead = bPlayerExists && bGameStarted && bPlayerStillHasPawn && IsTuesday();
-		if(bIsLegalWindow && !bIsPlayerDead)
+		if (bIsLegalWindow && !bIsPlayerDead)
 		{
 			DoSomething();
 		}
 
 
-* Use the virtual and override keywords when declaring an overriding method. When declaring a virtual function in a derived class that overrides a virtual function in the parent class, you *must* use both the virtual and the override keywords. For example:
+* Use the virtual and override keywords when declaring an overriding method. When declaring a virtual function in a derived class that overrides a virtual function in the parent class, you must use both the virtual and the override keywords. For example:
 		
 		class A
 		{
@@ -595,3 +643,21 @@ Note that there is a lot of existing code that does not follow this yet due to t
 		
 		FShaderType *Type
 		FShaderType * Type
+
+
+* Shadowed variables are not allowed. C++ allows variables to be shadowed from an outer scope, making usage ambiguous to a reader. For example, there are three usable 'Count' variables in this member function:
+
+		class FSomeClass
+		{
+		public:
+			void Func(const int32 Count)
+			{
+				for (int32 Count = 0; Count != 10; ++Count)
+				{
+					// Use Count
+				}
+			}
+
+		private:
+			int32 Count;
+		};
