Clarification of the 'no const return types' rule, which is not applicable to pointers or references to const types.

diff --git a/Engine/Documentation/Source/Programming/Development/CodingStandard/CodingStandard.INT.udn b/Engine/Documentation/Source/Programming/Development/CodingStandard/CodingStandard.INT.udn
--- a/Engine/Documentation/Source/Programming/Development/CodingStandard/CodingStandard.INT.udn
+++ b/Engine/Documentation/Source/Programming/Development/CodingStandard/CodingStandard.INT.udn
@@ -208,11 +208,20 @@ Put the const keyword on the end when making a pointer itself const (rather than
 		// Illegal
 		T& const Ref = ...;
 
-Never use const on a return type, as this inhibits move semantics for complex types and will give compile warnings for built-in types:
+Never use const on a return type, as this inhibits move semantics for complex types and will give compile warnings for built-in types. This rule only applies to the return type itself, not the target type of a pointer or reference being returned:
 
-		// Bad
+		// Bad - returning a const array
 		const TArray<FString> GetSomeArray();
 
+		// Fine - returning a reference to a const array
+		const TArray<FString>& GetSomeArray();
+
+		// Fine - returning a pointer to a const array
+		const TArray<FString>* GetSomeArray();
+
+		// Bad - returning a const pointer to a const array
+		const TArray<FString>* const GetSomeArray();
+
 
 ### Example Formatting
 
