Merge doc changes from //depot/UE4 to try to fix CIS

diff --git a/Engine/Documentation/Source/Engine/Landscape/Editing/Editing.INT.udn b/Engine/Documentation/Source/Engine/Landscape/Editing/Editing.INT.udn
--- a/Engine/Documentation/Source/Engine/Landscape/Editing/Editing.INT.udn
+++ b/Engine/Documentation/Source/Engine/Landscape/Editing/Editing.INT.udn
@@ -2,6 +2,7 @@ Availability: Public
 Title:Editing Landscapes
 Crumbs: %ROOT%, Engine, Engine/Landscape
 Description:Guide to using the tools for editing Landscape terrains.
+Version: 4.9
 
 [VAR:Topic]
 [OBJECT:Topiccompact]
@@ -23,24 +24,23 @@ Description:Guide to using the tools for editing Landscape terrains.
 [/OBJECT]
 [/VAR]
 
-[TOC(start:2 end:2)]
 
 After you have created a Landscape, you can now do any of the following to it:
 
-* Modify its shape in **[Sculpt](Engine/Landscape\Editing\SculptMode)** mode.
-* Modify its appearance in **[Paint](Engine/Landscape\Editing\PaintMode)** mode.
-* Change the size, shape, and type of **[Brush](Engine/Landscape\Editing\Brushes)** you use to modify the Landscape's shape and appearance.
-* Modify the size and shape of its components in **[Manage](Engine/Landscape\Editing\ManageMode)** mode.
-* Copy, paste, import, and export parts of the Landscape's heightmap with **[Landscape Gizmos](Engine/Landscape\Editing\Gizmos)**.
-* Modify the Landscape's **[Material](Engine/Landscape\Materials)**.
-* Create **[Splines](Engine/Landscape\Editing\Splines)** for any linear features you want to use in your Landscape.
-* Create **[holes](Engine/Landscape\Materials)** in your Landscape for features such as caves.
+* Modify its shape in [Sculpt](Engine/Landscape\Editing\SculptMode) mode.
+* Modify its appearance in [Paint](Engine/Landscape\Editing\PaintMode) mode.
+* Change the size, shape, and type of [Brush](Engine/Landscape\Editing\Brushes) you use to modify the Landscape's shape and appearance.
+* Modify the size and shape of its components in [Manage](Engine/Landscape\Editing\ManageMode) mode.
+* Copy, paste, import, and export parts of the Landscape's heightmap with [Landscape Gizmos](Engine/Landscape\Editing\Gizmos).
+* Modify the Landscape's [Material](Engine/Landscape\Materials).
+* Create [Splines](Engine/Landscape\Editing\Splines) for any linear features you want to use in your Landscape.
+* Create [holes](Engine/Landscape\Materials) in your Landscape for features such as caves.
 
-For information about creating Landscapes, see **[](Engine/Landscape\Creation)**.
+For information about creating Landscapes, see [](Engine/Landscape\Creation).
 
-For a quick reference guide to the various Landscape tool modes and options, see **[](Engine/Landscape/UI)**.
+For a quick reference guide to the various Landscape tool modes and options, see [](Engine/Landscape/UI).
 
-For examples of the Landscape tool in action, see the **[](Resources\ContentExamples\Landscapes)**.
+For examples of the Landscape tool in action, see the [](Resources\ContentExamples\Landscapes).
 
 
 [REGION:topics third]
diff --git a/Engine/Documentation/Source/Engine/Landscape/Editing/PaintMode/PaintMode.INT.udn b/Engine/Documentation/Source/Engine/Landscape/Editing/PaintMode/PaintMode.INT.udn
--- a/Engine/Documentation/Source/Engine/Landscape/Editing/PaintMode/PaintMode.INT.udn
+++ b/Engine/Documentation/Source/Engine/Landscape/Editing/PaintMode/PaintMode.INT.udn
@@ -2,6 +2,7 @@ Availability: Public
 Title:Paint Mode
 Crumbs: %ROOT%, Engine, Engine/Landscape, Engine/Landscape/Editing
 Description:Guide to the Landscape tool's Paint mode.
+Version: 4.9
 
 [VAR:Topic]
 [OBJECT:Topiccompact]
@@ -27,7 +28,7 @@ Description:Guide to the Landscape tool's Paint mode.
 
 The tools in **Paint** mode enable you to modify the appearance of your Landscape by selectively applying Material layers to parts of your Landscape.
 
-For more information about Landscape Materials, see [**Landscape Materials**](Engine/Landscape\Materials).
+For more information about Landscape Materials, see [Landscape Materials](Engine/Landscape\Materials).
 
 
 ## Painting Tools
@@ -49,7 +50,7 @@ The Painting Tools have some common options:
 | **Negative Mask** | When checked, and when **Use Region as Mask** is also checked, the region selection acts as a mask, but the active area is comprised of the unselected region. |
 
 
-For more information about Landscape Material layers, see [**Layers**](#Layers), later on this page.
+For more information about Landscape Material layers, see [Layers](#Layers), later on this page.
 
 
 ### Paint
@@ -71,9 +72,14 @@ The Smooth tool smooths the layer weight. The strength determines the amount of
 
 **Layer Smoothing**
 
-![Landscape_Smooth_Layer_Before.jpg](Landscape_Smooth_Layer_Before.jpg)(w:900 h:500)
-
-![Landscape_Smooth_Layer_After.jpg](Landscape_Smooth_Layer_After.jpg)(w:900 h:500)
+[OBJECT:ComparisonSlider]
+	[PARAM:before]
+	![Landscape Smooth Layer Before](Landscape_Smooth_Layer_Before.jpg)(w:900 h:500)
+	[/PARAM]
+	[PARAM:after]
+	![Landscape Smooth Layer After](Landscape_Smooth_Layer_After.jpg)(w:900 h:500)
+	[/PARAM]
+[/OBJECT]
 
 | **Additional Option** | **Description** |
 | --- | --- |
@@ -108,7 +114,7 @@ The tool applies a noise filter to the layer weight. The strength determines the
 | --- | --- |
 | **Use Target Value** | If checked, blends the values of the noise being applied toward a target value. |
 | **Noise Mode** | Determines whether to apply all noise effects, only the noise effects that result in increasing the application of the layer, or only the noise effects that result in reducing the application of the layer. |
-| **Noise Scale** | The size of the perlin noise filter used. The noise filter is related to position and scale, which means if you do not change **Noise Scale**, the same filter is applied to the same position many times. |
+| **Noise Scale** | The size of the perlin noise filter used. The noise filter is related to position and scale, which means if you do not change Noise Scale, the same filter is applied to the same position many times. |
 
 
 ## Layers
@@ -121,7 +127,7 @@ The layers defined in the Landscape Material automatically populate the list of
 
 ![Landscape_Target.png](Landscape_Target.png)(convert:false)
 
-Whichever layer is selected is the one you can apply to the Landscape with the Painting Tools, according to the tools' options and settings, and to the [**brush**](Engine\Landscape\Editing\Brushes) you are using.
+Whichever layer is selected is the one you can apply to the Landscape with the Painting Tools, according to the tools' options and settings, and to the [brush](Engine\Landscape\Editing\Brushes) you are using.
 
 ![Landscape_Layers.jpg](Landscape_Layers.jpg)(w:900 h:389)
 
@@ -129,7 +135,7 @@ Whichever layer is selected is the one you can apply to the Landscape with the P
 Many of the Painting Tools are similar to the Sculpting Tools, and you use them similarly, but to manipulate the application of layers instead of the heightmap.
 [/REGION]
 
-You create layers in the Material itself. For more information about layers and Landscape Materials, see [**Landscape Materials**](Engine\Landscape\Materials).
+You create layers in the Material itself. For more information about layers and Landscape Materials, see [Landscape Materials](Engine\Landscape\Materials).
 
 
 ### Layer Info Objects
@@ -146,7 +152,7 @@ You can either create a layer info object from the layer itself, or reuse an exi
 
 **To create a layer info object:**
 
-1. Press the plus icon (![Create Layer Info](Engine/Landscape\Editing\icon_Landscape_CreateLayerInfo.png)) to the right of the Layer name.
+1. Press the plus icon to the right of the Layer name.
 1. Choose **Weight-Blended Layer (normal)** or **Non Weight-Blended Layer**.
 ![](WeightBlended_NonWeightBlended.png)
 1. Choose the location to save the layer info object.
@@ -177,8 +183,8 @@ The primary purpose of layer info objects is to act as a unique key for painted
 
 | Option | Description |
 | ------ | ----------- |
-| **Phys Material** | The [**Physical Material**](Engine/Physics/PhysicalMaterials) (if any) assigned to areas of Landscape where this layer is dominant. |
-| **Hardness** | The value used by the [**Erosion**](Engine/Landscape/Editing/SculptMode/#Erosion) tool. |
+| **Phys Material** | The [Physical Material](Engine/Physics/PhysicalMaterials) (if any) assigned to areas of Landscape where this layer is dominant. |
+| **Hardness** | The value used by the [Erosion](Engine/Landscape/Editing/SculptMode/#Erosion) tool. |
 
 ### Orphaned Layers
 
diff --git a/Engine/Documentation/Source/Engine/Paper2D/HowTo/Animation/2DCharacterAnimation.INT.udn b/Engine/Documentation/Source/Engine/Paper2D/HowTo/Animation/2DCharacterAnimation.INT.udn
--- a/Engine/Documentation/Source/Engine/Paper2D/HowTo/Animation/2DCharacterAnimation.INT.udn
+++ b/Engine/Documentation/Source/Engine/Paper2D/HowTo/Animation/2DCharacterAnimation.INT.udn
@@ -5,6 +5,10 @@ Description:How to set up and play different animations for a 2D character based
 Related: Engine\Paper2D\HowTo\TopDown
 Related: Engine\Paper2D\Sprites
 Related: Engine\Paper2D\Flipbooks
+Related: Engine/Animation/StateMachines/CreatingStateMachines
+Related: Engine/Animation/StateMachines
+Version: 4.9
+SkillLevel: Intermediate
 
 On this page you will learn how to set up an Animation State machine for a 2D character which will allow the character to cycle between different Flipbook animations based on your defined criteria. 
 
diff --git a/Engine/Documentation/Source/Engine/Paper2D/HowTo/Paper2DHowTo.INT.udn b/Engine/Documentation/Source/Engine/Paper2D/HowTo/Paper2DHowTo.INT.udn
--- a/Engine/Documentation/Source/Engine/Paper2D/HowTo/Paper2DHowTo.INT.udn
+++ b/Engine/Documentation/Source/Engine/Paper2D/HowTo/Paper2DHowTo.INT.udn
@@ -4,6 +4,8 @@ Crumbs: %ROOT%, Engine, Engine/Paper2D
 Description:The Paper 2D How To pages provides several short step-by-step guides for working with Paper 2D. 
 Related: Engine\Paper2D\Sprites
 Related: Engine\Paper2D\Flipbooks
+Version: 4.9
+SkillLevel: Intermediate
 
 The **Paper 2D - How-to** pages provide short step-by-step examples of ways to accomplish certain tasks within Paper 2D in Unreal Engine 4. While there may be other ways to accomplish these tasks, these samples should get you started and give you ideas on how you may want to approach your objectives. 
 
diff --git a/Engine/Documentation/Source/Engine/Paper2D/HowTo/Physics/SpritePhysics.INT.udn b/Engine/Documentation/Source/Engine/Paper2D/HowTo/Physics/SpritePhysics.INT.udn
--- a/Engine/Documentation/Source/Engine/Paper2D/HowTo/Physics/SpritePhysics.INT.udn
+++ b/Engine/Documentation/Source/Engine/Paper2D/HowTo/Physics/SpritePhysics.INT.udn
@@ -4,6 +4,8 @@ Crumbs: %ROOT%, Engine, Engine/Paper2D
 Description:Examples of implementing physics for 2D games. 
 Related: Engine\Paper2D\Sprites
 Related: Engine\Paper2D\Flipbooks
+Version: 4.9
+SkillLevel: Intermediate
 
 This page provides an example of applying physics to sprites in 2D games. 
 
diff --git a/Engine/Documentation/Source/Engine/Paper2D/HowTo/Sockets/SpriteSockets.INT.udn b/Engine/Documentation/Source/Engine/Paper2D/HowTo/Sockets/SpriteSockets.INT.udn
--- a/Engine/Documentation/Source/Engine/Paper2D/HowTo/Sockets/SpriteSockets.INT.udn
+++ b/Engine/Documentation/Source/Engine/Paper2D/HowTo/Sockets/SpriteSockets.INT.udn
@@ -4,6 +4,9 @@ Crumbs: %ROOT%, Engine, Engine/Paper2D
 Description:How to attach objects to sockets defined in a Paper 2D Sprite. 
 Related: Engine\Paper2D\Sprites
 Related: Engine\Paper2D\Flipbooks
+Version: 4.9
+SkillLevel: Intermediate
+
 
 
 On this page you will learn how to add sockets to sprites then access those sockets from within a Blueprint which you can attach an object to. You can spawn particle effects from sprite sockets (such as a character on fire), attach a weapon to a character's hand (like a gun or sword) or any number of different items that you want to spawn and attach to a sprite's socket. 
diff --git a/Engine/Documentation/Source/Engine/Paper2D/HowTo/TopDown/TopDown.INT.udn b/Engine/Documentation/Source/Engine/Paper2D/HowTo/TopDown/TopDown.INT.udn
--- a/Engine/Documentation/Source/Engine/Paper2D/HowTo/TopDown/TopDown.INT.udn
+++ b/Engine/Documentation/Source/Engine/Paper2D/HowTo/TopDown/TopDown.INT.udn
@@ -7,6 +7,8 @@ Related: Engine\Paper2D\TileMaps
 Related: Gameplay/HowTo/UsingCameras/Blueprints
 Related: Engine\Paper2D\Sprites
 Related: Engine\Paper2D\Flipbooks
+Version: 4.9
+SkillLevel: Intermediate
 
 On this page you will learn how to modify the existing **Third Person Template** project to create a **Top Down 2D** style game. 
 
@@ -35,7 +37,7 @@ Here we will set up our project for this tutorial, you can import your own asset
 	![](TopDown3.png)
 
 	[REGION:note]
-	For more information on how to set up assets for importing, see the [**Paper 2D Import Options**](Engine\Paper2D\Importing) documentation. 
+	For more information on how to set up assets for importing, see the [Paper 2D Import Options](Engine\Paper2D\Importing) documentation. 
 	[/REGION]
 
 1.	**Right-click** on the **SoldierSprites** asset and select **create Flipbooks**. 
@@ -138,7 +140,7 @@ The character is setup, however there is no movement script applied to it which
 There is no need to reinvent the wheel here, since this template comes with a character with scripted functionality for movement, we can copy that into our Paper Character to provide movement for it. 
 
 [REGION:note]
-For more detailed information on setting up character movement and setting up character movement from a blank project, see the [**Setting Up Character Movement in Blueprints**](Gameplay/HowTo/CharacterMovement/Blueprints) documentation. 
+For more detailed information on setting up character movement and setting up character movement from a blank project, see the [Setting Up Character Movement in Blueprints](Gameplay/HowTo/CharacterMovement/Blueprints) documentation. 
 [/REGION]
 
 1.	Inside the **ThirdPersonBP** folder in the **Content Browser**, open the **Blueprints** folder then the **ThirdPersonCharacter** Blueprint. 
diff --git a/Engine/Documentation/Source/Engine/Paper2D/Paper2D.INT.udn b/Engine/Documentation/Source/Engine/Paper2D/Paper2D.INT.udn
--- a/Engine/Documentation/Source/Engine/Paper2D/Paper2D.INT.udn
+++ b/Engine/Documentation/Source/Engine/Paper2D/Paper2D.INT.udn
@@ -2,6 +2,7 @@ Availability:Public
 Title: Paper 2D
 Crumbs: %ROOT%, Engine
 Description:Paper 2D is a sprite-based system for creating 2D and 2D/3D hybrid games entirely within Unreal Engine 4.
+Version: 4.9
 
 [REGION:banner]
 ![](paper2d_banner.png)
diff --git a/Engine/Documentation/Source/Engine/Paper2D/Sprites/Collision/SpriteCollision.INT.udn b/Engine/Documentation/Source/Engine/Paper2D/Sprites/Collision/SpriteCollision.INT.udn
--- a/Engine/Documentation/Source/Engine/Paper2D/Sprites/Collision/SpriteCollision.INT.udn
+++ b/Engine/Documentation/Source/Engine/Paper2D/Sprites/Collision/SpriteCollision.INT.udn
@@ -6,6 +6,7 @@ Related: Engine\Paper2D\Sprites\Collision
 Related: Engine\Paper2D\Flipbooks
 Related: Engine/Rendering/Materials
 Related: Engine/Content/Types/Textures
+Version: 4.9
 
 
 [TOC(start:2)]
@@ -24,8 +25,8 @@ different settings:
 | Collision Type | Description |
 | -------------- | ----------- |
 | **None** | No collision geometry is generated; use this for purely decorative sprites. |
-| **Use 2D Physics** | This is an *experimental* option to generate collision geometry for use with Box2D. See [**2D Collision**](#2dcollision) for more information on the current limitations. |
-| **Use 3D Physics** | Collision geometry will be generated for use with PhysX. The 2D collision geometry defined in the sprite will be extruded using the **Collision Thickness** units deep in the perpendicular axis to make 3D collision geometry. |
+| **Use 2D Physics** | This is an *experimental* option to generate collision geometry for use with Box2D. See [2D Collision](#2dcollision) for more information on the current limitations. |
+| **Use 3D Physics** | Collision geometry will be generated for use with PhysX. The 2D collision geometry defined in the sprite will be extruded using the Collision Thickness units deep in the perpendicular axis to make 3D collision geometry. |
 
 ## Collision Geometry
 
@@ -34,7 +35,7 @@ geometry. The following settings are available:
 
 | Type | Description |
 | ---- | ----------- |
-| **Source Bounding Box** | Uses the bounding box of the [**Source Region**](Engine/Paper2D/Sprites/RenderGeometry) for the Sprite. [REGION:none]![](collision_box.png)(w:192)[/REGION]  |
+| **Source Bounding Box** | Uses the bounding box of the [Source Region](Engine/Paper2D/Sprites/RenderGeometry) for the Sprite. [REGION:none]![](collision_box.png)(w:192)[/REGION]  |
 | **Tight Bounding Box** | Generates a tightened bounding box that excludes any fully transparent areas of the Sprite. This results in better collision in most cases. [REGION:none]![](collision_tight.png)(w:192)[/REGION] |
 | **Shrink Wwrapped** | (*Experimental*) Generates complex geoemetry the conforms to the opaque areas of the Sprite. This results in the most realistic collision, but the added geometry could result in a performance hit. [REGION:none]![](collision_shrink.png)(w:192)[/REGION] |
 | **Fully Custom** | Enables you to specify the geometry used interactively in the viewport. [REGION:none]![](collision_custom.png)(w:192)[/REGION] |
diff --git a/Engine/Documentation/Source/Engine/Paper2D/Sprites/Paper2DSprites.INT.udn b/Engine/Documentation/Source/Engine/Paper2D/Sprites/Paper2DSprites.INT.udn
--- a/Engine/Documentation/Source/Engine/Paper2D/Sprites/Paper2DSprites.INT.udn
+++ b/Engine/Documentation/Source/Engine/Paper2D/Sprites/Paper2DSprites.INT.udn
@@ -8,6 +8,8 @@ Related: Engine\Paper2D\Sprites\RenderGeometry
 Related: Engine\Paper2D\Flipbooks
 Related: Engine/Rendering/Materials
 Related: Engine/Content/Types/Textures
+Version: 4.9
+SkillLevel: Intermediate
 
 [VAR:Topic]
 [OBJECT:Topic]
@@ -61,7 +63,7 @@ A **Sprite** in **Paper 2D** is a Texture Mapped Planar Mesh and associated Mate
 
 Sprites can also be edited inside UE4 in the **Sprite Editor** which offers four modes: **View** for previewing the sprite and general statistics, **Edit Source Region** which displays the full source texture and allows you to set the area within the source texture that composes the individual sprite, **Edit Collision** displays and allows editing of the sprite collision shapes, and **Edit RenderGeom** which displays and allows editing of the sprite render geometry. 
 
-Once you have a collection of Sprite assets, you can then begin animating them by creating a **Flipbook** which sequentially will play through the Sprites added to the Flipbook. See the [**Flipbook**](Engine\Paper2D\Flipbooks) documentation for more information. 
+Once you have a collection of Sprite assets, you can then begin animating them by creating a **Flipbook** which sequentially will play through the Sprites added to the Flipbook. See the [Flipbook](Engine\Paper2D\Flipbooks) documentation for more information. 
 [/EXCERPT]
 
 ## Creating Sprites
@@ -72,9 +74,9 @@ Sprites are created in the **Content Browser**, like other assets. They can be c
 
 **To create a new blank Sprite asset:**
 
-1. Click the ![Create](button_cb_new.png)(h:24) button in the **Content Browser**, then under _Miscellaneous_, select **Sprite**.
+1. Click the **Add New** button in the **Content Browser**, then under _Paper2D_, select **Sprite**.
 	
-	![](sprite_new_menu.png)
+	![](SpriteCB.jpg)
 
 	[REGION:tip]
 	You can also **Right-click** in the **Content Browser** to bring up the same context menu. 
@@ -82,22 +84,22 @@ Sprites are created in the **Content Browser**, like other assets. They can be c
 
 1. Name the new Sprite asset.
 
-	![](sprite_new_name.png)
+	![](MySpriteCB.jpg)
 
 1. The Sprite has been created but has not been saved yet (indicated by the asterisk in the lower-left corner). 
 
-	![](sprite_new_saveall.png)
+	![](SpriteSaveAll.jpg)
 
-	Click the ![Save All](save_all_button.png)(h:24) button to save the Sprite. 
+	Click the **Save All** button to save the Sprite. 
 
 1. **Double-click** the new Sprite asset to open it in the **Sprite Editor**. 
 
 1. In the **Details** panel, you can assign a Texture to the Sprite asset using the **Source Texture** property.
 	
-	![](sprite_new_details_texture.png)
+	![](DetailsSprite.jpg)
 
 [REGION:note]
-See the [**Sprite Editor Reference**](Engine\Paper2D\Sprites\Editor) documentation for information on working with Sprites in the **Sprite Editor**.
+See the [Sprite Editor Reference](Engine\Paper2D\Sprites\Editor) documentation for information on working with Sprites in the **Sprite Editor**.
 [/REGION]
 
 
@@ -107,7 +109,7 @@ See the [**Sprite Editor Reference**](Engine\Paper2D\Sprites\Editor) documentati
 
 1. In the **Content Browser**, **Right-click** on the Texture asset, then under _Sprite Actions_ choose **Create Sprite**.
 
-	![](sprite_create_menu.png)
+	![](ChooseSprite.jpg)
 
 1. Name the new Sprite asset.
 
@@ -132,7 +134,7 @@ See the [**Sprite Editor Reference**](Engine\Paper2D\Sprites\Editor) documentati
 **Paper 2D** includes an importer for JSON formatted sprite sheet descriptions, which can be exported from tools like **Adobe Flash CS6** or **Texture Packer**. It will import any referenced textures and create sprite assets for each sprite. The importer also assumes that all of the sprites are frames of an animation, so it will always create a Flipbook in addition to the individual sprites. You can delete the generated flipbook if it is not needed.
 
 [REGION:note]
-See [**Paper 2D Import Options**](Engine\Paper2D\Importing) for more information on Importing Options. 
+See [Paper 2D Import Options](Engine\Paper2D\Importing) for more information on Importing Options. 
 [/REGION]
 
 
diff --git a/Engine/Documentation/Source/Engine/Rendering/Materials/HowTo/Fresnel/Fresnel.INT.udn b/Engine/Documentation/Source/Engine/Rendering/Materials/HowTo/Fresnel/Fresnel.INT.udn
--- a/Engine/Documentation/Source/Engine/Rendering/Materials/HowTo/Fresnel/Fresnel.INT.udn
+++ b/Engine/Documentation/Source/Engine/Rendering/Materials/HowTo/Fresnel/Fresnel.INT.udn
@@ -5,6 +5,8 @@ Description:A How To Guide for using the Fresnel Material node.
 Related: Engine/Rendering/Materials/Functions/Reference/VectorOps
 Related: Engine/Rendering/Materials/Editor/Interface
 Related: Engine/Rendering/Materials/ExpressionReference/Utility
+Version: 4.9
+SkillLevel: Advanced
 
 
 [TOC(start:2 end:3)]
@@ -21,7 +23,7 @@ This way, the artist has more control over the look and feel of the effect and c
 
 ## Fresnel 
 
-Fresnel is the term used to describe how the light you see reflects at different intensities based off of the angle you are viewing it from. 
+**Fresnel** is the term used to describe how the light you see reflects at different intensities based off of the angle you are viewing it from. 
 For example, if you are standing over a pool looking straight down at the pool, you will not see a lot of reflections in the water. 
 As you start to move your head so that the water in the pool becomes more and more parallel to your eye level, you will begin to notice more and more reflections in the water. 
 
@@ -44,7 +46,7 @@ While these two Material Expressions have the same name and do the same thing, t
 
 ### Fresnel Utility
 
-The Fresnel that is found under **Utility** is the one that is used in this How - To and is the simplest one out of the two to setup and use. 
+The Fresnel that is found under **Utility** is the one that is used in this How-to and is the simplest one out of the two to setup and use. 
 It is **highly recommended** that you use this version to begin with and then move on to the other once you better understand how the node functions.
 
 ![](FN_Fresnel_Node.png)
diff --git a/Engine/Documentation/Source/Engine/Rendering/Materials/HowTo/Making_Parameters/Making_Parameters.INT.udn b/Engine/Documentation/Source/Engine/Rendering/Materials/HowTo/Making_Parameters/Making_Parameters.INT.udn
--- a/Engine/Documentation/Source/Engine/Rendering/Materials/HowTo/Making_Parameters/Making_Parameters.INT.udn
+++ b/Engine/Documentation/Source/Engine/Rendering/Materials/HowTo/Making_Parameters/Making_Parameters.INT.udn
@@ -6,6 +6,8 @@ Related: Engine/Rendering/Materials/ExpressionReference/Parameters
 Related: Engine/Rendering/Materials/MaterialInstances
 Related: Engine/Rendering/Materials/ParameterCollections
 Related: Engine/Matinee/HowTo/MHT_6
+Version: 4.9
+SkillLevel: Intermediate
 
 [TOC(start:2)]
 
@@ -89,7 +91,7 @@ Where the real power of Material Parameters comes into play is when you use them
 
 1. First find and add the following type and amount of Material Parameter Expression nodes to the Material Graph.
 
-	* **  Vector Parameter ** x 1
+	* ** Vector Parameter ** x 1
 	* ** Scalar Node** x 2
 	
 	[OBJECT:EmbeddedVideo]
@@ -110,9 +112,9 @@ Where the real power of Material Parameters comes into play is when you use them
 
 1. Next, make sure that all the Parameter nodes have names and default values. For this example, we are using the following names and default values.
 
-	* ** Base_Color** R:1.0 G:0.0 B:0.0
-	* ** Metallic** 0.0
-	* ** Roughness** 0.5
+	* **Base_Color** R:1.0 G:0.0 B:0.0
+	* **Metallic** 0.0
+	* **Roughness** 0.5
 
 	[OBJECT:EmbeddedVideo]
 	[PARAMLITERAL:width]
diff --git a/Engine/Documentation/Source/Engine/Rendering/Materials/HowTo/Subsurface_Profile/Subsurface_Profile.INT.udn b/Engine/Documentation/Source/Engine/Rendering/Materials/HowTo/Subsurface_Profile/Subsurface_Profile.INT.udn
--- a/Engine/Documentation/Source/Engine/Rendering/Materials/HowTo/Subsurface_Profile/Subsurface_Profile.INT.udn
+++ b/Engine/Documentation/Source/Engine/Rendering/Materials/HowTo/Subsurface_Profile/Subsurface_Profile.INT.udn
@@ -6,6 +6,8 @@ Related: Engine/Rendering/Materials/LightingModels/SubSurfaceProfile
 Related: Engine/Rendering/Materials/LightingModels/SubSurface
 Related: Resources/ContentExamples/MaterialNodes/1_13
 Related: Resources/ContentExamples/MaterialNodes/1_7
+Version: 4.9
+SkillLevel: Intermediate
 
 [TOC(start:2)]
 
@@ -68,8 +70,9 @@ In the following steps, we will be creating a Material and a Material Instance t
 1. The first thing that we need to do is to make sure that our Subsurface Material that was created above is open. If you have the Material open, please go on to the next step. 
    If not, to open the Material, first find it in the **Content Browser** and then **Double-Click** on it with the **Left Mouse Button** to open it up.
    When completed, you should have something like this on your screen. 
+	
 	[REGION:lightbox]
-	![](SSP_Open_New_Material.png)
+	[![](SSP_Open_New_Material.png)](SSP_Open_New_Material.png)
 	[/REGION]
 
 1. With the Material open, it is time to place down some Material Expression nodes. For this tutorial, we are going to make use of the following nodes. 
@@ -127,8 +130,9 @@ Once your mouse is over the input you want to use, release the **Left Mouse Butt
 
 1. With the Subsurface Profile applied and the Material expression linked to the main Material node, it is now time to compile the shader and get ready to make a Material Instance out of this Material. 
    When completed with the Material, it should look something like this. 
+	
 	[REGION:lightbox]
-	![](SSP_Completed_Material.png)
+	[![](SSP_Completed_Material.png)](SSP_Completed_Material.png)
 	[/REGION]
 
 ## Using a Subsurface Profile in a Material Instance
@@ -143,13 +147,14 @@ When completed, your **Content Browser** should look something like this.
 1. Now that the Material Instance has been created, open it up by using the **Left Mouse Button** to **Double-Click** on it in the **Content Browser**. Once open, you should see something that looks like this. 
 
 	[REGION:lightbox]
-	![](SSP_Material_Instance.png)
+	[![](SSP_Material_Instance.png)](SSP_Material_Instance.png)
 	[/REGION]
 
 1. To change the Subsurface Profile in a Material Instance, first make sure that you enable the **Override Subsurface Profile** and then provide the Subsurface Profile that you want to use 
 	by using the **Left Mouse  Button** to click on the **Subsurface Profile** input. Once you click on the input you should see the Subsurface Profiles that are available to choose from. 
+	
 	[REGION:lightbox]
-	![](SSP_SSP_In_Material_Instance.png)
+	[![](SSP_SSP_In_Material_Instance.png)](SSP_SSP_In_Material_Instance.png)
 	[/REGION]
 
 	[region:note]
@@ -190,7 +195,7 @@ When completed, your **Content Browser** should look something like this.
 	When completed, you should have something that looks like the following.
 
 	[REGION:lightbox]
-	![](Engine\Rendering\Materials\HowTo\Subsurface_Scattering\SSM_Light_Settings.png)
+	[![](Engine\Rendering\Materials\HowTo\Subsurface_Scattering\SSM_Light_Settings.png)](Engine\Rendering\Materials\HowTo\Subsurface_Scattering\SSM_Light_Settings.png)
 	[/REGION]
 
 1. With all of the lights added, we can now add a Static Mesh to test with. For this example, we will be using the following mesh. 
diff --git a/Engine/Documentation/Source/Gameplay/HowTo/SetUpInput/Blueprints/SetInput_BP.INT.udn b/Engine/Documentation/Source/Gameplay/HowTo/SetUpInput/Blueprints/SetInput_BP.INT.udn
--- a/Engine/Documentation/Source/Gameplay/HowTo/SetUpInput/Blueprints/SetInput_BP.INT.udn
+++ b/Engine/Documentation/Source/Gameplay/HowTo/SetUpInput/Blueprints/SetInput_BP.INT.udn
@@ -2,6 +2,8 @@ Availability: Public
 Title:How to Set Up Inputs in Blueprints
 Crumbs: %ROOT%, Gameplay, Gameplay\HowTo\SetUpInput
 Description:A How To Guide for Setting Up Inputs in Blueprints.
+Version: 4.9
+SkillLevel: Intermediate
 
 [TOC (start:2 end:2)]
 
@@ -13,7 +15,7 @@ On this page, you will create a simple character and set it up to receive input,
 
 In this section, we will create a character that we can assign input to (if you already have a **Character Blueprint**, you can skip this step). 
 
-_For this example, we are using the **Blank** project template with **Starter Content** enabled._
+_For this example, we are using the **Blank** project template with **Starter Content** included._
 
 1.	Inside your project, under **Shapes**, **Right-click** on the **Shape_NarrowCapsule** shape and select **Asset Actions** - **Create Blueprint Using**. 
 
diff --git a/Engine/Documentation/Source/Platforms/Android/Reference/AndroidGettingStarted.INT.udn b/Engine/Documentation/Source/Platforms/Android/Reference/AndroidGettingStarted.INT.udn
--- a/Engine/Documentation/Source/Platforms/Android/Reference/AndroidGettingStarted.INT.udn
+++ b/Engine/Documentation/Source/Platforms/Android/Reference/AndroidGettingStarted.INT.udn
@@ -2,6 +2,8 @@
 Title: Android Development Reference
 Crumbs:%ROOT%, Platforms, Platforms/Android
 Description:How to install different Android SDKs, set environment variables, and work with texture formats.
+Version: 4.9
+SkillLevel: Intermediate
 
 [TOC(start:2 end:3)]
 
diff --git a/Engine/Documentation/Source/Platforms/GearVR/BestPractices/GearVRBestPractices.INT.udn b/Engine/Documentation/Source/Platforms/GearVR/BestPractices/GearVRBestPractices.INT.udn
--- a/Engine/Documentation/Source/Platforms/GearVR/BestPractices/GearVRBestPractices.INT.udn
+++ b/Engine/Documentation/Source/Platforms/GearVR/BestPractices/GearVRBestPractices.INT.udn
@@ -117,10 +117,10 @@ In the following section we will cover what you need to be concerned with when m
 
 ### Content Hard limits
 Due to certain hardware restrictions on mobile devices the following restrictions must be adhered to when creating content. 
-Failing to follow these restrictions will result in your project now working.
+Failing to follow these restrictions will result in your project not working.
 
 * All mesh types can have up to but not over 65,000 vertices due to the lack of 32 bit index support on mobile hardware.
-* Skeletal Meshes can have up to but not 75 bones.
+* Skeletal Meshes can have up to but not over 75 bones.
 
 ### Mesh Considerations
 When creating Meshes to use with your Gear VR Project the following information needs to be kept in mind to ensure that what you are creating will work with the Gear VR Headset. 
@@ -129,14 +129,14 @@ When creating Meshes to use with your Gear VR Project the following information
 * Make sure to make use of [Static Mesh LODs](Engine/Content/FBX/StaticMeshes) whenever possible.
 * Keep an eye on your 3D meshes polygon count and try and keep it as [Low Poly](https://en.wikipedia.org/wiki/Low_poly) as possible.
 * Make sure that all static models have a 2nd UV set so they can make use of [Light Mapping](Engine/Content/Types/StaticMeshes/LightmapUnwrapping).
-* Whenever possible try and make use if [Imposter Sprites](Engine/Content/Tools/RenderToTextureTools/3) to replace far away Static Meshes as they are much cheaper to render.
+* Whenever possible try and make use of [Imposter Sprites](Engine/Content/Tools/RenderToTextureTools/3) to replace far away Static Meshes as they are much cheaper to render.
 
 ### Texture Considerations
 When creating Textures to use with your Gear VR Project the following information needs to be kept in mind to ensure that what you are creating will not be detrimental to your projects overall performance. 
 
 * Textures use the [ECT2](http://en.wikipedia.org/wiki/Ericsson_Texture_Compression) Texture compression format.
 * Textures can not be bigger than 2048 pixels in either the X or Y dimension.
-* Textures sizes must be in a power of 2 (2, 4, 8, 16, 32, 64, 128, 256, 512, 1024, 2048)
+* Textures sizes must be in a power of 2 (For example, 2, 4, 8, 16, 32, 64, 128, 256, 512, 1024, 2048)
 * Loading up to 512 MB of Textures into memory is possible with the Gear VR however it is recommended that you do not load more that 128 MB of Textures into memory.
 * Use square Textures whenever possible as them make a more efficient use of memory. 
 * Wherever possible try and use [Texture Atlasing](http://en.wikipedia.org/wiki/Texture_atlas) to reduce the number of Textures needed in memory.
diff --git a/Engine/Documentation/Source/Platforms/GearVR/GearVR.INT.udn b/Engine/Documentation/Source/Platforms/GearVR/GearVR.INT.udn
--- a/Engine/Documentation/Source/Platforms/GearVR/GearVR.INT.udn
+++ b/Engine/Documentation/Source/Platforms/GearVR/GearVR.INT.udn
@@ -4,7 +4,7 @@ Availability:Public
 Navigation:platform
 Platform:Gear VR
 Skilllevel: Intermediate
-Version: 4.8
+Version: 4.9
 prereq: Engine\Blueprints\GettingStarted
 
 
diff --git a/Engine/Documentation/Source/Platforms/HTML5/GettingStarted/HTML5GettingStarted.INT.udn b/Engine/Documentation/Source/Platforms/HTML5/GettingStarted/HTML5GettingStarted.INT.udn
--- a/Engine/Documentation/Source/Platforms/HTML5/GettingStarted/HTML5GettingStarted.INT.udn
+++ b/Engine/Documentation/Source/Platforms/HTML5/GettingStarted/HTML5GettingStarted.INT.udn
@@ -2,6 +2,8 @@
 Title: Getting Started: Developing HTML5 Projects
 Crumbs:%ROOT%, Platforms, Platforms/HTML5
 Description:Setting up for development for the HTML5 platform.
+Version: 4.9
+SkillLevel: Intermediate
 
 [TOC(start:2 end:3)]
 
@@ -127,7 +129,7 @@ Please refer to [http://kripken.github.io/emscripten-site/docs/getting_started/d
 
 1. [Install Python 2.7](https://www.python.org/download/releases/2.7/)
 
-	Any version of Python **2.7** will do. As of this writing, 2.7.9 is the latest 2.7 release. Do not use Python 3.X.
+	Any version of Python **2.7** will do. As of this writing, 2.7.10 is the latest 2.7 release. **Do not use Python 3.X**.
 
 1. Update your `Engine\Config\HTML5\HTML5Engine.ini` file:
 
diff --git a/Engine/Documentation/Source/Platforms/HTML5/HTML5.INT.udn b/Engine/Documentation/Source/Platforms/HTML5/HTML5.INT.udn
--- a/Engine/Documentation/Source/Platforms/HTML5/HTML5.INT.udn
+++ b/Engine/Documentation/Source/Platforms/HTML5/HTML5.INT.udn
@@ -4,6 +4,7 @@ Crumbs:%ROOT%, Platforms
 Description:Developing games for the HTML5 platform.
 Platform:HTML5
 Navigation:platform
+Version: 4.9
 
 
 [OBJECT:TopicButtonList]
diff --git a/Engine/Documentation/Source/Platforms/Mobile/Textures/MobileTextures.INT.udn b/Engine/Documentation/Source/Platforms/Mobile/Textures/MobileTextures.INT.udn
--- a/Engine/Documentation/Source/Platforms/Mobile/Textures/MobileTextures.INT.udn
+++ b/Engine/Documentation/Source/Platforms/Mobile/Textures/MobileTextures.INT.udn
@@ -53,7 +53,17 @@ To set the Texture LOD settings per-device you will need to do the following.
 
 1. To make a change just click on the property you want to adjust and then input what you want the setting to be.
 
-	![](T_Selecting_Editing_Profile.gif)(convert:false)
+	[OBJECT:EmbeddedVideo]
+	[PARAMLITERAL:width]
+	640
+	[/PARAMLITERAL]
+	[PARAMLITERAL:height]
+	360
+	[/PARAMLITERAL]
+	[PARAMLITERAL:videoid]
+	5FgGi68G8ZY
+	[/PARAMLITERAL]
+	[/OBJECT]
 
 
 
