fixe content error: Failed to compile Material /Engine/BufferVisualization/DecalMask.DecalMask for platform PCD3D_SM5, Default Material will be used in game.

diff --git a/Engine/Documentation/Source/Engine/Basics/Screenshots/Screenshots.INT.udn b/Engine/Documentation/Source/Engine/Basics/Screenshots/Screenshots.INT.udn
--- a/Engine/Documentation/Source/Engine/Basics/Screenshots/Screenshots.INT.udn
+++ b/Engine/Documentation/Source/Engine/Basics/Screenshots/Screenshots.INT.udn
@@ -136,7 +136,6 @@ The High Resolution Screenshot tool (HRSST) supports the ability to render out e
 
 * **DiffuseColor**
 * **LightingModel**
-* **DecalMask**
 * **MaterialAO**
 * **Opacity**
 * **Roughness**
@@ -170,7 +169,6 @@ After clicking the capture button, the HRSST writes out the following buffers to
 
 * **DiffuseColor**
 * **LightingModel**
-* **DecalMask**
 * **MaterialAO**
 * **Opacity**
 * **Roughness**
