Update documentation for ScreenPosition material expression.

diff --git a/Engine/Documentation/Source/Engine/Rendering/PostProcessEffects/PostProcessMaterials/PostProcessMaterials.INT.udn b/Engine/Documentation/Source/Engine/Rendering/PostProcessEffects/PostProcessMaterials/PostProcessMaterials.INT.udn
--- a/Engine/Documentation/Source/Engine/Rendering/PostProcessEffects/PostProcessMaterials/PostProcessMaterials.INT.udn
+++ b/Engine/Documentation/Source/Engine/Rendering/PostProcessEffects/PostProcessMaterials/PostProcessMaterials.INT.udn
@@ -189,17 +189,16 @@ We get a stable and anti-aliased image. In motion we might see some artifacts wi
 it works ok but outside the object we would also need to adjust the depth buffer (not done yet, costs extra performance). Note that TemporalAA is still under development - here it appears the outline is thicker
 but ideally this should not be. 
 
-## UV and ScreenPos
+## UV and ScreenPosition
 
-Post process materials you can look up into screen aligned buffers, but you need to know the right UVs. 
-The **ScreenPosition** material expression outputs the UV that you would expect (0,0 at the left top of the viewport and 1,1 at the bottom right).
-Using the **texture coordinate** material expression might give you different results. This is because the actual texture (more correctly it is a render target) is potentially larger than the viewport.
+With post process materials you can look up into screen aligned buffers, but you need to know the right UVs. 
+The **ScreenPosition** material expression with the mapping option set to **ViewportUV** outputs the UV with 0,0 at the left top of the viewport and 1,1 at the bottom right.
+In contrast, using the **SceneTextureUV** mapping in this material expression might give you different results. This is because the actual texture (more correctly it is a render target) is potentially larger than the viewport.
 It can be larger in editor because we share this texture for multiple viewports and the largest extent is use for all viewports.
 Even in game it can be larger in some cases (e.g. SceneCaptureActors might have a smaller viewport, Matinee black borders, Splitscreen, VR, ...).
-The texture coordinate material expression gives you the UV for this larger texture. If you only need a relative offset (e.g. pixel sized edge detection) you need to scale with the right size.
+The **SceneTextureUV** option gives you the UV for this larger texture. If you only need a relative offset (e.g. pixel sized edge detection) you need to scale with the right size.
 The **SceneTexture** material expression has outputs for the size and the inverse of the size (efficient and useful for pixel offsets).
-If you want to get the viewport UV (e.g. to map texture to the viewport) you can use the **ScreenPosition** material expression or you make the needed computations manually (more control, can be slower).
-For this, you need the **ViewSize** material expression. To test all that, you can use the console variable **r.ViewPortTest** which allows you to test various viewport configurations.
+To test all that, you can use the console variable **r.ViewPortTest** which allows you to test various viewport configurations.
 
 ## Filtered texture lookups
 
