Merging //UE4/Release-4.11 to //UE4/Main (Up to CL#2897738)

==========================
MAJOR FEATURES + CHANGES
==========================

Change 2875445 on 2016/02/22 by [REDACTED_PERSON]

	Added UE4.natvis to Visual Studio Projects

Change 2875456 on 2016/02/22 by [REDACTED_PERSON]

	Fix custom stencil shaders on Xbox One

Change 2875524 on 2016/02/22 by [REDACTED_PERSON]

	More log info when saving shader temp files fails. Increased the number of attemps when moving a file fails.

Change 2875698 on 2016/02/22 by [REDACTED_PERSON]

	UE4.11 - Add new bool for RHIs (unused currently)

Change 2875897 on 2016/02/22 by [REDACTED_PERSON]

Change 2875922 on 2016/02/22 by [REDACTED_PERSON]

	Fixed BP-implemented EQS generators crashing when trying to add generated value of wrong type #UE4

Change 2875960 on 2016/02/22 by [REDACTED_PERSON]

	Added a way to disable right click emulation on Mac and used it in TabNavigator to fix issues with its widgets not reacting to clicking

Change 2875984 on 2016/02/22 by [REDACTED_PERSON]

	Split output struct pins will no longer give a warning about override pins being removed.

Change 2876169 on 2016/02/22 by [REDACTED_PERSON]

	Changes to support building UHT plugins with the binary release.

	* Add source code and target files for UHT to binary distribution
	* Fix UBT deleting build products if we're only compiling a single module.
	* Fix UBT exception setting up compile environment when a module doesn't have any source files set to build.
	* Include DLL import libraries for UHT in the build
	* Add support for compiling UHT modules in BuildPluginCommand. Stages an empty host project to allow UHT to load any enabled plugins.

Change 2876219 on 2016/02/22 by [REDACTED_PERSON]

	UE4.11 - Integration from 2874609
	PC: Update D3D12 RHI
	- Implement _RenderThread versions of Create, Lock and Unlock Index/Vertex Buffer. Only synchronize threads on Readback
	- Limit GPU starvation on CPU bound scenarios by flushing work when the GPU is idle
	- Change texture streaming system to correctly account for placed textures. Also fix texture sizes so they accurately represent the real size of the allocation the GPU.
	- Disable API shader blobs
	- Add the ability to easily change allocation stategy for a given pool, also add a simple linear allocator and a 'Multi-Buddy Allocator' for efficiency in different scenarios
	- Pool Fences to prevent creation and destruction every frame when using Async Compute
	- Implement _RenderThread versions of CreateShaderResourceView and CreateUnorderedAccessView

Change 2876232 on 2016/02/22 by [REDACTED_PERSON]

	UE4.11 - Integration from 2876173
	PC: Update D3D12 RHI

	- Fix ResizeBuffers() failing due to dangling references to the backbuffer if deferred deletion is used.

	- Reorder when pending FRHIResources are deleted. This still needs to flush all pending deletes and ignore the deferred deletion queue otherwise some items may still be left in the engine's queue.

	- Fix UT build error due to missing FPlatformMisc::GetGPUDriverInfo()

Change 2876366 on 2016/02/22 by [REDACTED_PERSON]

	Adding Skeletal Meshes for Import Test Case

Change 2876401 on 2016/02/22 by [REDACTED_PERSON]

	fix for WindowsClient build from UFE and Project Launcher

Change 2876456 on 2016/02/22 by [REDACTED_PERSON]

	Use a more hierarchical directory structure for packaged builds, rather than just dumping everything in the root. Now defaults to <Share>\\PackagedBuilds\\<Branch>\\<CL>\\<ProjectName>_<Platform>_<Configuration>.

Change 2876507 on 2016/02/22 by [REDACTED_PERSON]

	use HOME (osx) and USERPROFILE (windows) on appropriate target platform

Change 2876537 on 2016/02/22 by [REDACTED_PERSON]

	Removed dubious fix for an old bug, no longer needed but I havn't figured out what has changed. This fixes a crash on Replace References, but does not reintroduce UE-9497

Change 2876545 on 2016/02/22 by [REDACTED_PERSON]

	SteamVR camera late-update fix

Change 2876825 on 2016/02/22 by [REDACTED_PERSON]

	Unfortunate edge case in lifetime of UEdGraph's schema, schema is assigned after construction so its modification is in the undo buffer, and we clear it after undoing.

Change 2876878 on 2016/02/22 by [REDACTED_PERSON]

	PSVR HMD Server support

diff --git a/Engine/Documentation/Source/Shared/Tutorials/InstallingVisualStudioTutorial/InstallingVisualStudioTutorial.INT.udn b/Engine/Documentation/Source/Shared/Tutorials/InstallingVisualStudioTutorial/InstallingVisualStudioTutorial.INT.udn
--- a/Engine/Documentation/Source/Shared/Tutorials/InstallingVisualStudioTutorial/InstallingVisualStudioTutorial.INT.udn
+++ b/Engine/Documentation/Source/Shared/Tutorials/InstallingVisualStudioTutorial/InstallingVisualStudioTutorial.INT.udn
@@ -6,21 +6,14 @@ Description: The contents of all introductory tutorials for Unreal Engine 4
 [EXCERPT:Stage1]
 ##Installing Visual Studio
 
-It appears that Visual Studio 2013 is not installed. 
+It appears that Visual Studio 2015 is not installed. 
 
-Installing Visual Studio 2013 will allow you to view Unreal header files, as well as create and compile your own code projects.
+Installing Visual Studio 2015 will allow you to view Unreal header files, as well as create and compile your own code projects.
 
-You can install Visual Studio 2013 Express from Microsoft for free using the following link:
+You can install Visual Studio 2015 Community from Microsoft for free using the following link:
 
 
 
-[Download Visual Studio 2013](http://www.microsoft.com/en-us/download/details.aspx?id=40787)
-
-
-[![](DownloadVSHighlight.png)](http://www.microsoft.com/en-us/download/details.aspx?id=40787)
-
----
-
-**For ease of installation, select wdexpress_full.exe for the download.**
+[Download Visual Studio 2015](http://go.microsoft.com/fwlink/?LinkId=532606)
 
 [/EXCERPT:Stage1]
