Copying //UE4/Dev-Framework to Dev-Main (//UE4/Dev-Main) @ 2926658

==========================
MAJOR FEATURES + CHANGES
==========================

Change 2821607 on 2016/01/08 by [REDACTED_PERSON]

	Added a way to limit amount of information logged by vlog by discarding logs from classes from outside of class whitelist #UE4

	This feature was followed by refactoring of functions taking FVisualLogEntry pointers to use references instead.

Change 2828384 on 2016/01/14 by [REDACTED_PERSON]

	Back out of visual log refactor done as part of CL#2821607 #UE4

Change 2910454 on 2016/03/15 by [REDACTED_PERSON]

	#ue4 - Properly exclude zero-distance MTD results in ComponentEncroachesBlockingGeometry_WithAdjustment() in the presense of multiple overlaps.

	UDN: https://udn.unrealengine.com/questions/270574/jeff-farris-hack-for-physx-mtd.html

Change 2910548 on 2016/03/15 by [REDACTED_PERSON]

	#ue4 - Handle MTD computation returning NaN direction when there is a "contact" with zero distance.

Change 2912311 on 2016/03/16 by [REDACTED_PERSON]

	Properly handle overlaps in C++ in documentation code and UE4 to Unity doc

Change 2913086 on 2016/03/17 by [REDACTED_PERSON]

	Adding ability to have 9 parameters to a dynamic delegate

Change 2913101 on 2016/03/17 by [REDACTED_PERSON]

	Fix some of the loctext error messages

Change 2913102 on 2016/03/17 by [REDACTED_PERSON]

	Console usability improvements

	Display console autocompletion commands from the lexicographically first element up to either the total number of commands or MAX_AUTOCOMPLETION_LINES, whichever is least. The previous behaviour started the list "in the middle" and hid the first elements if there were too many matches. Thus "[ab ac ad]" with "aa" hidden now becomes "[aa ab ac]" with "ad" hidden.
	To make scrolling work as expected, the input handling of the up and down arrow keys has been reversed so that the cursor iterates forward starting from the top with the down arrow key, and goes back up with the up arrow key. Command history is still accessed with the up arrow key.
	This commit also undoes one of the most evil uses of operator overloading I've ever seen, on par with "#define true false" but more subtle

	Color console autocomplete entries to denote their type: command, CVar or other (manual autocompletion entries). CVars are further divided into writeable and read-only variables.
	Assume that manual console autocompletion entries are commands. This makes the autocompletion list colors more consistent and less noisy

	Automatically select (but don't complete) a command on console character input. To prevent the autocomplete from becoming too trigger happy, no longer automatically complete commands for arbitrary key inputs that we happen to have a match for

	Allow cycling through console commands with the tab key
	Discriminate between first time and repeated tab presses and only scroll through autocomplete entries on the latter

	Fix off-by-one error in console: "x more matches" line was being shown when the number of autocomplete elements was equal to MAX_AUTOCOMPLETION_LINES
	Fix an off-by-one error that was causing the topmost console command to not be shown if there was an autocomplete scroll region

	Show the currently selected autocomplete entry faded out behind the user's typed input
	Slightly increase brightness of the normal input text colour to better distinguish between the typed and autocompleted parts of the input line

	Left-justify command descriptions in the console autocompletion box
	Detect overflow of console autocomplete lines on low resolutions and decrease the space used for description justification to compensate

	Make the console input, history and autocomplete colours user configurable

	Add console background transparency. Configurable, set to 15% by default

	Add missing closing quote to the console dump HTML template

	#github #2061: Console usability improvements from Mattiwatti

Change 2913104 on 2016/03/17 by [REDACTED_PERSON]

	Added indicator displayed on animation nodes when they use the 'fast path'

	Added checkbox that can be used to audit Blueprint fast-path usage.
	Switched almost all animation node widgets to derive from new SAnimationGraphNode. This creates the overlay widget that indicates whether this node is using a more optimal path.

	#doc Also added documentation tooltips and UDN doc files/images for the fast path systems.

Change 2913306 on 2016/03/17 by [REDACTED_PERSON]

	Cleaning up GetResourceSize
	- Made many call Super::GetResourceSize
	- Removed trivial implementations
	- Fixed HierarchicalInstanceStaticMeshComponent double counting an array

Change 2913535 on 2016/03/17 by [REDACTED_PERSON]

	fixed broken behavior tree graph data after subnode undo
	#ue4 UE-28198

Change 2913608 on 2016/03/17 by [REDACTED_PERSON]

	fixed behavior tree execution indices after undoing move in editor
	#ue4 UE-26705

Change 2913847 on 2016/03/17 by [REDACTED_PERSON]

	added new automation test for UE-28309
	#ue4

Change 2913849 on 2016/03/17 by [REDACTED_PERSON]

	fixed behavior tree skipping over branch when restart request comes during AbortCurrentTask call
	#ue4 UE-28309

Change 2913895 on 2016/03/17 by [REDACTED_PERSON]

	Added 'self' argument to Actor and PrimitiveComponent delegates that didn't already supply one
	Fixed up all C++ uses of these delegates

Change 2914743 on 2016/03/18 by [REDACTED_PERSON]

	Editing of primitive data in PhAT

diff --git a/Engine/Documentation/Extras/ConsoleHelpTemplate.html b/Engine/Documentation/Extras/ConsoleHelpTemplate.html
--- a/Engine/Documentation/Extras/ConsoleHelpTemplate.html
+++ b/Engine/Documentation/Extras/ConsoleHelpTemplate.html
@@ -9,7 +9,7 @@
 
 <br>
 
-<div style="padding:0px; margin:40px; min-width:450px; max-width:1250px;>
+<div style="padding:0px; margin:40px; min-width:450px; max-width:1250px;">
 
 <div style="padding:0px; margin:20px" >
 
diff --git a/Engine/Documentation/Source/Engine/Animation/Optimization/FastPath/FastPath.INT.udn b/Engine/Documentation/Source/Engine/Animation/Optimization/FastPath/FastPath.INT.udn
new file mode 100644
--- /dev/null
+++ b/Engine/Documentation/Source/Engine/Animation/Optimization/FastPath/FastPath.INT.udn
@@ -0,0 +1,38 @@
+Availability:Public
+Title: Animation Fast Path Optimization
+Crumbs: %ROOT%, Engine, Engine/Animation
+Description:Optimizations that can be applied to the structure of animation graphs to gain performance
+version: 4.11
+
+[EXCERPT:Intro]
+Animation Blueprints are very flexible ways of building animation systems. However the performance overheads 
+associated with running Blueprint-based logic can become quite high as your graph becomes more complex.
+To mitigate these overheads certain methods of structuring your animation graph can be used.
+[/EXCERPT:Intro]
+
+## Access Member Variables Directly
+
+![](DirectMemberAccess.png)
+
+## Access Negated Boolean Member Variables
+
+![](NegatedBooleanAccess.png)
+
+## Access Members of a Nested Struct
+
+![](NestedStructAccess.png)
+
+## Access Members Using "Break Struct" Nodes
+
+![](BreakStructAccess.png)
+
+Note that some "Break Struct" nodes, like "Break Transform" will not currently use the fast path as they perform
+conversions internally rather than simply copying data.
+
+## Auditing and Maintaining Optimized Animation Blueprints
+
+![](WarnAboutBlueprintUsage.png)
+
+Once an Animation Blueprint has been configured in an optimal way, subsequent iterations can potentially introduce
+new, non-performant constructs. To keep track of this you can enable the "Warn About Blueprint Usage" option in the 
+Animation Blueprint's Class Defaults. This is also useful when optimizing a particular Animation Blueprint.
\ No newline at end of file
diff --git a/Engine/Documentation/Source/Gameplay/ClassCreation/CodeAndBlueprints/CodeAndBlueprints.INT.udn b/Engine/Documentation/Source/Gameplay/ClassCreation/CodeAndBlueprints/CodeAndBlueprints.INT.udn
--- a/Engine/Documentation/Source/Gameplay/ClassCreation/CodeAndBlueprints/CodeAndBlueprints.INT.udn
+++ b/Engine/Documentation/Source/Gameplay/ClassCreation/CodeAndBlueprints/CodeAndBlueprints.INT.udn
@@ -62,15 +62,15 @@ LightSwitchBoth.h
 
 			/** called when something enters the sphere component */
 			UFUNCTION(BlueprintNativeEvent, Category="Switch Functions")
-			void OnOverlapBegin(class AActor* OtherActor, class UPrimitiveComponent* OtherComp, int32 OtherBodyIndex, bool bFromSweep, const FHitResult& SweepResult);
+			void OnOverlapBegin(class UPrimitiveComponent* OverlappedComp, class AActor* OtherActor, class UPrimitiveComponent* OtherComp, int32 OtherBodyIndex, bool bFromSweep, const FHitResult& SweepResult);
 			
-			void OnOverlapBegin_Implementation(class AActor* OtherActor, class UPrimitiveComponent* OtherComp, int32 OtherBodyIndex, bool bFromSweep, const FHitResult& SweepResult);
+			void OnOverlapBegin_Implementation(class UPrimitiveComponent* OverlappedComp, class AActor* OtherActor, class UPrimitiveComponent* OtherComp, int32 OtherBodyIndex, bool bFromSweep, const FHitResult& SweepResult);
 
 			/** called when something leaves the sphere component */
 			UFUNCTION(BlueprintNativeEvent, Category="Switch Functions")
-			void OnOverlapEnd(class AActor* OtherActor, class UPrimitiveComponent* OtherComp, int32 OtherBodyIndex);
+			void OnOverlapEnd(class UPrimitiveComponent* OverlappedComp, class AActor* OtherActor, class UPrimitiveComponent* OtherComp, int32 OtherBodyIndex);
 
-			void OnOverlapEnd_Implementation(class AActor* OtherActor, class UPrimitiveComponent* OtherComp, int32 OtherBodyIndex);
+			void OnOverlapEnd_Implementation(class UPrimitiveComponent* OverlappedComp, class AActor* OtherActor, class UPrimitiveComponent* OtherComp, int32 OtherBodyIndex);
 
 			/** Toggles the light component's visibility*/
 			UFUNCTION()
@@ -115,7 +115,7 @@ LightSwitchBoth.cpp
 		Sphere1->OnComponentEndOverlap.AddDynamic(this, &ALightSwitchBoth::OnOverlapEnd);		// set up a notification for when this component overlaps something
 	}
 
-	void ALightSwitchBoth::OnOverlapBegin_Implementation(class AActor* OtherActor, class UPrimitiveComponent* OtherComp, int32 OtherBodyIndex, bool bFromSweep, const FHitResult& SweepResult)
+	void ALightSwitchBoth::OnOverlapBegin_Implementation(UPrimitiveComponent* OverlappedComp, AActor* OtherActor, UPrimitiveComponent* OtherComp, int32 OtherBodyIndex, bool bFromSweep, const FHitResult& SweepResult)
 	{
 		if (OtherActor && (OtherActor != this) && OtherComp)
 		{
@@ -123,7 +123,7 @@ LightSwitchBoth.cpp
 		}
 	}
 
-	void ALightSwitchBoth::OnOverlapEnd_Implementation(class AActor* OtherActor, class UPrimitiveComponent* OtherComp, int32 OtherBodyIndex)
+	void ALightSwitchBoth::OnOverlapEnd_Implementation(UPrimitiveComponent* OverlappedComp, AActor* OtherActor, UPrimitiveComponent* OtherComp, int32 OtherBodyIndex)
 	{
 		if (OtherActor && (OtherActor != this) && OtherComp)
 		{
diff --git a/Engine/Documentation/Source/Gameplay/ClassCreation/CodeOnly/CodeOnly.INT.udn b/Engine/Documentation/Source/Gameplay/ClassCreation/CodeOnly/CodeOnly.INT.udn
--- a/Engine/Documentation/Source/Gameplay/ClassCreation/CodeOnly/CodeOnly.INT.udn
+++ b/Engine/Documentation/Source/Gameplay/ClassCreation/CodeOnly/CodeOnly.INT.udn
@@ -49,11 +49,11 @@ This means that their properties will be visible in the **Details** tab of a Lig
 
 		/** called when something enters the sphere component */
 		UFUNCTION()
-		void OnOverlapBegin(class AActor* OtherActor, class UPrimitiveComponent* OtherComp, int32 OtherBodyIndex, bool bFromSweep, const FHitResult& SweepResult);
+		void OnOverlapBegin(class UPrimitiveComponent* OverlappedComp, class AActor* OtherActor, class UPrimitiveComponent* OtherComp, int32 OtherBodyIndex, bool bFromSweep, const FHitResult& SweepResult);
 
 		/** called when something leaves the sphere component */
 		UFUNCTION()
-		void OnOverlapEnd(class AActor* OtherActor, class UPrimitiveComponent* OtherComp, int32 OtherBodyIndex);
+		void OnOverlapEnd(class UPrimitiveComponent* OverlappedComp, class AActor* OtherActor, class UPrimitiveComponent* OtherComp, int32 OtherBodyIndex);
 
 * Declare `ToggleLight`, a function that toggles the visibility of the PointLightComponent.
 
@@ -103,11 +103,11 @@ LightSwitchCodeOnly.h
 
 			/** called when something enters the sphere component */
 			UFUNCTION()
-			void OnOverlapBegin(class AActor* OtherActor, class UPrimitiveComponent* OtherComp, int32 OtherBodyIndex, bool bFromSweep, const FHitResult& SweepResult);
+			void OnOverlapBegin(class UPrimitiveComponent* OverlappedComp, class AActor* OtherActor, class UPrimitiveComponent* OtherComp, int32 OtherBodyIndex, bool bFromSweep, const FHitResult& SweepResult);
 	
 			/** called when something leaves the sphere component */
 			UFUNCTION()
-			void OnOverlapEnd(class AActor* OtherActor, class UPrimitiveComponent* OtherComp, int32 OtherBodyIndex);
+			void OnOverlapEnd(class UPrimitiveComponent* OverlappedComp, class AActor* OtherActor, class UPrimitiveComponent* OtherComp, int32 OtherBodyIndex);
 
 			/** Toggles the light component's visibility*/
 			UFUNCTION()
@@ -180,7 +180,7 @@ LightSwitchCodeOnly.cpp
 
 	}
 
-	void ALightSwitchCodeOnly::OnOverlapBegin(class AActor* OtherActor, class UPrimitiveComponent* OtherComp, int32 OtherBodyIndex, bool bFromSweep, const FHitResult& SweepResult)
+	void ALightSwitchCodeOnly::OnOverlapBegin(class UPrimitiveComponent* OverlappedComp, class AActor* OtherActor, class UPrimitiveComponent* OtherComp, int32 OtherBodyIndex, bool bFromSweep, const FHitResult& SweepResult)
 	{
 		if (OtherActor && (OtherActor != this) && OtherComp)
 		{
@@ -188,7 +188,7 @@ LightSwitchCodeOnly.cpp
 		}
 	}
 
-	void ALightSwitchCodeOnly::OnOverlapEnd(class AActor* OtherActor, class UPrimitiveComponent* OtherComp, int32 OtherBodyIndex)
+	void ALightSwitchCodeOnly::OnOverlapEnd(class UPrimitiveComponent* OverlappedComp, class AActor* OtherActor, class UPrimitiveComponent* OtherComp, int32 OtherBodyIndex)
 	{
 		if (OtherActor && (OtherActor != this) && OtherComp)
 		{
diff --git a/Engine/Documentation/Source/GettingStarted/FromUnity/UnrealEngineforUnityDevelopers.INT.udn b/Engine/Documentation/Source/GettingStarted/FromUnity/UnrealEngineforUnityDevelopers.INT.udn
--- a/Engine/Documentation/Source/GettingStarted/FromUnity/UnrealEngineforUnityDevelopers.INT.udn
+++ b/Engine/Documentation/Source/GettingStarted/FromUnity/UnrealEngineforUnityDevelopers.INT.udn
@@ -897,28 +897,9 @@ UE4 C++:
 			Trigger.SetCollisionEnabled(ECollisionEnabled::QueryOnly);
 		}
 
-		void BeginPlay()
-		{
-			// Register to find out when an overlap occurs
-			OnActorBeginOverlap.AddDynamic(this, &AMyActor::OnTriggerEnter);
-			OnActorEndOverlap.AddDynamic(this, &AMyActor::OnTriggerExit);
-
-			Super::BeginPlay();
-		}
-
-		void EndPlay(const EEndPlayReason::Type EndPlayReason)
-		{
-			OnActorBeginOverlap.RemoveDynamic(this, &AMyActor::OnTriggerEnter);
-			OnActorEndOverlap.RemoveDynamic(this, &AMyActor::OnTriggerExit);
-
-			Super:EndPlay(EndPlayReason);
-		}
-
-		UFUNCTION()
-		void OnTriggerEnter(AActor* Other);
+		virtual void NotifyActorBeginOverlap(AActor* Other) override;
 
-		UFUNCTION()
-		void OnTriggerExit(AActor* Other);
+		virtual void NotifyActorEndOverlap(AActor* Other) override;
 	};
 
 UE4 Blueprint:
diff --git a/Engine/Documentation/Source/Programming/Tutorials/FirstPersonShooter/3/4/FirstPersonShooter_3_4.INT.udn b/Engine/Documentation/Source/Programming/Tutorials/FirstPersonShooter/3/4/FirstPersonShooter_3_4.INT.udn
--- a/Engine/Documentation/Source/Programming/Tutorials/FirstPersonShooter/3/4/FirstPersonShooter_3_4.INT.udn
+++ b/Engine/Documentation/Source/Programming/Tutorials/FirstPersonShooter/3/4/FirstPersonShooter_3_4.INT.udn
@@ -33,7 +33,7 @@ Now that you can detect your projectile's collision interactions, you can determ
 
 		// Called when projectile hits something
 		UFUNCTION()
-		void OnHit(class AActor* OtherActor, class UPrimitiveComponent* OtherComp, FVector NormalImpulse, const FHitResult& Hit);
+		void OnHit(UPrimitiveComponent* HitComp, AActor* OtherActor, UPrimitiveComponent* OtherComp, FVector NormalImpulse, const FHitResult& Hit);
 
 1.	`FPSProjectile.h` should now look like the following:
 
@@ -75,14 +75,14 @@ Now that you can detect your projectile's collision interactions, you can determ
 		
 			// Called when projectile hits something
 			UFUNCTION()
-			void OnHit(class AActor* OtherActor, class UPrimitiveComponent* OtherComp, FVector NormalImpulse, const FHitResult& Hit);
+			void OnHit(UPrimitiveComponent* HitComp, AActor* OtherActor, UPrimitiveComponent* OtherComp, FVector NormalImpulse, const FHitResult& Hit);
 		};
 
 1.	Locate the `FPSProjectile` class implementation file in the **Solution Explorer** and open `FPSProjectile.cpp`.
 
 1.	Add the following code to the `OnHit` function definition:
 
-		void AFPSProjectile::OnHit(class AActor* OtherActor, class UPrimitiveComponent* OtherComp, FVector NormalImpulse, const FHitResult& Hit)
+		void AFPSProjectile::OnHit(UPrimitiveComponent* HitComp, AActor* OtherActor, UPrimitiveComponent* OtherComp, FVector NormalImpulse, const FHitResult& Hit)
 		{
 			if (OtherActor != NULL && OtherActor != this && OtherComp != NULL)
 			{
@@ -163,7 +163,7 @@ Now that you can detect your projectile's collision interactions, you can determ
 			}
 		}
 		
-		void AFPSProjectile::OnHit(class AActor* OtherActor, class UPrimitiveComponent* OtherComp, FVector NormalImpulse, const FHitResult& Hit)
+		void AFPSProjectile::OnHit(UPrimitiveComponent* HitComp, AActor* OtherActor, UPrimitiveComponent* OtherComp, FVector NormalImpulse, const FHitResult& Hit)
 		{
 			if (OtherActor != NULL && OtherActor != this && OtherComp != NULL)
 			{
diff --git a/Engine/Documentation/Source/Shared/GraphNodes/Animation/BlendNodes.INT.udn b/Engine/Documentation/Source/Shared/GraphNodes/Animation/BlendNodes.INT.udn
--- a/Engine/Documentation/Source/Shared/GraphNodes/Animation/BlendNodes.INT.udn
+++ b/Engine/Documentation/Source/Shared/GraphNodes/Animation/BlendNodes.INT.udn
@@ -209,3 +209,17 @@ Modifies a bone transform to look at another bone along a specified axis, or to
 Handles re-targeting of an IK Bone chain. Moves IK bone chain to meet FK hand bones, based on HandFKWeight (to favor either side). (0 = favor left hand, 1 = favor right hand, 0.5 = equal weight).
 
 [/EXCERPT:UAnimGraphNode_HandIKRetargeting]
+
+[EXCERPT:GraphNode_FastPathInfo]
+
+Animation nodes can use the 'fast path', avoiding somewhat slower Blueprint usage by ensuring that they either:
+- Access member variables directly.
+- Access a negated boolean member variable.
+- Access members of a nested struct.
+- Access some "Break Struct" cases as well (Break Transform being a notable exception).
+
+[VAR:ToolTipFullLink]
+Engine/Animation/Optimization/FastPath
+[/VAR]
+
+[/EXCERPT:GraphNode_FastPathInfo]
\ No newline at end of file
