New clause about avoiding namespaces in UHT-parsed headers.

diff --git a/Engine/Documentation/Source/Programming/Development/CodingStandard/CodingStandard.INT.udn b/Engine/Documentation/Source/Programming/Development/CodingStandard/CodingStandard.INT.udn
--- a/Engine/Documentation/Source/Programming/Development/CodingStandard/CodingStandard.INT.udn
+++ b/Engine/Documentation/Source/Programming/Development/CodingStandard/CodingStandard.INT.udn
@@ -561,6 +561,7 @@ You can use namespaces to organize your classes, functions and variables where a
 * Note that forward-declared types need to be declared within their respective namespace, otherwise you will get link errors.
 * If you declare a lot of classes/types within a namespace, it can make it difficult to use those types in other global-scoped classes. (e.g. function signatures will need to use explicit namespace when appearing in class declarations.)
 * You can use "using" to alias only specific variables within a namespace into your scope (e.g. using Foo::FBar), but we do not usually do that in Unreal code.
+* Namespaces are not supported by UnrealHeaderTool, so should not be used when defining UCLASSes, USTRUCTs etc.
 
 
 ## Physical Dependencies
