improved documentation

diff --git a/Engine/Documentation/Source/Programming/Development/Tools/ConsoleManager/ConsoleManager.INT.udn b/Engine/Documentation/Source/Programming/Development/Tools/ConsoleManager/ConsoleManager.INT.udn
--- a/Engine/Documentation/Source/Programming/Development/Tools/ConsoleManager/ConsoleManager.INT.udn
+++ b/Engine/Documentation/Source/Programming/Development/Tools/ConsoleManager/ConsoleManager.INT.udn
@@ -218,6 +218,14 @@ Those settings can be changed in the editor UI. Project settings should not inte
 
 Other settings can come from the Scalability feature. Look at **Config/BaseScalability.ini** or the Scalability documentation for more info.
 
+## Command line
+
+The command line allows to set console variables, call console commands or exec commands. This an example:
+
+	UE4Editor.exe GAMENAME -ExecCmds="r.BloomQuality 12;vis 21;Quit"
+
+Here we execute 3 commands. Note that setting a console variable this way requires you to omit the '=' you would need in an ini file.
+
 ## Priority
 
 Console variables can be overridden from various sources e.g. user/editor/project settings, command line, consolevariables.ini, ...
