New rules about lambda captures.

diff --git a/Engine/Documentation/Source/Programming/Development/CodingStandard/CodingStandard.INT.udn b/Engine/Documentation/Source/Programming/Development/CodingStandard/CodingStandard.INT.udn
--- a/Engine/Documentation/Source/Programming/Development/CodingStandard/CodingStandard.INT.udn
+++ b/Engine/Documentation/Source/Programming/Development/CodingStandard/CodingStandard.INT.udn
@@ -387,10 +387,84 @@ Lambdas can now be used on all compilers, but their usage should be considered.
 		// Sort array in reverse order of name
 		AnotherArray.Sort([](const Thing& Lhs, const Thing& Rhs){ return Lhs.GetName() > Rhs.GetName(); });
 
-Be mindful of stateful lambdas that capture stack variables and whether you capture by value or reference -- a by-value capture of a non-trivial type could be a performance problem. Also, stateful lambdas cannot be assigned to function pointers, which we tend to use a lot.
+Be aware that stateful lambdas cannot be assigned to function pointers, which we tend to use a lot.
 
 Documentation of non-trivial lambdas and anonymous functions should be considered in the same way a regular function is documented. Don't be afraid to split them over a few more lines in order to comment them.
 
+Prefer explicit capture over automatic capture ([&] and [=]), especially for large lambdas and deferred execution. Accidentally capturing a variable with the wrong capture semantics can have negative consequences, and this is more likely to happen as the code is maintained over time:
+
+* By-reference capture and by-value capture of pointers can cause accidental dangling references, if the lambda is executed outside the context of the captured variables:
+
+		void Func()
+		{
+			int32 Value = GetSomeValue();
+
+			// Lots of code
+
+			AsyncTask([&]()
+			{
+				// Value is invalid here
+				for (int Index = 0; Index != Value; ++Index)
+				{
+					// ...
+				}
+			});
+		}
+
+* By-value capture can be a performance concern because of unnecessary copies:
+
+		void Func()
+		{
+			int32 ValueToFind = GetValueToFind();
+
+			// The lambda takes a copy of ArrayOfThings because it is accidentally captured by [=] when it was only meant to capture ValueToFind
+			FThing* Found = ArrayOfThings.FindByPredicate(
+				[=](const FThing& Thing)
+				{
+					return Thing.Value == ValueToFind && Thing.Index < ArrayOfThings.Num();
+				}
+			);
+		}
+
+* Accidentally captured UObject pointers are invisible to the garbage collector:
+
+		void Func(AActor* MyActor)
+		{
+			// Capturing the MyActor pointer will not prevent the object from being collected
+			AsyncTask([=]()
+			{
+				MyActor->DoSlowThing();
+			});
+		}
+
+* Automatic capture always captures 'this' implicitly if any member variables are referenced, even with [=]. [=] gives the impression of the lambda having its own copy of the member when it doesn't:
+
+		void FStruct::Func()
+		{
+			int32 Local = 5;
+			Member = 5;
+
+			auto Lambda = [=]()
+			{
+				UE_LOG(LogTest, Log, TEXT("Local: %d, Member: %d"), Local, Member);
+			};
+
+			Local = 100;
+			Member = 100;
+
+			Lambda(); // Logs "Local: 5, Member: 100"
+		}
+
+Prefer explicit return types for large lambdas or when you are returning the result of another function call. These should be considered in the same way as the 'auto' keyword:
+
+		// Without the return type here, the return type is unclear
+		auto Lambda = []() -> FMyType
+		{
+			return SomeFunc();
+		};
+
+Automatic captures and implicit return types are acceptable for trivial lambdas, e.g. in Sort calls, where the semantics are obvious and being explicit would make it overly verbose - use your best judgment.
+
 
 ### Strongly-Typed Enums
 
