Comments about custom deleter support not being supported removed.

diff --git a/Engine/Documentation/Source/Programming/UnrealArchitecture/SmartPointerLibrary/SmartPointerLibrary.INT.udn b/Engine/Documentation/Source/Programming/UnrealArchitecture/SmartPointerLibrary/SmartPointerLibrary.INT.udn
--- a/Engine/Documentation/Source/Programming/UnrealArchitecture/SmartPointerLibrary/SmartPointerLibrary.INT.udn
+++ b/Engine/Documentation/Source/Programming/UnrealArchitecture/SmartPointerLibrary/SmartPointerLibrary.INT.udn
@@ -194,7 +194,6 @@ The SharedPointerTesting.h file (located in [UE4RootLocation]/Engine/Source/Runt
 ### Limitations
 
 * Shared pointers are not compatible with Unreal objects (UObject classes)!
-* Currently only types with that have regular destructors (no custom deleters).
 * Dynamically-allocated arrays are not supported yet (e.g. `MakeSharable( new int32[20] )`).
 * Implicit conversion of TSharedPtr/TSharedRef to bool is not supported yet.
 
@@ -206,7 +205,6 @@ The SharedPointerTesting.h file (located in [UE4RootLocation]/Engine/Source/Runt
 * TSharedFromThis returns a shared reference, not a shared pointer.
 * Some features were omitted (e.g. use_count(), unique(), etc.).
 * No exceptions are allowed (all related features have been omitted).
-* Custom allocators and custom delete functions are not supported yet.
 * Our implementation supports non-nullable smart pointers (TSharedRef).
 * Several other new features added, such as MakeShareable and NULL assignment.
 
