Community contribution - fixing up typos and other misc errors

JIRA UE-30043

GitHub PR #2337

diff --git a/Engine/Documentation/Source/Engine/AI/BehaviorTrees/NodeReference/Services/Services.INT.udn b/Engine/Documentation/Source/Engine/AI/BehaviorTrees/NodeReference/Services/Services.INT.udn
--- a/Engine/Documentation/Source/Engine/AI/BehaviorTrees/NodeReference/Services/Services.INT.udn
+++ b/Engine/Documentation/Source/Engine/AI/BehaviorTrees/NodeReference/Services/Services.INT.udn
@@ -16,7 +16,7 @@ The Default Focus creates a short cut to accessing an Actor in your Blueprints a
 
 | Property | Description |
 | -------- | ----------- |
-| **Blackboard Key**| This takes in an Actor Key (and Object Key with its base class set to Actor) to be quickly referenced by the AI Controller using `GetFocusActor`. |
+| **Blackboard Key**| This takes in an Actor Key (an Object Key with its base class set to Actor) to be quickly referenced by the AI Controller using `GetFocusActor`. |
 | **Node Name** | The name the node should display in the Behavior Tree graph. |
 
 [REGION:note]
diff --git a/Engine/Documentation/Source/Engine/AI/EnvironmentQuerySystem/NodeReference/NodeReference.INT.udn b/Engine/Documentation/Source/Engine/AI/EnvironmentQuerySystem/NodeReference/NodeReference.INT.udn
--- a/Engine/Documentation/Source/Engine/AI/EnvironmentQuerySystem/NodeReference/NodeReference.INT.udn
+++ b/Engine/Documentation/Source/Engine/AI/EnvironmentQuerySystem/NodeReference/NodeReference.INT.udn
@@ -71,7 +71,7 @@ Finds all of the Actors of a class within the **Search Radius** around the **Sea
 
 ### On circle
 
-Radiates traces out from the **Circle Center**, returning the hits or the edge of the radius (of the generated Circle) as Items.
+Radiates traces out from the **Circle Center**, returning the hits on the edge of the radius (of the generated Circle) as Items.
 
 ![](gen_circle.png)(convert:false)
 
diff --git a/Engine/Documentation/Source/Engine/AI/EnvironmentQuerySystem/QuickStart/14/14.INT.udn b/Engine/Documentation/Source/Engine/AI/EnvironmentQuerySystem/QuickStart/14/14.INT.udn
--- a/Engine/Documentation/Source/Engine/AI/EnvironmentQuerySystem/QuickStart/14/14.INT.udn
+++ b/Engine/Documentation/Source/Engine/AI/EnvironmentQuerySystem/QuickStart/14/14.INT.udn
@@ -25,11 +25,11 @@ This page covers a number of ideas for improvements you can make now that you're
 
 ## Ideas
 
-* TopDownCharacter is far to specific, update everything to a more flexible system.
+* TopDownCharacter is far too specific, update everything to a more flexible system.
 * The Behavior Tree shows one way to set a system like this up:
 	* Try using a Wait node instead of the Cooldown Decorator. This will require a Bool Blackboard entry and some extra work in the AI Controller.
 	* Add a secondary Behavior that, after a time, the AI will begin to wander around.
-	* Bonus - Store the last know location(s) of the Player(s) and use those in an Environment Query to find a better hiding spot.
+	* Bonus - Store the last known location(s) of the Player(s) and use those in an Environment Query to find a better hiding spot.
 	* Double Bonus - Have the AI attempt to approach a Player from behind without entering the forward view arc of the Player(s).
 * Did you know you can create your own Generators in Blueprint? Maybe there's a better way to find the possible points for the AI to hide.
 * [REGION:raw]The ' key (single quote) turns on a host of Gameplay Debugging tools. Check out all of the options available in the **Editor Preferences** %arrow% **Gameplay Debugger**[/REGION]
diff --git a/Engine/Documentation/Source/Engine/Actors/DecalActor/DecalUserGuide.INT.udn b/Engine/Documentation/Source/Engine/Actors/DecalActor/DecalUserGuide.INT.udn
--- a/Engine/Documentation/Source/Engine/Actors/DecalActor/DecalUserGuide.INT.udn
+++ b/Engine/Documentation/Source/Engine/Actors/DecalActor/DecalUserGuide.INT.udn
@@ -153,4 +153,4 @@ And using it fixes the artifacts on the decal that benefits from it.
 
 ### Current limitations
 
-* DecalMipmapLevel doesn't support custom UV, but you can patch'up it's output.
+* DecalMipmapLevel doesn't support custom UV, but you can patch up its output.
diff --git a/Engine/Documentation/Source/Engine/Animation/Blendspaces/Blendspaces.INT.udn b/Engine/Documentation/Source/Engine/Animation/Blendspaces/Blendspaces.INT.udn
--- a/Engine/Documentation/Source/Engine/Animation/Blendspaces/Blendspaces.INT.udn
+++ b/Engine/Documentation/Source/Engine/Animation/Blendspaces/Blendspaces.INT.udn
@@ -187,7 +187,7 @@ When setting up your Blend Space asset, there are additional options that can be
 
 |---|---|
 |**Axis to Scale Animation**| Not available for 	1D Blend Space assets, this allows you (if you have input interpolation) to determine which axis to drive animation speed (scale). For example, for locomotion animation, speed axis will drive animation speed (thus scale).|
-|**Interpolation Params**| Input interpolation parameter for all 3 axis, for each axis input, decide how you’d like to interpolate input to by setting Interpolation Time and Interpolation Type. |
+|**Interpolation Params**| Input interpolation parameter for all 3 axis, for each axis input, decide how you'd like to interpolate input to by setting Interpolation Time and Interpolation Type. |
 
 ### Sample Interpolation
 
diff --git a/Engine/Documentation/Source/Engine/Animation/NodeReference/SkeletalControls/SkeletalControls.INT.udn b/Engine/Documentation/Source/Engine/Animation/NodeReference/SkeletalControls/SkeletalControls.INT.udn
--- a/Engine/Documentation/Source/Engine/Animation/NodeReference/SkeletalControls/SkeletalControls.INT.udn
+++ b/Engine/Documentation/Source/Engine/Animation/NodeReference/SkeletalControls/SkeletalControls.INT.udn
@@ -108,7 +108,7 @@ Scale - from one bone to another.
 
 The **Bone Driven Controller** allows a "Driver" bone to dynamically affect the motion of one or more "Driven" bones. This is useful for characters that have accessories attached as you can use this to avoid geometry from intersecting during an animation, even when a lot of blending is being used. 
 
-You can either set the "Driver" value directly with a **Multiplier**, remap it to a whole new range, or use a **Curve** asset to drive the motion. Using a Curve is usually the best approach as it lest you define the response naturally and interactively tweak points/tangents while seeing the changes reflected in real-time.  
+You can either set the "Driver" value directly with a **Multiplier**, remap it to a whole new range, or use a **Curve** asset to drive the motion. Using a Curve is usually the best approach as it lets you define the response naturally and interactively tweak points/tangents while seeing the changes reflected in real-time.  
 
 In the example below, the attached accessory (green) has no authored animations, and is being driven in two axes as a function of the character's thigh bone. This is all calculated at runtime, so a blend of animations ends up working quite well here, even without hand-authored tweaking. 
 
diff --git a/Engine/Documentation/Source/Engine/Animation/Persona/Modes/Mesh/Mesh.INT.udn b/Engine/Documentation/Source/Engine/Animation/Persona/Modes/Mesh/Mesh.INT.udn
--- a/Engine/Documentation/Source/Engine/Animation/Persona/Modes/Mesh/Mesh.INT.udn
+++ b/Engine/Documentation/Source/Engine/Animation/Persona/Modes/Mesh/Mesh.INT.udn
@@ -49,7 +49,7 @@ In the **LOD Info** section, there are settings that allow you to modify how the
 | Property| Description |
 |--- | --- |
 | **Display Factor** | Indicates when to use the LOD. A smaller number means use the LOD when further away.|
-| **LODHysteresis** | Used to avoid "flickering" when on LOD boundary. Only taken into account when moving from complex to single. |
+| **LODHysteresis** | Used to avoid "flickering" when on LOD boundary. Only taken into account when moving from complex to simple. |
 | **LODMaterial Map** | Specifies the mapping table from the LOD's materials to the USkeletalMesh materials array. |
 | **Triangle Sorting** | Set the current triangle sorting method. |
 | **Custom Left Right Axis** | Set the Left Right Axis.  |
diff --git a/Engine/Documentation/Source/Engine/Animation/Persona/Toolbar/PersonaToolbar.INT.udn b/Engine/Documentation/Source/Engine/Animation/Persona/Toolbar/PersonaToolbar.INT.udn
--- a/Engine/Documentation/Source/Engine/Animation/Persona/Toolbar/PersonaToolbar.INT.udn
+++ b/Engine/Documentation/Source/Engine/Animation/Persona/Toolbar/PersonaToolbar.INT.udn
@@ -94,7 +94,7 @@ Persona toolbar buttons provide easy access to switching between asset modes, su
 | ![](button_Compression.png)																				| Brings multiple options of compression for the currently selected animation. |
 | ![](button_AnimNotifies.png)																				| Brings up the Anim Notifies window, which displays all Anim Notifies. |
 | ![](button_Key.png)																						| Adds a key to the Additive Layer Track of the currently selected bone for the currently selected animation. |
-| ![](button_Apply.png)																						| Applys added keys to the currently selected animation. |
+| ![](button_Apply.png)																						| Applies added keys to the currently selected animation. |
 
 ## Mode Switching Tabs
 
diff --git a/Engine/Documentation/Source/Engine/Animation/RootMotion/RootMotion.INT.udn b/Engine/Documentation/Source/Engine/Animation/RootMotion/RootMotion.INT.udn
--- a/Engine/Documentation/Source/Engine/Animation/RootMotion/RootMotion.INT.udn
+++ b/Engine/Documentation/Source/Engine/Animation/RootMotion/RootMotion.INT.udn
@@ -103,7 +103,7 @@ In the image above, the Root Motion of **Jog_Loop_Fwd_RM** and **Jog_Loop_Right_
 
 ### Root Motion from Montages Only
 
-This method is intended for network games which is the reason for its more restricted functionally. Because **Root Motion from Everything** is designed to be unrestricted, it is recommended that it only be used in games where animation isn’t replicated over a network.  
+This method is intended for network games which is the reason for its more restricted functionally. Because **Root Motion from Everything** is designed to be unrestricted, it is recommended that it only be used in games where animation isn't replicated over a network.  
 
 
 
diff --git a/Engine/Documentation/Source/Engine/Physics/Destructibles/DestructibleProperties/DestructibleProperties.INT.udn b/Engine/Documentation/Source/Engine/Physics/Destructibles/DestructibleProperties/DestructibleProperties.INT.udn
--- a/Engine/Documentation/Source/Engine/Physics/Destructibles/DestructibleProperties/DestructibleProperties.INT.udn
+++ b/Engine/Documentation/Source/Engine/Physics/Destructibles/DestructibleProperties/DestructibleProperties.INT.udn
@@ -39,7 +39,7 @@ Version: 4.9
 | **Damage Spread** | Controls the distance into the destructible Actor to propagate the damage. The damage applied to the chunk is multiplied by Damage Spread to get the propagation distance. |
 | **Enable Impact Damage** | Whether to apply damage to a destructible when colliding with an object. |
 | **Impact Damage** | If a chunk is at a depth which has impact damage set, it will take damage equal to Impact Damage times the impact force when it collides with another object (assuming Report Rigid Collisions are enabled on both objects). |
-| **Custom Impact Resistance** | By default objects that collide with destructibles will bounce back. Custom resistance allows for finer control of how much a destrucible "pushes back" against a colliding object. |
+| **Custom Impact Resistance** | By default objects that collide with destructibles will bounce back. Custom resistance allows for finer control of how much a destructible "pushes back" against a colliding object. |
 | **Impact Resistance** | When a chunk takes impact damage due to physical contact, this parameter is the maximum impulse the contact can generate. Weak materials will have this set low, so heavier objects can pass through them during fracture. 0.0 is interpreted at infinite. |
 | **Default Impact Damage** | Impact damage will be taken to this depth. If negative, impact damage is disabled. |
 | [REGION:tablesection] _Debris Parameters_ [/REGION] ||
diff --git a/Engine/Documentation/Source/Engine/Rendering/LightingAndShadows/LightTypes/Directional/Directional.INT.udn b/Engine/Documentation/Source/Engine/Rendering/LightingAndShadows/LightTypes/Directional/Directional.INT.udn
--- a/Engine/Documentation/Source/Engine/Rendering/LightingAndShadows/LightTypes/Directional/Directional.INT.udn
+++ b/Engine/Documentation/Source/Engine/Rendering/LightingAndShadows/LightTypes/Directional/Directional.INT.udn
@@ -179,7 +179,7 @@ Property | Description |
 | **Dynamic Shadow Distance MovableLight** | How far **Cascaded Shadow Map** dynamic shadows will cover for a movable light, measured from the camera. |
 | **Dynamic Shadow Distance StationaryLight** | How far **Cascaded Shadow Map** dynamic shadows will cover for a stationary light, measured from the camera. |
 | **Num Dynamic Shadow Cascades** | Number of cascades to split the view frustum into for the whole scene.  |
-| **Cascade Distribution Exponent** | Controls whether the cascades are distributed closer to the camera (larget exponent) or further from the camera (smaller exponent). |
+| **Cascade Distribution Exponent** | Controls whether the cascades are distributed closer to the camera (larger exponent) or further from the camera (smaller exponent). |
 | **Cascade Transition Fraction** | Proportion of the fade region between cascades. |
 | **Shadow Distance Fadeout Fraction** | Controls the size of the fade out region at the far extent of the dynamic shadow's influence.|
 | **Use Inset Shadows for Movable Objects** | (Stationary lights only) Whether to use per-object inset shadows for movable components, even though cascaded shadow maps are enabled.|
diff --git a/Engine/Documentation/Source/Engine/Rendering/Materials/ExpressionReference/Coordinates/CoordinatesExpressions.INT.udn b/Engine/Documentation/Source/Engine/Rendering/Materials/ExpressionReference/Coordinates/CoordinatesExpressions.INT.udn
--- a/Engine/Documentation/Source/Engine/Rendering/Materials/ExpressionReference/Coordinates/CoordinatesExpressions.INT.udn
+++ b/Engine/Documentation/Source/Engine/Rendering/Materials/ExpressionReference/Coordinates/CoordinatesExpressions.INT.udn
@@ -17,7 +17,7 @@ Version:4.9
 
 ## LightmapUVs
 
-The **LightmapUVs** expression outputs the lightmap UV texture coordinates in the forum of a two-channel vector value. If lightmap UVs are unavailable, it will output a two-channel vector value of (0,0).
+The **LightmapUVs** expression outputs the lightmap UV texture coordinates in the form of a two-channel vector value. If lightmap UVs are unavailable, it will output a two-channel vector value of (0,0).
 
 
 [INCLUDE:Engine/Rendering/Materials/ExpressionReference/Vector#ObjectOrientation]
diff --git a/Engine/Documentation/Source/Engine/Rendering/Materials/ExpressionReference/Utility/UtilityExpressions.INT.udn b/Engine/Documentation/Source/Engine/Rendering/Materials/ExpressionReference/Utility/UtilityExpressions.INT.udn
--- a/Engine/Documentation/Source/Engine/Rendering/Materials/ExpressionReference/Utility/UtilityExpressions.INT.udn
+++ b/Engine/Documentation/Source/Engine/Rendering/Materials/ExpressionReference/Utility/UtilityExpressions.INT.udn
@@ -312,7 +312,7 @@ In this example, the preview object will smoothly fade to black as the camera ex
 
 ## AntialiasedTextureMask
 
-The **AntialiasedTextureMask** expression allows you to create a material using a soft (anti-aliased) transition mask. The mask can be used to blend between two complex material properties or to fade out an alpha blended material (works well with SoftMasked). Simple specify a texture with the mask specified in one channel (red, green, blue, or alpha), set the used channel in the expression and specify the comparison value. Assuming the channel stores a grayscale value in the range 0=black to 1=white the comparison function defines if the resulting mask should be 0 or 1. This expression is a parameter, allowing the **Texture** property to be overridden by child MaterialInstances.
+The **AntialiasedTextureMask** expression allows you to create a material using a soft (anti-aliased) transition mask. The mask can be used to blend between two complex material properties or to fade out an alpha blended material (works well with SoftMasked). Simply specify a texture with the mask specified in one channel (red, green, blue, or alpha), set the used channel in the expression and specify the comparison value. Assuming the channel stores a grayscale value in the range 0=black to 1=white the comparison function defines if the resulting mask should be 0 or 1. This expression is a parameter, allowing the **Texture** property to be overridden by child MaterialInstances.
 
 | Item | Description |
 | -------- | ----------- |
diff --git a/Engine/Documentation/Source/Engine/Rendering/Materials/Functions/Reference/WorldPositionOffset/WorldPositionOffset.INT.udn b/Engine/Documentation/Source/Engine/Rendering/Materials/Functions/Reference/WorldPositionOffset/WorldPositionOffset.INT.udn
--- a/Engine/Documentation/Source/Engine/Rendering/Materials/Functions/Reference/WorldPositionOffset/WorldPositionOffset.INT.udn
+++ b/Engine/Documentation/Source/Engine/Rendering/Materials/Functions/Reference/WorldPositionOffset/WorldPositionOffset.INT.udn
@@ -47,7 +47,7 @@ The **SplineThicken** function serves as a way to make very thin polygons look t
 | **UVs for texturing (Vector2)** | This is the UV channel for texturing. This must be the same UV channel you want to have the 3D-parallax correction. |
 | **Expand U or V UV Channel (StaticBool)** | Sets whether to expand the mesh in U or V. Defaults to *true*, which is U.|
 | **DriveNormalZ (Boolean)** | This will use DeriveNormalZ to make a perfectly round normal map. |
-| **AngleCorrectedNormal (Boolean)** | Uses DerivedNormalZ to make perfectly round normal map. Requires a center tesselation vertex otherwise the vale of Z will be 0 for the entire face, causing harsh lighting. |
+| **AngleCorrectedNormal (Boolean)** | Uses DerivedNormalZ to make perfectly round normal map. Requires a center tesselation vertex otherwise the value of Z will be 0 for the entire face, causing harsh lighting. |
 | **AdditionalNormal (Vector3)** | This will add a texture normal to the normalmap transform. |
 | **FlattenPixelNormal (Boolean)** |  |
 | **UVs For NormalShading (Vector2)** |  |
@@ -67,7 +67,7 @@ The **CameraOffset** function is great to assist with depth sorting, as it allow
 | -------- | -------- |
 |[REGION:tablesection]Inputs[/REGION] ||
 | **OffsetAmount (Scalar)** | Enter a negative or positive number to offset your model in the direction of the camera. Note that positive values will bring the model closer to the camera and cause rendering errors if the mesh greatly exceeds the models bounding box. |
-| **Clampe Padding (Scalar)** | The amount of padding to prevent the clamped version of the offset value from clipping into the camera.  |
+| **Clamp Padding (Scalar)** | The amount of padding to prevent the clamped version of the offset value from clipping into the camera.  |
 | **WorldPosition (Vector3)** | Enter the model's vertex world position. The Default Value = World Position. |
 |[REGION:tablesection]Outputs[/REGION] ||
 | **Clamped Camera Offset** | The Camera offset is clamped to avoid early camera intersections. Adjust the padding to change the amount of space reserved to avoid camera intersection. |
@@ -126,8 +126,8 @@ The **RotateAboutWorldAxis_cheap** function cheaply rotates objects around world
 | -------- | -------- |
 |[REGION:tablesection]Inputs[/REGION] ||
 | **Rotation Amount (Scalar)** | A value of one is equal to one full rotation.|
-| **PivotPoint (Vector3)** | World space location of individual vertices. The WorldPosition node is typically used.|
-| **WorldPosition (Vector3)** | World space location of pivot point to rotate around. The default value is the model's pivot point.|
+| **PivotPoint (Vector3)** | World space location of pivot point to rotate around. The default value is the model's pivot point.|
+| **WorldPosition (Vector3)** | World space location of individual vertices. The WorldPosition node is typically used.|
 |[REGION:tablesection]Outputs[/REGION] ||
 | **X-Axis** | Cheaply rotates the object around the world X-axis. |
 | **Y-Axis** | Cheaply rotates the object around the world Y-axis. |
diff --git a/Engine/Documentation/Source/Engine/Rendering/Materials/LayeredMaterials/LayeredMaterials.INT.udn b/Engine/Documentation/Source/Engine/Rendering/Materials/LayeredMaterials/LayeredMaterials.INT.udn
--- a/Engine/Documentation/Source/Engine/Rendering/Materials/LayeredMaterials/LayeredMaterials.INT.udn
+++ b/Engine/Documentation/Source/Engine/Rendering/Materials/LayeredMaterials/LayeredMaterials.INT.udn
@@ -78,7 +78,7 @@ Within the Material Functions list of the Material Editor is a list of various M
 |**MatLayerBlend_BreakBaseColor**|Outputs the Base Color from an incoming Material Layer.|
 |**MatLayerBlend_BreakNormal**|Outputs the Normal from an incoming Material Layer.|
 |**MatLayerBlend_Decal**|Blends a decal sheet over the Material using the 2nd UV channel.|
-|**MatLayerBlend_Decal_UV3**|Blends a decal sheet over the Material Layer using the 2nd UV channel.|
+|**MatLayerBlend_Decal_UV3**|Blends a decal sheet over the Material Layer using the 3rd UV channel.|
 |**MatLayerBlend_Emissive**|Blends an Emissive texture over the Material Layer.|
 |**MatLayerBlend_GlobalNormal**|Blends a Normal texture over the Material Layer.|
 |**MatLayerBlend_LightmassReplace**|Replaces the Base Color in Lightmass, allowing for changes to indirect lighting results.|
diff --git a/Engine/Documentation/Source/Engine/Rendering/ParticleSystems/Reference/Modules/Collision/CollisionModules.INT.udn b/Engine/Documentation/Source/Engine/Rendering/ParticleSystems/Reference/Modules/Collision/CollisionModules.INT.udn
--- a/Engine/Documentation/Source/Engine/Rendering/ParticleSystems/Reference/Modules/Collision/CollisionModules.INT.udn
+++ b/Engine/Documentation/Source/Engine/Rendering/ParticleSystems/Reference/Modules/Collision/CollisionModules.INT.udn
@@ -103,8 +103,8 @@ After the Min Desired Frame Rate has been changed, re-cook or re-deploy your pro
 |_EPCC_Freeze_|Freeze the particle in place.|
 |_EPCC_HaltCollisions_|Stop collision checks, but keep updating. This will likely lead to particles `falling through` objects.|
 |_EPCC_FreezeTranslation_|Stop translating the particle, but keep updating everything else.|
-|_EPCC_FreeRotation_|Stop rotating the particle, but keep updating everything else.|
-|_EPCC_FreeMovement_|Stop translating/rotating the particle, but keep updating everything else.|
+|_EPCC_FreezeRotation_|Stop rotating the particle, but keep updating everything else.|
+|_EPCC_FreezeMovement_|Stop translating/rotating the particle, but keep updating everything else.|
 [/REGION]
 [/EXCERPT:CollisionCompletionOption]
 
diff --git a/Engine/Documentation/Source/Engine/Rendering/ParticleSystems/Reference/TypeData/Mesh/MeshTypeData.INT.udn b/Engine/Documentation/Source/Engine/Rendering/ParticleSystems/Reference/TypeData/Mesh/MeshTypeData.INT.udn
--- a/Engine/Documentation/Source/Engine/Rendering/ParticleSystems/Reference/TypeData/Mesh/MeshTypeData.INT.udn
+++ b/Engine/Documentation/Source/Engine/Rendering/ParticleSystems/Reference/TypeData/Mesh/MeshTypeData.INT.udn
@@ -24,7 +24,7 @@ The Mesh type data module has the following properties:
 |**Pitch**| The 'pre' rotation pitch (in degrees) to apply to the Static Mesh used.|
 |**Roll**| The 'pre' rotation roll (in degrees) to apply to the Static Mesh used.|
 |**Yaw**| The 'pre' rotation yaw (in degrees) to apply to the Static Mesh used.|
-|**Axis Lock Option**| The axis to lock the mesh on. This overrides TypeSpecific mesh alignment as well as the LockAxis module. The following options are provided: [INCLUDE:#CameraFacingOption]|
+|**Axis Lock Option**| The axis to lock the mesh on. This overrides TypeSpecific mesh alignment as well as the LockAxis module. The following options are provided: [INCLUDE:#AxisLockOption]|
 |[REGION:tablesection]CameraFacing[/REGION]||
 |**Camera Facing**| If _true_, then the X-axis of the mesh will always point towards the camera. When set, _Axis Lock Option_ as well as all other locked axis/screen alignment settings are ignored.|
 |**Camera Facing Option**| Determines how the mesh is oriented when the _Camera Facing_ option is enabled. The following options are provided: [INCLUDE:#CameraFacingOption]|
diff --git a/Engine/Documentation/Source/Engine/Rendering/PostProcessEffects/PostProcessMaterials/PostProcessMaterials.INT.udn b/Engine/Documentation/Source/Engine/Rendering/PostProcessEffects/PostProcessMaterials/PostProcessMaterials.INT.udn
--- a/Engine/Documentation/Source/Engine/Rendering/PostProcessEffects/PostProcessMaterials/PostProcessMaterials.INT.udn
+++ b/Engine/Documentation/Source/Engine/Rendering/PostProcessEffects/PostProcessMaterials/PostProcessMaterials.INT.udn
@@ -179,7 +179,7 @@ Here you can see how we used the custom depth input to visualize a silhouette ar
 
 Notice that the former image has no anti-aliasing on the silhouette but in motion you also would see the silhouette jittering around 1 pixel.
 This is because temporal anti-aliasing moves the rendering of the whole scene by a sub pixel each frame. Multiple frames together are getting combined to produce the final anti-aliased image.
-We can however move the material to an earlier sport in the post processing graph to fix this issue.
+We can however move the material to an earlier spot in the post processing graph to fix this issue.
 
 This is the result:
 
@@ -194,7 +194,7 @@ but ideally this should not be.
 With post process materials you can look up into screen aligned buffers, but you need to know the right UVs. 
 The **ScreenPosition** material expression with the mapping option set to **ViewportUV** outputs the UV with 0,0 at the left top of the viewport and 1,1 at the bottom right.
 In contrast, using the **SceneTextureUV** mapping in this material expression might give you different results. This is because the actual texture (more correctly it is a render target) is potentially larger than the viewport.
-It can be larger in editor because we share this texture for multiple viewports and the largest extent is use for all viewports.
+It can be larger in editor because we share this texture for multiple viewports and the largest extent is used for all viewports.
 Even in game it can be larger in some cases (e.g. SceneCaptureActors might have a smaller viewport, Matinee black borders, Splitscreen, VR, ...).
 The **SceneTextureUV** option gives you the UV for this larger texture. If you only need a relative offset (e.g. pixel sized edge detection) you need to scale with the right size.
 The **SceneTexture** material expression has outputs for the size and the inverse of the size (efficient and useful for pixel offsets).
diff --git a/Engine/Documentation/Source/Engine/UI/LevelEditor/Viewports/ShowFlags/ShowFlags.INT.udn b/Engine/Documentation/Source/Engine/UI/LevelEditor/Viewports/ShowFlags/ShowFlags.INT.udn
--- a/Engine/Documentation/Source/Engine/UI/LevelEditor/Viewports/ShowFlags/ShowFlags.INT.udn
+++ b/Engine/Documentation/Source/Engine/UI/LevelEditor/Viewports/ShowFlags/ShowFlags.INT.udn
@@ -153,7 +153,6 @@ The Show Flags are found under the Viewport **Show** menu. They are a means by w
 |**Test Image**|Shows a test image for configuration of monitor and graphics settings.|
 |**Vector Fields**|Shows/hides any vector fields in the scene.|
 
-|**Blur GBuffer**|Shows/hides the results of Blue GBuffer Intensity, a property found on Post Process Volumes. This feature blurs the GBuffer normals and adjusts the GBuffer roughness to reduce specular aliasing. It is, however, a high performance cost for the result and you should generally consider just using Temporal AA.  |
 
 ## Visualize
 
diff --git a/Engine/Documentation/Source/Engine/UMG/UserGuide/WidgetTypeReference/MenuAnchor/MenuAnchor.INT.udn b/Engine/Documentation/Source/Engine/UMG/UserGuide/WidgetTypeReference/MenuAnchor/MenuAnchor.INT.udn
--- a/Engine/Documentation/Source/Engine/UMG/UserGuide/WidgetTypeReference/MenuAnchor/MenuAnchor.INT.udn
+++ b/Engine/Documentation/Source/Engine/UMG/UserGuide/WidgetTypeReference/MenuAnchor/MenuAnchor.INT.udn
@@ -26,7 +26,7 @@ In the **Details** panel for a placed **Menu Anchor**, there are a couple of spe
 | **Menu Class** | The Widget Blueprint (pop-up) to spawn when the menu is called. Creates the widget freshly each time. |
 | **Placement** | Several placement options use to determine the location of the created widget. |
 | **OnGetMenuContent** | Allows for the customization of the pop-up by binding a function or property to the **OnGetMenuContent** Event which is called when the pop-up is summoned (see below). |
-| **OnMenuOpenChanged** | Allows for the customization of the pop-up by binding a function or property to the **OnMenuOpenChanged** Event which is called opened state of the menu changes. |
+| **OnMenuOpenChanged** | Allows for the customization of the pop-up by binding a function or property to the **OnMenuOpenChanged** Event which is called when the opened state of the menu changes. |
 
 For the **OnGetMenuContent** Event, you can create and bind a function directly from the **Details** panel which will be called whenever the Menu Anchor is called. This allows you to provide extra functionality to the creation of your pop-up menu. For example, below we've run a check to determine if the player is "In Battle" and if not, we allow them to access the pop-up menu. If they are we do not allow them to access the pop-up menu. 
 
diff --git a/Engine/Documentation/Source/Shared/GraphNodes/Material/Material.INT.udn b/Engine/Documentation/Source/Shared/GraphNodes/Material/Material.INT.udn
--- a/Engine/Documentation/Source/Shared/GraphNodes/Material/Material.INT.udn
+++ b/Engine/Documentation/Source/Shared/GraphNodes/Material/Material.INT.udn
@@ -181,7 +181,7 @@ Engine/Rendering/Materials/ExpressionReference/Coordinates#camerapositionws
 [/EXCERPT:UMaterialExpressionCameraPositionWS]
 
 [EXCERPT:UMaterialExpressionLightmapUVs]
-The LightmapUVs expression outputs the lightmap UV texture coordinates in the forum of a two-channel vector value. If lightmap UVs are unavailable, it will output a two-channel vector value of (0,0). 
+The LightmapUVs expression outputs the lightmap UV texture coordinates in the form of a two-channel vector value. If lightmap UVs are unavailable, it will output a two-channel vector value of (0,0).
 
 [VAR:ToolTipFullLink]
 Engine/Rendering/Materials/ExpressionReference/Coordinates#lightmapuvs
@@ -910,7 +910,7 @@ Engine/Rendering/Materials/ExpressionReference/Utility#rotateaboutaxis
 [EXCERPT:UMaterialExpressionAntialiasedTextureMask]
 The AntialiasedTextureMask expression allows you to create a material using a soft (anti-aliased) transition mask. The mask can be used to blend between two complex material properties or to fade out an alpha blended material (works well with SoftMasked). 
 
-Simple specify a texture with the mask specified in one channel (red, green, blue, or alpha), set the used channel in the expression and specify the comparison value.
+Simply specify a texture with the mask specified in one channel (red, green, blue, or alpha), set the used channel in the expression and specify the comparison value.
 
 Assuming the channel stores a grayscale value in the range 0=black to 1=white the comparison function defines if the resulting mask should be 0 or 1. This expression is a parameter, allowing the Texture property to be overridden by child MaterialInstances. 
 
