Gameflow document. Needs a proper image for the thumbnail on the previous page, placeholder is submitted here.

diff --git a/Engine/Documentation/Source/Gameplay/Framework/GameFlow/GameFlow.INT.udn b/Engine/Documentation/Source/Gameplay/Framework/GameFlow/GameFlow.INT.udn
new file mode 100644
--- /dev/null
+++ b/Engine/Documentation/Source/Gameplay/Framework/GameFlow/GameFlow.INT.udn
@@ -0,0 +1,31 @@
+Availability:Public
+Title:Game Flow Overview
+Crumbs:%ROOT%, Gameplay/Framework
+Description:The process of starting the engine and launching a game or play-in-editor session.
+Version: 4.12
+type:reference
+parent:Gameplay/Framework
+order:5
+tags:Gameplay Framework
+
+This document explains the process of starting the engine and launching a game. There are two main paths shown here: the editor path, and the standalone path. The general order of events is to 
+initialize the engine, create and initialize a **GameInstance**, then load a level and start playing. However, there are differences between standalone and editor modes, both in terms of the 
+exact order of some of the functions that are called, and in which functions are called. The flowchart below shows the two paths running in parallel until they converge at the start of the game.
+
+## Game Flow Breakdown
+
+[![](GameFlowChart.png)](GameFlowChart.png)
+
+## Standalone
+
+In Standalone mode, which is the mode used by games played outside of the editor, the objects needed to play the game are created and initialized immediately following engine intialization on 
+startup. Objects such as the **GameInstance** are created and initialized before starting the engine (distinct from creating and initializing it), and the starting map is loaded immediately after the 
+engine's start function is called. Gameplay officially begins at that point with the level creating the appropriate [](Gameplay/Framework/GameMode), and then the other 
+[](Programming/UnrealArchitecture/Actors).
+
+## Editor
+
+In Editor mode, which is used by **Play In Editor** and **Simulate In Editor**, a different flow is used. The engine initializes and starts immediately, as it is needed to run the editor, but 
+creation and initializtion of objects such as the **GameInstance** are deferred until the user presses the button to launch the PIE or SIE session. In addition, the [](Programming/UnrealArchitecture/Actors) 
+in the level are duplicated so that in-game changes do not affect the level in the editor, and every object, including the **GameInstance** object, has a separate copy for each PIE instance. The editor 
+path rejoins the standalone path with the beginning of gameplay in the **UWorld** class.
diff --git a/Engine/Documentation/Source/Gameplay/Framework/GameplayFramework.INT.udn b/Engine/Documentation/Source/Gameplay/Framework/GameplayFramework.INT.udn
--- a/Engine/Documentation/Source/Gameplay/Framework/GameplayFramework.INT.udn
+++ b/Engine/Documentation/Source/Gameplay/Framework/GameplayFramework.INT.udn
@@ -30,6 +30,20 @@ to pages with overviews of each class and information about how they can be set
 
 [EXCERPT:LandingBoxes]
 [REGION:topics third]
+[OBJECT:Topic]
+	[PARAM:image]
+	![](gameflow_lander.png)
+	[/PARAM]
+	[PARAM:title]
+	%Gameplay/Framework/GameFlow:title%
+	[/PARAM]
+	[PARAM:path]
+	[RELATIVE:Gameplay/Framework/GameFlow]
+	[/PARAM]
+	[PARAM:description]
+	%Gameplay/Framework/GameFlow:description%
+	[/PARAM]
+[/OBJECT]
 [OBJECT:Topic]
 	[PARAM:image]
 	![](gamemode_lander.png)
