#UE4 docs adding info on the Set Sound Mix Class Override to the Overview of the sound system.

This will probably need a How-To once we get in and do some Audio How-To's

diff --git a/Engine/Documentation/Source/Engine/Audio/Overview/Overview.INT.udn b/Engine/Documentation/Source/Engine/Audio/Overview/Overview.INT.udn
--- a/Engine/Documentation/Source/Engine/Audio/Overview/Overview.INT.udn
+++ b/Engine/Documentation/Source/Engine/Audio/Overview/Overview.INT.udn
@@ -2,7 +2,7 @@ Availability: Public
 Title:Audio System Overview
 Crumbs: %ROOT%, Engine, Engine/Audio
 Description:Overview of the audio system used for playing in-game sounds, including the use of SoundCues node-based audio assets.
-version: 4.9
+version: 4.12
 
 [VAR:Topic]
 [OBJECT:Topic]
@@ -135,6 +135,12 @@ Sound Mixes allow you to set the **EQ Settings** (Equalizer Settings) and modify
 
 Multiple Sound Mixes can be active at the same time, all contributing to the overall audio effect. You can **Push** (Activate) or **Pop** (Deactivate) Sound Mixes directly inside a Blueprint with the **Push Sound Mix Modifier** and **Pop Sound Mix Modifier** nodes or activate them passively whenever a sound with a given Sound Class is playing within a specified threshold.
 
+However, the **Push/Pop** method can become very complex very quickly if you have a large number of mixes you are trying to switch between. This is where the **Set Sound Mix Class Override** Blueprint Node comes into play. It can set an active Sound Mix to use any Sound Class you have, and interpolate between its current Sound Class and the new Sound Class over time.
+
+![](image_37.jpg)
+
+Then you can set the Sound Mix back to its original setting by using the **Clear Sound Mix Class Override**.
+
 Within the Sound Mix asset itself, which can be opened by **Double-Clicking** the asset in the **Content Browser**, several properties exist. 
 You can specify EQ Settings for the mix to adjust the high, middle, and low frequencies and gains. As the EQ Settings of multiple Sound Mixes cannot be combined, the **EQ Priority** allows you to control which active mix's properties are applied at any given time.
 
