#UE4 docs wording change to the audio overview

diff --git a/Engine/Documentation/Source/Engine/Audio/Overview/Overview.INT.udn b/Engine/Documentation/Source/Engine/Audio/Overview/Overview.INT.udn
--- a/Engine/Documentation/Source/Engine/Audio/Overview/Overview.INT.udn
+++ b/Engine/Documentation/Source/Engine/Audio/Overview/Overview.INT.udn
@@ -135,7 +135,7 @@ Sound Mixes allow you to set the **EQ Settings** (Equalizer Settings) and modify
 
 Multiple Sound Mixes can be active at the same time, all contributing to the overall audio effect. You can **Push** (Activate) or **Pop** (Deactivate) Sound Mixes directly inside a Blueprint with the **Push Sound Mix Modifier** and **Pop Sound Mix Modifier** nodes or activate them passively whenever a sound with a given Sound Class is playing within a specified threshold.
 
-However, the **Push/Pop** method can become very complex very quickly if you have a large number of mixes you are trying to switch between. This is where the **Set Sound Mix Class Override** Blueprint Node comes into play. It can set an active Sound Mix to use any Sound Class you have, and interpolate between its current Sound Class and the new Sound Class over time.
+However, the **Push/Pop** method can become very complex, quickly if you have a large number of mixes you are trying to switch between. This is where the **Set Sound Mix Class Override** Blueprint Node comes into play. It can set an active Sound Mix to use any Sound Class you have, and interpolate between its current Sound Class and the new Sound Class over time.
 
 ![](image_37.jpg)
 
