#UE4 docs added overview for the doc

diff --git a/Engine/Documentation/Source/Engine/Performance/CookingStats/CookingStats.INT.udn b/Engine/Documentation/Source/Engine/Performance/CookingStats/CookingStats.INT.udn
--- a/Engine/Documentation/Source/Engine/Performance/CookingStats/CookingStats.INT.udn
+++ b/Engine/Documentation/Source/Engine/Performance/CookingStats/CookingStats.INT.udn
@@ -3,6 +3,10 @@ Title: Cooking Statistics
 Crumbs: %ROOT%, Engine
 Description:How to identify performance problems and fix them.
 
+While cooking a project, Unreal Engine 4 creates a number of detailed logs. One of these is the **Stats.csv** file. This document details how to read this file to hunt down cooking performance issues.
+
+## Reading the Log
+
 When you cook your project, statistics are always created. You can find them saved into:
 
 	<projectdirectory>\Saved\Stats\<datetime>Stats.csv
