#UE4 Docs: reduced size of inline image

diff --git a/Engine/Documentation/Source/Gameplay/Framework/GameFlow/GameFlow.INT.udn b/Engine/Documentation/Source/Gameplay/Framework/GameFlow/GameFlow.INT.udn
--- a/Engine/Documentation/Source/Gameplay/Framework/GameFlow/GameFlow.INT.udn
+++ b/Engine/Documentation/Source/Gameplay/Framework/GameFlow/GameFlow.INT.udn
@@ -12,7 +12,10 @@ This document explains the process of starting the engine and launching a game.
 
 ## Game Flow Breakdown
 
-[![](GameFlowChart.png)](GameFlowChart.png)
+[![](GameFlowChart.png)(w:320)](GameFlowChart.png)
+[REGION:caption]
+Click to view full size image
+[/REGION]
 
 ## Standalone
 
