#UE4 Docs: updated image extension

diff --git a/Engine/Documentation/Source/Support/Builds/ReleaseNotes/2016/4_12/ReleaseNotes_4_12.INT.udn b/Engine/Documentation/Source/Support/Builds/ReleaseNotes/2016/4_12/ReleaseNotes_4_12.INT.udn
--- a/Engine/Documentation/Source/Support/Builds/ReleaseNotes/2016/4_12/ReleaseNotes_4_12.INT.udn
+++ b/Engine/Documentation/Source/Support/Builds/ReleaseNotes/2016/4_12/ReleaseNotes_4_12.INT.udn
@@ -297,7 +297,7 @@ You can specify the Base Frame and Twist Frame using bone name and local axis, a
 
 With Full Scene Importing you can now take that environment you’ve already constructed in Max / Maya / Blender and import the entire scene into a level in the editor. 
 
-![image alt text](image_24.jpg)(w:201 h:477 convert:false) ![image alt text](image_25.jpg)(w:572 h:470 convert:false)
+![image alt text](image_24.jpg)(w:201 h:477 convert:false) ![image alt text](image_25.png)(w:572 h:470 convert:false)
 
  A new Scene FBX asset is created for your environment so that you can re-import it allowing iteration on the environment as well.  An early version was available in 4.11 supporting only scenes containing static meshes. The 4.12 version is more complete including the ability to import scenes that include skeletal meshes..
 
@@ -3582,4 +3582,3 @@ Collision can now be generated for Spline meshes created at runtime.
 * Changed how Hierarchical LOD system saves out generated content: meshes, materials and textures are now saved into a sub-folder called "HLOD".
 
 * Corrected URL path search when using "download maps on-the-fly" to allow pak file fetches to work from paths other than the "root" path.
-
