#UE4 Docs: fixed remaining image extension issues

diff --git a/Engine/Documentation/Source/Support/Builds/ReleaseNotes/2016/4_12/ReleaseNotes_4_12.INT.udn b/Engine/Documentation/Source/Support/Builds/ReleaseNotes/2016/4_12/ReleaseNotes_4_12.INT.udn
--- a/Engine/Documentation/Source/Support/Builds/ReleaseNotes/2016/4_12/ReleaseNotes_4_12.INT.udn
+++ b/Engine/Documentation/Source/Support/Builds/ReleaseNotes/2016/4_12/ReleaseNotes_4_12.INT.udn
@@ -10,7 +10,7 @@ This release brings hundreds of updates for Unreal Engine 4, including 106 impro
 
 aninda, amcofi, Artem V. Navrotskiy (bozaro), Beau Hardy (psychogony), Black Phoenix (PhoenixBlack), Brent Scriver (FineRedMist), Cedric Neukirchen (eXifreXi), Céleste (CelPlays), chendi (cdsama), Courtland Idstrom (Frugality), David Palacios (Hevedy), DavidNSllva, Erik (erikbye), esdasuka, Frank Force (KilledByAPixel), Graham Hughes (graehu), Grimmick, Hannah Gamiel (hgamiel), Hongjun Han (h27han), Ian Fox (MaideCa), jamendxman3, James Baxter (TheJamsh), Jarl Gullberg (Nihlus), Jason (Abatron), Jeremy (JeroMiya), Jorge Martins (SNikon), Kartik (i64), Kim Sang Kyun (skkim-ginno), Konstantin Zaitcev (kosz78), Leon Rosengarten (lion03), M3d10n, MacAndor, Maik Klein (MaikKlein), Manny (Manny-MADE), Marat Radchenko (slonopotamus), Mathew Wadstein (MWadstein), Lucyberad, Matthias Hölzl (hoelzl), MAVRIN (artemavrin), Maxim Pestun (maxpestun), mbehensky, meChiel, mfortin-bhvr, mipg, Narendra Umate (ardneran), Nikola Dimitroff (nikoladimitroff), Pablo Zurita (pzurita), phoboz-net, Pierre-Luc Loyer (plloyer), Piotr Bąk (Pierdek), projectgheist, Rama (EverNewJoy), Robert Chubb (SaxonRah), Robert Khalikov (nbjk667), Robert Segal (robertfsegal), RobertTroughton, Ryan Nielson (RyanNielson), Samuel Groß (saelo), Sasha Petrovich (xulture), [REDACTED_PERSON] (SRombauts), Stephen Whittle (stephenwhittle), Taces, Temaran, Tim Lindeberg (timlindeberg), Tobias Mollstam (mollstam), wdhwg001, Will Stahl (merlin91), WymzeeLabs, Yohann Martel (ymartel06), Zelex.
 
-# What’s New
+# What's New
 
 **Sequencer** is ready to use in Unreal Engine 4.12!  This is our fully non-linear, **real-time cinematics and animation tool**, and the successor to Matinee. It has many powerful features, including automatic keyframing, cinematic cameras, and live gameplay recording. Sequencer allows for easier collaboration on large animation projects, because you can freely organize your sequences into as many assets as you want.  Sequencer will offer to import your existing Matinee animations when you open them, but Matinee is still supported as well.
 
@@ -174,7 +174,7 @@ This is a new feature and has a number of limitations:
 
 For this release, we’ve added various new rendering features to allow for highly realistic surface reflections!
 
-![image alt text](image_11.png)(w:929 h:518 convert:false)
+![image alt text](image_11.jpg)(w:929 h:518 convert:false)
 
 We wanted to achieve very high quality, crisp reflections in order to achieve more realistic car paints among other things. This required improvements to several areas from how we process and encode geometry data to enabling high resolution reflection captures of the scene. 
 
@@ -241,7 +241,7 @@ A stationary directional light can now cast whole-scene dynamic CSM shadows from
 
 GPU particles with vector fields now work on high end mobile devices!
 
-![image alt text](image_18.png)(w:628 h:628 convert:false)
+![image alt text](image_18.jpg)(w:628 h:628 convert:false)
 
 * This feature requires a mobile device that supports 32-bit floating point render targets with MRT support.
 
