Adding images as per feedback.

diff --git a/Engine/Documentation/Source/Programming/UnrealArchitecture/Reference/Properties/Properties.INT.udn b/Engine/Documentation/Source/Programming/UnrealArchitecture/Reference/Properties/Properties.INT.udn
--- a/Engine/Documentation/Source/Programming/UnrealArchitecture/Reference/Properties/Properties.INT.udn
+++ b/Engine/Documentation/Source/Programming/UnrealArchitecture/Reference/Properties/Properties.INT.udn
@@ -34,8 +34,11 @@ Integer properties can now be exposed to the Editor as bitmasks. To mark an inte
 	int32 BasicBits;
 
 Adding this meta tag will cause the integer to be editable as a drop-down list of generically-named flags ("Flag 1", "Flag 2", "Flag 3", etc.) that can be 
-turned on or off individually. In order to customize the bitflags' names, create a UENUM with the "bitflags" meta tag and reference it with the "BitmaskEnum" 
-meta tag, like this:
+turned on or off individually.
+
+![](GenericBitmask.png)
+
+In order to customize the bitflags' names, we must first create a UENUM with the "bitflags" meta tag:
 
 	UENUM(Meta = (Bitflags))
 	enum class EColorBits
@@ -45,6 +48,8 @@ meta tag, like this:
 		ECB_Blue
 	};
 
+And then we can reference it with the "BitmaskEnum" meta tag, like this:
+
 	UPROPERTY(EditAnywhere, Meta = (Bitmask, BitmaskEnum = "EColorBits"))
 	int32 ColorFlags;
 
@@ -52,6 +57,8 @@ After this change, the bitflags listed in the drop-down box will take on the nam
 is value 0, meaning it will activate bit 0 (adding 1 to ColorFlags) when checked. ECB_Green corresponds to bit 1 (adding 2 to ColorFlags), and ECB_Blue 
 corresponds to bit 2 (adding 4 to ColorFlags).
 
+![](CustomBitmask.png)
+
 [REGION:note]While enumerated types can contain more than 32 entries, only the first 32 values will be visible in a bitmask association in the Property Editor 
 UI. Similarly, while explicit-value entries are accepted, entries with explicit values not between 0 and 31 will not be included in the drop-down.[/REGION]
 
