Slight phrasing improvement.

diff --git a/Engine/Documentation/Source/Programming/UnrealArchitecture/Reference/Properties/Properties.INT.udn b/Engine/Documentation/Source/Programming/UnrealArchitecture/Reference/Properties/Properties.INT.udn
--- a/Engine/Documentation/Source/Programming/UnrealArchitecture/Reference/Properties/Properties.INT.udn
+++ b/Engine/Documentation/Source/Programming/UnrealArchitecture/Reference/Properties/Properties.INT.udn
@@ -48,12 +48,12 @@ In order to customize the bitflags' names, we must first create a UENUM with the
 		ECB_Blue
 	};
 
-And then we can reference it with the "BitmaskEnum" meta tag, like this:
+After creating this UENUM, we can reference it with the "BitmaskEnum" meta tag, like this:
 
 	UPROPERTY(EditAnywhere, Meta = (Bitmask, BitmaskEnum = "EColorBits"))
 	int32 ColorFlags;
 
-After this change, the bitflags listed in the drop-down box will take on the names and values of the enumerated class entries. In the example above, ECB_Red 
+Following this change, the bitflags listed in the drop-down box will take on the names and values of the enumerated class entries. In the example above, ECB_Red 
 is value 0, meaning it will activate bit 0 (adding 1 to ColorFlags) when checked. ECB_Green corresponds to bit 1 (adding 2 to ColorFlags), and ECB_Blue 
 corresponds to bit 2 (adding 4 to ColorFlags).
 
