Updating version to 4.12.

diff --git a/Engine/Documentation/Source/Programming/UnrealArchitecture/Reference/Properties/Properties.INT.udn b/Engine/Documentation/Source/Programming/UnrealArchitecture/Reference/Properties/Properties.INT.udn
--- a/Engine/Documentation/Source/Programming/UnrealArchitecture/Reference/Properties/Properties.INT.udn
+++ b/Engine/Documentation/Source/Programming/UnrealArchitecture/Reference/Properties/Properties.INT.udn
@@ -2,7 +2,7 @@ Availability:Public
 Title: Properties
 Crumbs:%ROOT%, Programming, Programming/UnrealArchitecture/Reference
 Description:Reference for creating and implementing properties for gameplay classes.
-Version: 4.9
+Version: 4.12
 
 ## Property Declaration
 
