Added information from DanY about hack protection.

diff --git a/Engine/Documentation/Source/Gameplay/Networking/CharacterMovementComponent/CharacterMovementComponent.INT.udn b/Engine/Documentation/Source/Gameplay/Networking/CharacterMovementComponent/CharacterMovementComponent.INT.udn
--- a/Engine/Documentation/Source/Gameplay/Networking/CharacterMovementComponent/CharacterMovementComponent.INT.udn
+++ b/Engine/Documentation/Source/Gameplay/Networking/CharacterMovementComponent/CharacterMovementComponent.INT.udn
@@ -3,7 +3,7 @@ Title: Character Movement Component
 Crumbs:%ROOT%, Gameplay, Gameplay/Networking, Gameplay/Networking/Actors
 Description:Detailed explanation of Character Movement Component
 Related: Gameplay/Networking/Actors
-version: 4.9
+version: 4.12
 
 [TOC (start:2 end:3)]
 
@@ -43,6 +43,35 @@ A few useful tools exist to debug and analyze Character networking. Usually the
 
 Another helpful way to diagnose issues is to turn on some logging of the data sent by the CharacterMovement networked movement functions. The console command "log LogNetPlayerMovement Verbose" enables logging of each send and receive of character movement data including location, rotation, and acceleration. This may help explain why an error correction occurs, such as when location is being updated only on the client because it was not done in a way that the location change is replicated to the server.
 
+## Movement Speed Hack Protection (Networked Games)
+
+Networked games are often the target of users trying to gain unfair advantages through exploits. A common, game-generic exploit uses cheat software to speed up how quickly time 
+passes on the game client. Games made using unmodified CharacterMovementComponent functionality were susceptible to this exploit, resulting in cheaters being able to move at 
+increased rates. To prevent this type of exploit, we have added the ability to detect and resolve time discrepencies to `CharacterMovementComponent`.
+
+### Detection
+Time discrepency detection works by having the server compare the reported time difference between `ServerMove` RPCs (Remote Procedure Calls) coming from the client against 
+the known (and trusted) amount of time that has passed on the server. This difference is kept as a running total, and once the total exceeds a user-configurable threshold, resolution 
+measures are taken.
+
+### Resolution
+The method of resolution employed is to override the timestamps of future `ServerMove` RPCs coming from the client to match the server's time. A portion of the movement delta 
+between the RPCs can then be subtracted to "pay back" the time discrepency until the client is back in sync with the server.
+
+### Configuration
+By default, time discrepency detection and resolution are disabled. The following variables in `GameNetworkManager` can be configured, and their default values are found in 
+BaseGame.ini. These values can be overridden in your project-specific game configuration file.
+| --- | --- |
+| **Variable Name** | **Effect** |
+| bMovementTimeDiscrepancyDetection | Enables detection. Upon detection, warning logs will be triggered. Resolution will also be applied, if resolution is enabled. |
+| bMovementTimeDiscrepancyResolution | Enables resolution. Clients will be corrected and made to "pay back" time if sufficient discrepencies are detected. |
+| MovementTimeDiscrepancyMaxTimeMargin | Maximum time a client can be ahead of the server's expected game time before triggering detection/resolution. |
+| MovementTimeDiscrepancyMinTimeMargin | Maximum time a client can be behind the server's expected game time before triggering detection/resolution. |
+| MovementTimeDiscrepancyResolutionRate | Rate at which to "pay back" time during resolution. The default is 100%, meaning the client cannot move until the time debt is repaid. |
+| MovementTimeDiscrepancyDriftAllowance | Accepted drift in clocks between client and server as a percent per second allowed. This helps to prevent burst packet loss or performance hitches from triggering false positives. |
+| bMovementTimeDiscrepancyForceCorrectionsDuringResolution | Whether or not to force client updates during resolution. This is needed for projects that have lenient movement error tolerances or ClientAuthorativePosition enabled. |
+These settings may need to be tuned to fit your specific project. Basic testing can be done by using the "slomo" cheat on a client. The cvar "p.DebugTimeDiscrepancy" will activate time-discrepency logging on servers.
+
 ##Advanced Topic: Adding New Movement Abilities to Character Movement
 
 When adding new movement abilities to a character, you have many options on how to accomplish this. Let's step through a theoretical series of attempts to add a "Teleport" ability to a character that works in a networked game. The basic idea is that when a player presses the **T** key, they will teleport forward 10 meters, if the destination is free of obstruction.
