Fixed up according to peer review notes.

diff --git a/Engine/Documentation/Source/Gameplay/Networking/CharacterMovementComponent/CharacterMovementComponent.INT.udn b/Engine/Documentation/Source/Gameplay/Networking/CharacterMovementComponent/CharacterMovementComponent.INT.udn
--- a/Engine/Documentation/Source/Gameplay/Networking/CharacterMovementComponent/CharacterMovementComponent.INT.udn
+++ b/Engine/Documentation/Source/Gameplay/Networking/CharacterMovementComponent/CharacterMovementComponent.INT.udn
@@ -7,47 +7,70 @@ version: 4.12
 
 [TOC (start:2 end:3)]
 
-Characters using CharacterMovementComponent automatically have client-server networking built in. Here's how player movement prediction, replication and correction works in network games through CharacterMovementComponent:
-	
-Every tick, the TickComponent() function is called.  It figures out the acceleration and rotation change for the frame, and then calls PerformMovement() (for locally controlled Characters), or ReplicateMoveToServer() (if it's a network client).
-	
-ReplicateMoveToServer() saves the move (in the PendingMove list), calls PerformMovement(), and then replicates the move to the server by calling the replicated function ServerMove(), passing the movement parameters, the client's resultant position, and a timestamp.
-	
-ServerMove() is executed on the server.  It decodes the movement parameters and causes the appropriate movement to occur.  It then looks at the resulting position and if too much time has passed since the last response, or the position error is significant enough, the server calls ClientAdjustPosition(), a replicated function.
+Characters using `CharacterMovementComponent` automatically have client-server networking built in. Here's how player movement prediction, replication and correction works in networked 
+games through `CharacterMovementComponent`:
+
+On every tick, the `TickComponent` function is called.  It figures out the acceleration and rotation change for the frame, and then calls `PerformMovement` (for locally 
+controlled Characters), or `ReplicateMoveToServer` (for network clients).
 	
-ClientAdjustPosition() is executed on the client.  The client sets its position to the server's version of position, and sets the bUpdatePosition flag to true.
+`ReplicateMoveToServer` saves the move (into a list of pending moves), calls `PerformMovement` locally, and then replicates the move to the server by calling the replicated 
+function `ServerMove`. `ServerMove` receives the movement parameters, including the client's final position and a timestamp. It executes on the server, where it decodes the 
+movement parameters and causes the appropriate movement to occur.  It then looks at the resulting position and evaluates the time since the last client response and the difference 
+between the client's stated position and the position determined by the server. If the difference is big enough, the server calls `ClientAdjustPosition`, which replicates to 
+the client and passes the corrected position along.
 	
-When TickComponent() is called on the client again, if bUpdatePosition is true, the client will call ClientUpdatePosition() before calling PerformMovement().  ClientUpdatePosition() replays all the moves in the pending move list which occurred after the timestamp of the move the server was adjusting.
+When `TickComponent` is called on the client again, if a correction from the server has been received, the client will call `ClientUpdatePosition` before calling 
+`PerformMovement`.  This process will replay all of the moves in the pending move list which occurred after the timestamp of the move the server adjusted.
 
 ##Character Movement and Simulated Proxies
 
-The networking approach in CharacterMovementComponent described so far has dealt only with the details of a single client connected to an authoritative server. So how about AI that move on the server, or other players playing from a different computer? On our client machine, those are all considered simulated proxies, and go through a slightly different code path.
+The networking approach in `CharacterMovementComponent` described so far has dealt only with the details of a single client connected to an authoritative server. AI-controlled 
+Characters and players on remote computers are considered "simulated proxies", and go through a slightly different code path.
 
-Movement for non-human controlled characters is normally updated on the server (who is the authority in this regard) using the normal PerformMovement() code. Actor state such as location, rotation, velocity, as well as some other select Character-specific state (jumping, etc) are replicated to other machines via normal replication mechanisms, meaning they don't necessarily come down the network every frame. In order to provide a smoother view of those characters on a remote client, that client machine will run a simulation update every frame for simulated proxies until new (server authoritative) data arrives. The same is true for a local client's view of another remote human player; that remote player will send their updates to the server, which does a full movement update for that player, and periodically replicates that data to all other players.
+Movement for remotely-controlled Characters is normally updated on the server (who is the authority in this regard) using the normal `PerformMovement` code. Actor state data 
+such as location, rotation, and velocity, and other Character-specific information (e.g. jumping, crouching) is replicated to other machines via normal replication mechanisms, 
+meaning they generally don't receive updates every frame like locally-controlled Characters do. In order to provide a smoother view of those Characters, the client machine will run a 
+simulation update every frame for simulated proxies until new, authoritative data arrives from the server. While AI-controlled Characters are usually run directly on the server, 
+remote players will send their own local updates to the server, which then does a full movement update for that player, and periodically replicates that data to all other players.
 
-This update is intended to simulate the estimated result of movement based on replicated state, to "fill in the gaps" until the next update. So rather than placing proxies at the new location sent by the server and leaving them there until the next update arrives (which could be a few frames in the future), the client simulates an update in position each frame by applying velocity and the rules of movement. Once another update arrives, this effectively resets the local simulation and starts a new one.
+Simulating the expected results of movement based on replicated state "fills in the gaps" with reasonable predictions until the next server-authrotitative update arrives. Once the 
+next update arrives, it effectively resets the local simulation and starts a new one based on the latest information.
 
-The bulk of the movement update for simulated proxies is performed in UCharacterMovementComponent::SimulateMovement(), and in turn MoveSmooth(). MoveSmooth() is generally a pared-down version of the full movement mode updates for various movement modes (walking, flying, etc) that is a bit cheaper to run and a bit less complex in its intentions.
+The bulk of the movement update for simulated proxies is performed in `UCharacterMovementComponent::SimulateMovement`, and in turn `MoveSmooth`. `MoveSmooth` is a 
+pared-down version of the full movement mode updates for various movement modes (walking, flying, etc) that is a bit cheaper to run and a bit less complex in its intentions.
 
 ###Simulated Proxy Smoothing
 
-In cases where a character was simply moving forward, the simulated update is likely to very closely match the next replicated update, because moving in a straight line is quite simple to predict. Even in the case of running in to a static wall in the world, it's likely the deflection and subsequent updates will also simulate with high accuracy.
+In cases where a character is simply moving forward, the simulated update is likely to very closely match the next replicated update, because moving in a straight line is quite simple 
+to predict. Even in the case of running in to a static wall in the world, it's likely that the deflection and subsequent updates will also simulate with high accuracy.
 
-However, local simulation based on a previous replicated state snapshot is bound to diverge from the actual correct location in some cases. Consider a replicated state that said a character was moving with velocity (100,0,0) at time t=0. After a second of updates at t=1, the simulated proxy will have moved another 100 units in the X direction. However if the actual movement on the server had immediately ceased to move after sending the (100,0,0) velocity update, the next replication would place the server's version of the proxy location at a different spot than the locally simulated one.
+However, local simulation based on a previous replicated state snapshot is bound to diverge from the actual correct location in some cases, especially with human-controlled Characters. 
+Consider a replicated state that said a character was moving at a certain velocity. While waiting for the next update, the simulated proxy will continue to move at that velocity. 
+However, it is possible that the remote Character ceased to move immediately after sending the velocity update. The local simulation has no way of knowing this, and will make incorrect 
+predictions until the next server update comes in.
 
-To avoid a visual "pop" in position of the simulated proxies when a server update is received, we smooth the location of the Character's Mesh (complex visual representation of the body, rather than simple shape used for collision representation) using the function UCharacterMovementComponent::SmoothClientPosition(). By default this applies a simple smoothing function to reach the target destination in a specific amount of time (set by "SmoothNetUpdateTime" on the client's network data).
+To avoid a visual "pop" in position of the simulated proxies when a server update is received, we smooth the location of the Character's visual representation, using the function 
+`UCharacterMovementComponent::SmoothClientPosition`. By default, this applies a simple smoothing function to reach the target destination in a specific amount of time (set by 
+"SmoothNetUpdateTime" on the client's network data).
 
 ##Debugging CharacterMovementComponent Networking
 
-A few useful tools exist to debug and analyze Character networking. Usually the first thing you'll want to do is enter "p.NetShowCorrections 1" in the console of a client that is behaving erratically (this only works on a non-Shipping build). It may also be useful to enable this on the server. This makes it possible to see whenever a network correction is received on a client (or  sent on the server), by both logging to the output console and by drawing a collision shape at the "correct" location (in green), and the "incorrect" location (in red). On the client, the "correct" location is the one the server sent down as a correction, while the "incorrect" one is the local location that was judged to be outside the error tolerance on the server. On the server the idea is similar-- the "correct" server location is drawn in green, while the "incorrect" received client location is in red. "p.NetCorrectionLifetime" controls how long in seconds the debug visualizations persist in the world before going away (for example, "p.NetCorrectionLifetime 5" keeps them around for 5 seconds.
+A few useful tools exist to debug and analyze Character networking. Usually, the first thing you'll want to do is enter "p.NetShowCorrections 1" in the console of a client that is 
+behaving erratically (this only works on a non-Shipping build). It may also be useful to enable this on the server. This makes it possible to see whenever a network correction is 
+received on a client (or  sent on the server), by both logging to the output console and by drawing a collision shape at the "correct" location (in green), and the "incorrect" 
+location (in red). On the client, the "correct" location is the one the server sent down as a correction, while the "incorrect" one is the local location that was judged to be 
+outside the error tolerance on the server. On the server the idea is similar - the "correct" server location is drawn in green, while the "incorrect" received client location is in 
+red. The variable "p.NetCorrectionLifetime" controls how long, in seconds, the debug visualizations persist in the world before going away.
 
-Another helpful way to diagnose issues is to turn on some logging of the data sent by the CharacterMovement networked movement functions. The console command "log LogNetPlayerMovement Verbose" enables logging of each send and receive of character movement data including location, rotation, and acceleration. This may help explain why an error correction occurs, such as when location is being updated only on the client because it was not done in a way that the location change is replicated to the server.
+Another helpful way to diagnose issues is to turn on some logging of the data sent by the CharacterMovement networked movement functions. The console command 
+"log LogNetPlayerMovement Verbose" enables logging of each send and receive of character movement data including location, rotation, and acceleration. This may help explain why an 
+error correction occurs, such as when location is being updated only on the client because it was not done in a way that the location change is replicated to the server.
 
 ## Movement Speed Hack Protection (Networked Games)
 
 Networked games are often the target of users trying to gain unfair advantages through exploits. A common, game-generic exploit uses cheat software to speed up how quickly time 
-passes on the game client. Games made using unmodified CharacterMovementComponent functionality were susceptible to this exploit, resulting in cheaters being able to move at 
-increased rates. To prevent this type of exploit, we have added the ability to detect and resolve time discrepencies to `CharacterMovementComponent`.
+passes on the game client. Games using unmodified CharacterMovementComponent functionality were susceptible to this exploit, resulting in cheaters being able to move at increased 
+rates. To prevent this type of exploit, we have added the ability to detect and resolve time discrepencies to `CharacterMovementComponent`.
 
 ### Detection
 Time discrepency detection works by having the server compare the reported time difference between `ServerMove` RPCs (Remote Procedure Calls) coming from the client against 
@@ -55,63 +78,89 @@ the known (and trusted) amount of time that has passed on the server. This diffe
 measures are taken.
 
 ### Resolution
-The method of resolution employed is to override the timestamps of future `ServerMove` RPCs coming from the client to match the server's time. A portion of the movement delta 
-between the RPCs can then be subtracted to "pay back" the time discrepency until the client is back in sync with the server.
+The method of resolution employed is to override the timestamps of future `ServerMove` RPCs coming from the client to match the server's time. A portion of the difference in the 
+Character's movement data (location, rotation, acceleration, etc.) between the RPCs can then be subtracted to "pay back" the time discrepency until the client is back in sync with 
+the server.
 
 ### Configuration
 By default, time discrepency detection and resolution are disabled. The following variables in `GameNetworkManager` can be configured, and their default values are found in 
 BaseGame.ini. These values can be overridden in your project-specific game configuration file.
 | --- | --- |
 | **Variable Name** | **Effect** |
-| bMovementTimeDiscrepancyDetection | Enables detection. Upon detection, warning logs will be triggered. Resolution will also be applied, if resolution is enabled. |
+| bMovementTimeDiscrepancyDetection | Enables detection. Upon detection, warning logs will be triggered. If resolution is enabled, it will be applied as well. |
 | bMovementTimeDiscrepancyResolution | Enables resolution. Clients will be corrected and made to "pay back" time if sufficient discrepencies are detected. |
 | MovementTimeDiscrepancyMaxTimeMargin | Maximum time a client can be ahead of the server's expected game time before triggering detection/resolution. |
 | MovementTimeDiscrepancyMinTimeMargin | Maximum time a client can be behind the server's expected game time before triggering detection/resolution. |
 | MovementTimeDiscrepancyResolutionRate | Rate at which to "pay back" time during resolution. The default is 100%, meaning the client cannot move until the time debt is repaid. |
 | MovementTimeDiscrepancyDriftAllowance | Accepted drift in clocks between client and server as a percent per second allowed. This helps to prevent burst packet loss or performance hitches from triggering false positives. |
 | bMovementTimeDiscrepancyForceCorrectionsDuringResolution | Whether or not to force client updates during resolution. This is needed for projects that have lenient movement error tolerances or ClientAuthorativePosition enabled. |
-These settings may need to be tuned to fit your specific project. Basic testing can be done by using the "slomo" cheat on a client. The cvar "p.DebugTimeDiscrepancy" will activate time-discrepency logging on servers.
+These settings may need to be tuned to fit your specific project. Basic testing can be done by using the "slomo" cheat on a client. The variable "p.DebugTimeDiscrepancy" will activate 
+time-discrepency logging on servers.
 
 ##Advanced Topic: Adding New Movement Abilities to Character Movement
 
-When adding new movement abilities to a character, you have many options on how to accomplish this. Let's step through a theoretical series of attempts to add a "Teleport" ability to a character that works in a networked game. The basic idea is that when a player presses the **T** key, they will teleport forward 10 meters, if the destination is free of obstruction.
+There are several ways to add new movement abilities to Characters. As an example, we can grant a teleportation ability to a Character. Let's suppose that this ability enables the 
+Character to move 10 meters forward upon pressing the T key, as long as the destination is free of obstructions. Let's further suppose that this ability needs to work in a networked 
+game.
 
 ###Approach 1: Execute only on Client
-In this approach, we didn't think about networking at all and just teleport forward when pressing **T**.
+One approach might be to check for obstructions, then move the character directly forward 10 meters if the destination is clear. This does not consider networking, and simply executes 
+on the client.
 
-**Result**: Fails in networked games. The local client will appear to teleport forward but will then quickly warp back to their starting location.
+**Result**: Fails in networked games. The local client will teleport forward, but will then quickly warp back to their starting location.
 
-**Analysis**: This really is just the baseline that doesn't do anything special for networking, and is destined to fail. Why? Because the server has no indication that the ability was triggered, so is just working with the location, rotation, and acceleration of the client's character to trigger movement, and so from the server's perspective, it doesn't make sense that the character moved forward so far.
+**Analysis**: This is just a baseline that establishes that our code is being run. It could work for non-networked games. However, since it doesn't account for networking, it will 
+fail in a networked game. The server has no understanding of the ability, so it uses the location, rotation, and acceleration of the client's Character to figure out where the client 
+should end up. As far as the server is concerned, the 10 meter difference between its expected location for the Character and the location given by the client is error on the client's 
+part. The server corrects the error, and the client applies the server's correction locally, undoing the teleport.
 
 ###Approach 2: Server only RPC
-The simplest approach to making it work networked is to set up a reliable network RPC that triggers the ability. In this case, we just call a Server function that does the teleport.
+The simplest working approach for networked games is to have the server handle the ability and then inform the client of the result. To do this, we will set up a `UFUNCTION` tagged 
+as `Reliable` and `Server` that executes the teleport code when triggered from the client.
 
-**Result**: Works in networked games, but with major issues, such as a noticeable delay in execution on the client side and possible loss of some teleport functionality.
+**Result**: Works in networked games. However, there are a few problems with the way it works. There will be noticeable delay for the client using the ability, and the movement will 
+appear to be smoothed instead of being an instantaneous change of location.
 
-**Analysis**: This works, but it's far from ideal. Here is the sequence of events: The server gets the function call that teleports the player forward. The character movement update from the client is sent to the server with the unteleported location. The server sees this as an error, since the server's version teleported via the RPC. The server sends a correction to the client, and the client position changes as if it teleported. (Debugging tip: Using the "p.NetShowCorrections 1" command would help us here, because we'd see on the client that the movement is a result of a network correction.)
+First, there will be noticeable lag for the player using the ability, because the 
+command from the client needs to be sent from the server, executed there, and then sent back before the client knows that it has actually happened. Second, the smoothing code will 
+make the teleport feel like a smooth movement, where it should feel like an instant change of position. And third, the server's movement of the player will be perceived as a 
+correction, which is not the case and can mask real corrections when trying to debug the game's networked play.
 
-This is plagued with problems: in the presence of network lag, the client will have to wait a full round trip from client to server and back again to have the teleport actually occur on the local machine. This is a poor player experience. Additionally, any additional functionality occurring in the "Teleport" ability might not trigger at all on the client's side, since they never actually have that ability locally execute. So if for example you played sound and particle effects at the destination, those would be lost on the client.
+**Analysis**: This works, but it's far from ideal. When the player presses the T key teleport, the command is sent to the server, but nothing happens locally. This can leave players 
+feeling like the controls are unresponsive, and the player will not immediately know whether the teleport command worked or not. When the server receives the command, it moves the 
+Character, but does not tell the client. On the next time the client sends an update, the server will see the client's reported location as being 10 meters off, and will send a 
+correction. The correction will cause the client to move (smoothly) to the server's expected location, effectively causing the teleport to succeed, although it wouldn't look right. If 
+we had set "p.NetShowCorrections" to 1, we would be notified of this as a network correction. In addition, if we want to polish the ability with things like particle or sound effects, 
+those effects will not show up correctly since the actual ability only runs on the server. Although the client knows when the ability is attempted (and the command is sent to the 
+server), there is no report of it succeeding or failing, just a correction that shows up shortly afterward.
 
 ###Approach 3: Server RPC and Local Trigger
-In this approach, the client executes the Teleport and then also calls the server RPC to do the Teleport.
+In this approach, the client executes the Teleport locally, and then also calls the server's Teleport RPC.
 
 **Result**: Works in networked games, with rare but potentially serious issues.
 
-**Analysis**: This approach attempts to remedy the drawbacks of the previous attempt, namely the delay in local movement and the fact that movement occurred due to the correction, not the ability. We also get the full functionality of the Teleport, such as sound and particle effects. This actually works pretty well, but with some important caveats, and can still break in real networked environments.
-
-The main issue is that if a client is corrected back in time to before the time when the Teleport was triggered, they won't know to retrigger the Teleport as they catch up to their current time, and the Teleport will appear to have been lost on the client side.
+**Analysis**: This approach attempts to remedy the drawbacks of the previous attempt, namely the delay in local movement and the fact that movement occurred due to a server correction, 
+not the ability. We also get the full functionality of the Teleport, such as sound and particle effects, since the ability is now running on the client. This works much better than the 
+previous method, but with some important caveats, and can still break in real networked environments. The main issue is that if a client is corrected back in time to before the 
+Teleport was triggered, that client won't know to retrigger the Teleport as they repeat their pending move list, and the Teleport will appear to have been lost on the client side. This 
+will make the game feel unresponsive in certain cases, especially in bad network conditions.
 
 ###Approach 4: CharacterMovementComponent Ability Implementation
-In this approach, we add knowledge of the teleport ability to the CharacterMovementComponent code (in a child class), to make it robust to network corrections.
-
-**Result**: Works in networked games, with some implementation care. 
-
-**Analysis**: This requires a bit more background: as mentioned earlier, Characters using CharacterMovementComponent queue up the results of input in something called "a saved move list" or "pending move list" (in C++ code these moves derive from the class "FSavedMove_Character"). Each saved move records the state at the beginning of movement for a frame, such as location, rotation, acceleration (usually the result of player input), jump state, etc.
-
-As moves are sent from client to server, we hold on to these moves and remove old ones as moves are acknowledged by the server (remember that the server's view of our position is slightly behind our own immediate view). In the case of a correction, we know when in time that occurred, and can "replay" all moves that occurred after that point in time. This is nice because the client may not even notice that there was a correction, since they will try to move forward again in time from that corrected point in time. This also means that abilities triggered after the correction will get replayed as well. You probably want to check for this case in your ability, so you don't retrigger some effects if it's retriggered (see the "bReplayingMoves" parameter on the UCharacterMovementComponent::DoJump() function for an example).
-
-Now back to adding the "Teleport" ability! This is done in a very similar way as a jump is handled already in CharacterMovementComponent. We just need to indicate that this ability has been triggered, and process it correctly on the server. We still execute it locally of course, for quick client responsiveness. Sending and receiving the data between client and server is automatically done as part of the existing networking, you just need to pack and unpack the data.
-
-First we need to derive from UCharacterMovementComponent class, then override AllocateNewMove() to create our own version of FSavedMove_Character. At this point we need to implement a few methods in FSavedMove_Character, most importantly GetCompressedFlags(), which packs in a bit of data indicating we triggered the Teleport ability. Then we override UCharacterMovementComponent::UpdateFromCompressedFlags(), which handles unpacking those flags and triggering the ability on the server side.
-
-That's mostly it. There can be quite a bit more potential complexity to implementing custom flags in CharacterMovementComponent saved moves, but this should get you started on the correct path. If you want robust networking for your movement, it's worth the effort.
+In this approach, we add knowledge of the teleport ability to a child class of `CharacterMovementComponent`. This version will be robust over the network.
+
+**Result**: Works in networked games, with some implementation care.
+
+**Analysis**: This requires a bit more background: As mentioned earlier, Characters using `CharacterMovementComponent` queue up the results of input in a pending move list. Each 
+move in the list records the state at the beginning of movement for a frame, such as location, rotation, acceleration (usually the result of player input), jump state, and so on. As 
+moves are sent from client to server, we hold on to these moves and remove old ones as the server acknowledges that they have been completed. In the case of a server correction, the 
+client knows when in time the correction occurred, and can "replay" all moves that were made after that point in time. This is helpful because the player may not even notice that there 
+was a correction, since the Character's movement past the corrected point in time will be reapplied on top of the correction. This also means that abilities triggered after the 
+correction will be reapplied. Special effects can be managed so that they don't double-play in this case by checking the `bReplayingMoves` parameter, as seen in the 
+`UCharacterMovementComponent::DoJump` function.
+
+The Teleport ability itself will be added in a similar way to the Jump ability in `CharacterMovementComponent`. We need to indicate when this ability has been triggered, and process 
+it correctly on the server. We will still execute it locally so that the game feels responsive. Sending and receiving the data between client and server is automatically done as part 
+of the existing networking, we just need to pack and unpack the data. Once we have made our child class from `UCharacterMovementComponent`, we can override `AllocateNewMove` to 
+create our own version of `FSavedMove_Character`, which is what we'll store in the pending move list. Our version of `FSavedMove_Character` will have a few overridden methods. 
+`GetCompressedFlags` will be overridden with a new flag to indicate triggering of the Teleport ability. `UpdateFromCompressedFlags` will also be updated to unpack our new Teleport 
+flag and trigger the ability on the server side. This method will allow for robust movement abilities to be created in networked games.
