#UE4 docs HLOD docs initial submit

diff --git a/Engine/Documentation/Source/Engine/HLOD/HLODLanding.INT.udn b/Engine/Documentation/Source/Engine/HLOD/HLODLanding.INT.udn
new file mode 100644
--- /dev/null
+++ b/Engine/Documentation/Source/Engine/HLOD/HLODLanding.INT.udn
@@ -0,0 +1,24 @@
+title:Hierarchical Level of Detail
+description:Landing page for the Hierarchical Level of Detail system in Unreal Engine 4
+crumbs:
+availability:docs
+version:4.12
+
+**Hierarchical Level of Detail** meshes (HLODs) have been developed as a way to increase performance and keep the quality bar high for projects using Unreal Engine 4.11 or later. HLODs can replace multiple static mesh actors with a group of static mesh actors at long view distances. This helps reduce the number of actors that need to be rendered for the scene, increasing performance by lowering the number of draw calls per frame. 
+
+[OBJECT:TopicButtonList]
+	[PARAM:icon]
+		![](%ROOT%/tech_icon.png)(convert:false)
+	[/PARAM]
+	[PARAM:title]
+		HLOD Topics
+	[/PARAM]
+	[PARAM:description]
+	[/PARAM]
+	[PARAM:links]
+	* [](Engine/HLOD/Overview "%Engine/HLOD/Overview:description%")
+	* [](Engine/HLOD/Reference "%Engine/HLOD/Reference:description%")
+	* [](Engine/HLOD/HowTo "%Engine/HLOD/HowTo:description%")
+	[/PARAM]
+[/OBJECT]
+
diff --git a/Engine/Documentation/Source/Engine/HLOD/HowTo/BuildingHLODs/BuildingHLODs.INT.udn b/Engine/Documentation/Source/Engine/HLOD/HowTo/BuildingHLODs/BuildingHLODs.INT.udn
new file mode 100644
--- /dev/null
+++ b/Engine/Documentation/Source/Engine/HLOD/HowTo/BuildingHLODs/BuildingHLODs.INT.udn
@@ -0,0 +1,75 @@
+title:How-To Build Hierarchical Level of Detail Meshes
+description:This how to covers how to generate the HLOD meshes for your HLOD enabled Unreal Engine 4 project.
+crumbs:
+availability:docs
+version:4.12
+related:Engine/HLOD/Overview
+related:Engine/HLOD/Reference
+
+To build your Hierarchical Level of Detail (HLOD) Meshes, there are 2 major steps. Firstly, you must **Generate Clusters**, which groups Actors in your level together based on the setting you specify in the **Cluster Generation Settings**. This process is fairly . Then you have to Generate Proxy Meshes, which can take a very long time based on scene complexity and settings specified in the **Mesh Generation Settings**.
+
+For each LOD level you'll need to specify the settings for that LOD (such as 
+
+## Steps
+
+Below are the steps to generate your HLOD meshes broken into 2 major sections: Generate Clusters and Generate Proxy Meshes.
+
+### Generate Clusters
+
+Cluster generation uses the settings from the individual HLOD Levels to decide how it should group static mesh actors in the scene together. This generation process can take some time depending on the settings use, the number of actors being groups, whether materials are generated, and generally on your hardware’s specs. 
+
+1. Once you’ve setup your specific settings for the individual HLOD levels you need, click the Generate Clusters button.
+	
+	![image alt text](Engine/HLOD/Overview/image_3.png)(h:28 w:112 convert:false)
+
+	Once the process starts you’ll see a progress bar appear that indicates the LOD level being generated.
+
+	![image alt text](Engine/HLOD/Overview/image_4.png)(h:58 w:398 convert:false)
+
+1. Now that the process has completed you’ll see the HLOD Outliner populated with all the clustered actors. 
+
+	![image alt text](Engine/HLOD/Overview/image_5.png)(h:350 w:624 convert:false)
+
+1. You can expand the individual LODActors to see what static meshes make up this cluster by clicking the arrow button to the left of the name.
+
+	![image alt text](Engine/HLOD/Overview/image_6.png)(h:350 w:624 convert:false)
+
+1. You can visualize the clusters in the level by selecting a LODActor from the HLOD Outliner and locating it in the editor viewport.
+
+	![image alt text](Engine/HLOD/Overview/image_7.png)(h:481 w:624 convert:false)
+
+To visualize how the generated clusters work in the editor you can use the "Forced Viewing Level" slider to see the HLOD in action without the need to have it transition at a specific screen size. This is helpful for troubleshooting any issues appearing on screen that may be part of the generated cluster.
+
+![image alt text](Engine/HLOD/Overview/image_8.png)(h:17 w:624 convert:false)
+
+Should you run into any issues you don’t like with the generated cluster you can expand the specified cluster and select the offending static mesh actor. By right-clicking on the actor you have the option to Remove or Exclude it from the generation of the cluster. 
+
+![image alt text](Engine/HLOD/Overview/image_9.png)(h:101 w:392 convert:false)
+
+You also have some Right-Click options that work for a single LODActor or the selection.
+
+![image alt text](Engine/HLOD/Overview/image_10.png)(h:132 w:234 convert:false)
+
+Actors can also be excluded on a per-instance basis by selecting them in the level and in their Details panel set the option for **Enable Auto LODGeneration** to false.
+
+![image alt text](Engine/HLOD/Overview/image_11.png)(h:236 w:269 convert:false)
+
+Repeat this process while adjusting the values in the HLOD's level **Cluster Generation Settings** until you are happy with the clusters being generated, then continue on to the next section: [**Generate Proxy Meshes**](#generateproxymeshes).
+
+### Generate Proxy Meshes
+
+After you’re happy with the results of the generated cluster you can move on to the option to Build the clusters into a Proxy Mesh. This proxy mesh will be a newly created static mesh actor that combines materials (if enabled), has it’s own lightmaps, and it’s own editable static mesh that can be opened in the static mesh editor. 
+
+1. If you’re ready to build the proxy meshes you can now click the button  "Generate Proxy Mesh" button to start the process. 
+
+	![image alt text](Engine/HLOD/Overview/image_13.png)(h:28 w:148 convert:false)
+
+	Once the process starts you’ll see the progress bar appear that indicates the HLOD Level being worked on and the number of Proxy Meshes being generated. This progress bar does not show all the HLOD Levels and total number of proxy meshes being created, only that specific level.
+
+	![image alt text](Engine/HLOD/Overview/image_14.png)(h:94 w:458 convert:false)
+
+	Depending on your HLOD Level settings, the number of proxy meshes being created, and your system specs this process can take a while, even for high-end machines!
+
+	[REGION:note]
+	As a reference the Proxy Mesh Generation process for my machine (12-core i7 processor, GTX-980, and 64GB of RAM) took ~10-12 minutes with default settings for both HLOD Levels and ~100+ LODActors per HLOD level.
+	[/REGION]
diff --git a/Engine/Documentation/Source/Engine/HLOD/HowTo/Enabling/EnablingHLODs.INT.udn b/Engine/Documentation/Source/Engine/HLOD/HowTo/Enabling/EnablingHLODs.INT.udn
new file mode 100644
--- /dev/null
+++ b/Engine/Documentation/Source/Engine/HLOD/HowTo/Enabling/EnablingHLODs.INT.udn
@@ -0,0 +1,24 @@
+title:How-To Enable the HLOD System
+description:This how-to covers how to enable the HLOD system in your Unreal Engine 4 project.
+crumbs:
+availability:docs
+version:4.12
+related:Engine/HLOD/Overview
+related:Engine/HLOD/Reference
+
+The **Hierarchical Level of Detail System** (HLOD) must be enabled on each Level that wants to make use of it. Follow the steps below to enable the Hierarchical LOD System.
+
+## Steps
+
+1. Enable by going to **World Settings -> LODSystem -> Enable Hierarchical LODSystem**
+
+	[REGION:raw]
+	![image alt text](Engine/HLOD/Overview/image_0.png)
+	[/REGION]
+
+1. Once this is enabled you can go to the **Main Menu -> Window -> Hierarchical LOD Outliner** to set up how HLODs are generated for your scene.
+
+	[REGION:raw]
+	![image alt text](Engine/HLOD/Overview/image_1.png)
+	[/REGION]
+
diff --git a/Engine/Documentation/Source/Engine/HLOD/HowTo/HLODHowTos.INT.udn b/Engine/Documentation/Source/Engine/HLOD/HowTo/HLODHowTos.INT.udn
new file mode 100644
--- /dev/null
+++ b/Engine/Documentation/Source/Engine/HLOD/HowTo/HLODHowTos.INT.udn
@@ -0,0 +1,20 @@
+title:Hierarchical Level of Detail How-To's
+description:Words
+crumbs:
+availability:docs
+
+[OBJECT:TopicButtonList]
+	[PARAM:icon]
+		![](%ROOT%/tech_icon.png)(convert:false)
+	[/PARAM]
+	[PARAM:title]
+		HLOD Topics
+	[/PARAM]
+	[PARAM:description]
+	[/PARAM]
+	[PARAM:links]
+		* [](Engine/HLOD/HowTo/BuildingHLODs "%Engine/HLOD/HowTo/BuildingHLODs:description%")
+		* [](Engine/HLOD/HowTo/Enabling "%Engine/HLOD/HowTo/Enabling:description%")
+	[/PARAM]
+[/OBJECT]
+[]()
\ No newline at end of file
diff --git a/Engine/Documentation/Source/Engine/HLOD/Overview/HLODOverview.INT.udn b/Engine/Documentation/Source/Engine/HLOD/Overview/HLODOverview.INT.udn
new file mode 100644
--- /dev/null
+++ b/Engine/Documentation/Source/Engine/HLOD/Overview/HLODOverview.INT.udn
@@ -0,0 +1,74 @@
+title:Hierarchical Level of Detail Overview
+description:Overview of the Hierarchical Level of Detail system in Unreal Engine 4
+crumbs:%ROOT%
+availability:docs
+
+HLODs, in their simplest form, combine pre-existing **Static Mesh Actors** into a single Static Mesh Actor and material, with atlased textures, to reduce draw calls which in-turn helps with your performance of your game. There is a robust system in place to help you set specifically what needs to be generated when you create these HLOD actors for your scenes and covers more options that is normally available for LODs for a single actor. 
+
+## Proxy Meshes
+
+Once your proxy meshes have been generated you can find these located under the HLOD folder generated in your project. This will include the LODActors referenced in the HLOD Outliner. 
+
+![image alt text](image_15.png)(h:320 w:373 convert:false)
+
+These proxy meshes can be opened independently and settings adjusted. If you decided to merge the materials you’ll have access to these as well to adjust the texture source and any settings herein. 
+
+![image alt text](image_16.png)(h:357 w:404 convert:false)
+
+The proxy mesh above contains several different elements that originally all had their own texture that have now been combined into a single texture below.
+
+[REGION:note]
+
+This can be problematic for anything using Mask or Opacity. These alpha channels are not passed to these asset for this merged texture. If opacity or masked items are need you would want to consider unchecking the option for Merge Textures. When doing this the Proxy Mesh will assign the original material as a Material Element ID instead of combining. This has it’s own advantages and disadvantages.
+
+[/REGION]
+
+![image alt text](image_17.png)(h:420 w:422 convert:false)
+
+## HLOD Volume
+
+**The HLOD Volume** is a volume used to manually define/create a HLOD cluster. This volume can be added from the **Modes** panel **Place Mode** and dragged into your level.  
+
+![image alt text](image_18.png)(h:266 w:229 convert:false)
+
+Place this volume around actors you want to be grouped into the same cluster. It can be helpful to make this slightly larger than the actors and should not be tightly bounding them. Once you are have placed the volume around the actors you can use the option **Generate Clusters** from the HLOD Outliner to create a new cluster with these grouped actors.
+
+Example for Comparison: Generated Clusters before adding HLOD Volume
+
+![image alt text](image_19.png)(h:264 w:354 convert:false)
+
+HLOD Level Desired Bound Radius: 500
+
+![image alt text](image_20.png)(h:288 w:372 convert:false)
+
+Drag in the Volume and scale accordingly to cover your actors.
+
+ ![image alt text](image_21.png)(h:273 w:370 convert:false)
+
+Selecting the generated LODActor from the HLOD Outliner shows the cluster created and the Cluster bounds. 
+
+Right-click on the LODActor in the HLOD Outliner and click on **Select Contained Actors** to see the actors used for this particular LODActor in your scene.
+
+![image alt text](image_22.png)(h:93 w:212 convert:false)
+
+![image alt text](image_23.png)(h:268 w:357 convert:false)
+
+## Details Panel for LODActor
+
+By selecting a LODActor in the level you have access to some settings in the Details panel to override some aspects of your generated HLOD along with details to which static meshes have been included with this cluster. This is not unlike any other actors in your scene where you have the option to override or change some values to work within the confines of your design. 
+
+![image alt text](image_24.png)(h:237 w:270 convert:false)
+
+[VAR:spacer]![](%ROOT%/spacer.png)(w:150 convert:false)[/VAR]
+
+[VAR:spacer2]![](%ROOT%/spacer.png)(w:40 convert:false)[/VAR]
+
+<!--
+
+We should have a doc that covers the "Simplygon" features of HLOD
+
+## HLODMesh Culling Volume
+
+The HLODMesh Culling Volume is a volume that can be added to a HLOD cluster to remove triangles from source meshes before generating HLOD Mesh.
+
+-->
\ No newline at end of file
diff --git a/Engine/Documentation/Source/Engine/HLOD/Reference/HLODReference.INT.udn b/Engine/Documentation/Source/Engine/HLOD/Reference/HLODReference.INT.udn
new file mode 100644
--- /dev/null
+++ b/Engine/Documentation/Source/Engine/HLOD/Reference/HLODReference.INT.udn
@@ -0,0 +1,105 @@
+title:Hierarchical LOD Reference
+description:A reference of the interface elements and properties found in the Hierarchical LOD Outliner.
+crumbs:
+availability:docs
+
+This page details the buttons and properties that exist on the Hierarchical LOD Outliner.
+
+## Property and Interface Reference
+
+![HLOD Outliner Fully Expanded](image_2.png)
+
+[REGION:raw]
+### Interface Elements
+
+Below are the descriptions for the 3 buttons and slider that exist in the **Hierarchical LOD Outliner**.
+
+| --- | :---: | :--- |
+| **Generate Clusters** || This will generate spherical clusters of static mesh actors in the scene that can be grouped. The way this is handled is by using the settings for each LOD settings specified under the LODSystem section of the HLOD Outliner. |
+| **Delete Clusters** || This will delete all HLODs that have been previously setup in this level. This includes all generated proxy meshes as well as the clusters that have been setup. This will not wipe the previously set values or properties for the HLOD Levels. |
+| **Generate Proxy Meshes** || [INCLUDE:#gpm] |
+| **Forced Viewing Level** || This will allow you to force a specific LOD level to be visible in the editor. This is helpful for debugging any LOD issues you may have. |
+| | %spacer%  | |
+[/REGION]
+
+### Properties
+
+Below, broken out by major section, are the properties that can be found in a **Hierarchical LODSetup** in the **LODSystem** panel located in the lower portion of the **Hierarchical LOD Outliner** interface.
+
+[REGION:raw]
+
+### Cluster Generation Settings
+
+| Property | %spacer% | Description |
+| --- | :---: | :--- |
+| **Desired Bound Radius** || This is not guaranteed but used to calculate filling factor for auto clustering. |
+| **Desired Filling Percentage** || This is not guaranteed but used to calculate filling factor for auto clustering. |
+| **Min Number of Actors to Build** || Min number of actors to build LODActor. |
+
+### Mesh Generation Settings
+
+| Property | %spacer% | Description |
+| --- | :---: | :--- |
+| **Transition Screen Size** || The screen radius a mesh object should reach before swapping to the LOD actor. When the parent LOD displays it won’t display any of the children. |
+| [REGION:tablesection]**Merge Settings**[/REGION] |||
+| **Generate Light Map UV** || Whether to generate a lightmap UV for a merged mesh. |
+| **Target Light Map UVChannel** || The target UV channel to assign the lightmap UV. |
+| **Target Light Map Resolution** || The lightmap resolution set for the merged mesh. |
+| **Pivot Point at Zero** || Whether a merged mesh should have a pivot at the World Origin, or at the first merged actor from a generated cluster. |
+| **Merge Physics Data:** || Whether to merge the collision primitives when combining the actors. |
+| **Merge Materials** || Whether to merge source materials into one flat material. |
+| [REGION:tablesection]**[Material Settings See Below](#materialsettings)**[/REGION] |||
+| **Bake Vertex Data ** || This will bake any vertex data found into the LOD. |
+| **Calculate Correct LODModel** || Calculates the correct LOD Model. |
+| **Export Specific LOD** || Exports the specified LOD. |
+
+### Material Settings
+
+| Property | %spacer% | Description |
+| --- | :---: | :--- |
+| **Texture Size** || Size of the generated BaseColor texture map. |
+| **Texture Sizing Type** || [INCLUDE:#tst] |
+| **Normal Map** || Whether to generate a Normal map. |
+| **Metallic Map** || Whether to generate a Metallic map. |
+| **Metallic Constant** || Allows to you set a roughness constant value in place of a texture map. |
+| **Roughness Map** || Whether to generate a Roughness map. |
+| **Roughness Constant** || Allows to you set a roughness constant value in place of a texture map. |
+| **Specular Map** || Whether to generate a Specular map. |
+| **Specular Constant** || Allows to you set a roughness constant value in place of a texture map. |
+| **Emissive Map** || Whether to generate aa Emissive map. |
+| **Diffuse Texture Size** || Override for the Diffuse Texture Size. Enabled when **Texture Sizing Type** is set to **Use per property manually overriden texture sizes**. |
+| **Normal Texture Size** || Override for the Normal Texture Size. Enabled when **Texture Sizing Type** is set to **Use per property manually overriden texture sizes**. |
+| **Metallic Texture Size** || Override for the Metallic Texture Size. Enabled when **Texture Sizing Type** is set to **Use per property manually overriden texture sizes**. |
+| **Roughness Texture Size** || Override for the Roughness Texture Size. Enabled when **Texture Sizing Type** is set to **Use per property manually overriden texture sizes**. |
+| **Specular Texture Size** || Override for the Specular Texture Size. Enabled when **Texture Sizing Type** is set to **Use per property manually overriden texture sizes**. |
+| **Emissive Texture Size** || Override for the Emissive Texture Size. Enabled when **Texture Sizing Type** is set to **Use per property manually overriden texture sizes**. |
+[/REGION]
+
+[VAR:spacer]![](%ROOT%/spacer.png)(w:150 convert:false)[/VAR]
+
+[VAR:spacer2]![](%ROOT%/spacer.png)(w:40 convert:false)[/VAR]
+
+<!--
+[EXCERPT:GPM]
+When you generate the proxy meshes this will use the clustered LODs to merge these actors together into a single asset. Depending on what settings you changed for the HLOD levels your static meshes will generate based on this. This can affect your materials, lightmaps, and various other properties as well. 
+
+[REGION:note]
+Be aware that this process can take a **very long time** depending on your scene complexity and settings for your HLOD levels.
+[/REGION]
+[/EXCERPT]
+
+[EXCERPT:TST]
+
+Texture Sizing Type Description
+
+[REGION:raw]
+| Type | %spacer2% | Description |
+| --- | :---: | :--- |	
+| **Use TextureSize for all material properties** || Will use the value entered into **Texture Size** for sizing. |
+| **Use automatically biased texture sizes based on TextureSize** || Will use the value entered into **Texture Size**, but bias it based on usage. |
+| **Use per property manually overridden texture sizes** || Enables the per texture overrides, and uses those values instead of **Texture Size**. |
+| **Use Simplygon’s automatic texture sizing** || Uses **Simplygon's** sizing. |
+[/REGION]
+
+[/EXCERPT]
+-->
