Merging

//depot/UE4/Engine/Documentation/Source/Engine/LevelStreaming/...

to //UE4/Main/Engine/Documentation/Source/Engine/LevelStreaming/...

diff --git a/Engine/Documentation/Source/Engine/LevelStreaming/HowTo/HowToStreamLevels.INT.udn b/Engine/Documentation/Source/Engine/LevelStreaming/HowTo/HowToStreamLevels.INT.udn
--- a/Engine/Documentation/Source/Engine/LevelStreaming/HowTo/HowToStreamLevels.INT.udn
+++ b/Engine/Documentation/Source/Engine/LevelStreaming/HowTo/HowToStreamLevels.INT.udn
@@ -1,4 +1,4 @@
-Availability: Public
+Availability: Docs
 Title:Level Streaming How-Tos
 Crumbs:%ROOT%, Engine, Engine/UI/LevelEditor
 Description:How to dynamically stream in levels using Level Streaming Volumes, Blueprints, and C++ code
diff --git a/Engine/Documentation/Source/Engine/LevelStreaming/HowTo/StreamWithBlueprints/StreamBP.INT.udn b/Engine/Documentation/Source/Engine/LevelStreaming/HowTo/StreamWithBlueprints/StreamBP.INT.udn
--- a/Engine/Documentation/Source/Engine/LevelStreaming/HowTo/StreamWithBlueprints/StreamBP.INT.udn
+++ b/Engine/Documentation/Source/Engine/LevelStreaming/HowTo/StreamWithBlueprints/StreamBP.INT.udn
@@ -1,4 +1,4 @@
-Availability: Public
+Availability: Docs
 Title:Stream Sublevels with Blueprints
 Crumbs:%ROOT%, Engine, Engine/UI/LevelEditor
 Description:How to stream levels with a custom streaming Actor created with Blueprints
diff --git a/Engine/Documentation/Source/Engine/LevelStreaming/HowTo/StreamWithC++/StreamC++.INT.udn b/Engine/Documentation/Source/Engine/LevelStreaming/HowTo/StreamWithC++/StreamC++.INT.udn
--- a/Engine/Documentation/Source/Engine/LevelStreaming/HowTo/StreamWithC++/StreamC++.INT.udn
+++ b/Engine/Documentation/Source/Engine/LevelStreaming/HowTo/StreamWithC++/StreamC++.INT.udn
@@ -1,4 +1,4 @@
-Availability: Public
+Availability: Docs
 Title:Stream Sublevels with C++
 Crumbs:%ROOT%, Engine, Engine/UI/LevelEditor
 Description:How to stream levels with a custom streaming Actor created in C++
diff --git a/Engine/Documentation/Source/Engine/LevelStreaming/HowTo/StreamWithVolumes/StreamLSV.INT.udn b/Engine/Documentation/Source/Engine/LevelStreaming/HowTo/StreamWithVolumes/StreamLSV.INT.udn
--- a/Engine/Documentation/Source/Engine/LevelStreaming/HowTo/StreamWithVolumes/StreamLSV.INT.udn
+++ b/Engine/Documentation/Source/Engine/LevelStreaming/HowTo/StreamWithVolumes/StreamLSV.INT.udn
@@ -1,4 +1,4 @@
-Availability: Public
+Availability: Docs
 Title:Stream Sublevels with Level Streaming Volumes
 Crumbs:%ROOT%, Engine, Engine/UI/LevelEditor
 Description:How to use volumes to control the streaming of Levels based on the player's viewpoint.
diff --git a/Engine/Documentation/Source/Engine/LevelStreaming/Overview/Overview.INT.udn b/Engine/Documentation/Source/Engine/LevelStreaming/Overview/Overview.INT.udn
--- a/Engine/Documentation/Source/Engine/LevelStreaming/Overview/Overview.INT.udn
+++ b/Engine/Documentation/Source/Engine/LevelStreaming/Overview/Overview.INT.udn
@@ -1,4 +1,4 @@
-Availability: Public
+Availability: Docs
 Title:Level Streaming Overview
 Crumbs:%ROOT%, Engine, Engine/UI/LevelEditor
 Description:Streaming Levels can be loaded with Level Streaming Volumes or programmatically with Blueprints or C++.
