Merging

//depot/UE4/Engine/Documentation/Source/Platforms/SteamVR/...

to //UE4/Main/Engine/Documentation/Source/Platforms/SteamVR/...

diff --git a/Engine/Documentation/Source/Platforms/SteamVR/QuickStart/2/SteamVRQuickStart_2.INT.udn b/Engine/Documentation/Source/Platforms/SteamVR/QuickStart/2/SteamVRQuickStart_2.INT.udn
--- a/Engine/Documentation/Source/Platforms/SteamVR/QuickStart/2/SteamVRQuickStart_2.INT.udn
+++ b/Engine/Documentation/Source/Platforms/SteamVR/QuickStart/2/SteamVRQuickStart_2.INT.udn
@@ -22,10 +22,10 @@ Platform:SteamVR
 
 %Steps%
 
-In the following section we will go over how to setup a new Unreal Engine 4 (UE4) project to work with SteamVR.
+In the following section, we will go over how to setup a new Unreal Engine 4 (UE4) project to work with SteamVR.
 
 [region:note]
-If you have not done so already make sure to run the SteamVR **Room Setup**  to establish and calibrate the tracking area. 
+If you have not done so already, make sure to run the SteamVR **Room Setup** to establish and calibrate the VR tracking area. 
 Failing to do this could result in SteamVR and UE4 not working correctly together.
 [/region]
 
@@ -34,7 +34,7 @@ Failing to do this could result in SteamVR and UE4 not working correctly togethe
 	![](T_CreateNewProject.png)
 
 1. Once the project loads click on the small triangle next to the **Play** button and then from the menu that is displayed, select the **VR Preview** option. 
-	If your VR Headset is properly set up you should now see the basic level being display and have the ability to rotate your head around in any direction when you put the HMD on.
+	If your VR Headset is properly set up, you should now see the basic level being displayed and have the ability to rotate your head around in any direction when you put the HMD on.
 	
 	![](Check_VR_Working_00.png)
 
@@ -73,26 +73,38 @@ Failing to do this could result in SteamVR and UE4 not working correctly togethe
 
 	![](Set_Base_Eye_Height.png)
 
-1. In the **Components** tab, add a **SteamVRChaperone** and **Camera** component makeing sure that the camera's **Transform Location** is set to **0,0,0**.
+1. From the **Components** tab, add the following Components to the **VR_Pawn** Blueprint.
+	|Component Name|Amount|
+	|-------|-----|
+	|Capsule| 1|
+	|Scene|1|
+	|Camera|1|
+	|SteamVRChaperone|1|
 
 	![](Adding_A_Camera.png)
 
-1. Inside of the level, select the Static Mesh used for the floor and set the X, Y and Z location to **0,0,0**.
+1. Select the **Capsule** Component and under the **Details** panel in the **Shape** section, change the **Capsule Half Height** to **96** and the Capsule Radius to **16**.
 
-	![](Move_Floor_Down.png)
+	![](Capsule_Setup.png)
 
-1. Select the **Player Start** and position it so that the center of the Player Start is just above the Static Mesh used for the ground. 
-In this example we used a **Z** height of **1.0 CM** for the Player Start.
+1. Click on the **Camera** and then drag it on top of the **Scene** Component, making the Camera a child of the Scene Component.
 
-	![](Move_Player_Start.png)
+	![](Parenting_The_Camera.png)
 
-	[region:note]
-	Normally you do not want to have your **Player Start** intersecting with any type of geometry in your level. 
-	However because SteamVR reads your HMD's origin as 0,0,0, or sitting on the floor, moving the Player Start to a normal height will cause there to be an offset applied to the VR camera making it appear that your player is floating above the ground.  
-	[/region]
+1. Select the new **Scene** Component and go to **Details** > **Transform** > **Location** and set the **Z** value to **-110** then make sure to **Compile** and **Save** the Blueprint.
+
+	![](Change_Scene_Height.png)
+
+1. Inside of the level, select the Static Mesh used for the floor and set the X, Y and Z location to **0,0,0**.
 
-1. Launch the level using the **VR Preview** button, put your HMD on and then stand up.
-You should now feel as though the virtual floor that you are standing on matches up to the real world floor.
+	![](Move_Floor_Down.png)
+
+1. Launch the level using the **VR Preview** button, put your HMD on and then stand up and look around.
+When you look down at your feet, you should now feel as though the virtual floor that you are standing on matches up to the real world floor.
+
+[region:note]
+If it feels like you are too close or too far away from the virtual floor, open the  VR_Pawn Blueprint, select the Scene Component that the Camera is parented to and adjust the Z height value up or down till the floor feels like it is in the correct location for your height.
+[/region]
 
 
 %Steps%
