Set pages not to publish pending review

diff --git a/Engine/Documentation/Source/Platforms/OSVR/OSVR.INT.udn b/Engine/Documentation/Source/Platforms/OSVR/OSVR.INT.udn
--- a/Engine/Documentation/Source/Platforms/OSVR/OSVR.INT.udn
+++ b/Engine/Documentation/Source/Platforms/OSVR/OSVR.INT.udn
@@ -1,7 +1,7 @@
 Title:OSVR Development	
 Crumbs:%ROOT%, Platforms
 Description:
-Availability:Public
+Availability:Docs
 Navigation:platform
 Platform:OSVR
 Version: 4.12
diff --git a/Engine/Documentation/Source/Platforms/OSVR/QuickStart/1/OSVRQuickStart_1.INT.udn b/Engine/Documentation/Source/Platforms/OSVR/QuickStart/1/OSVRQuickStart_1.INT.udn
--- a/Engine/Documentation/Source/Platforms/OSVR/QuickStart/1/OSVRQuickStart_1.INT.udn
+++ b/Engine/Documentation/Source/Platforms/OSVR/QuickStart/1/OSVRQuickStart_1.INT.udn
@@ -1,4 +1,4 @@
-Availability:Public
+Availability:Docs
 Title:1. OSVR Prerequisites 
 Crumbs:%ROOT%, Platforms/OSVR
 Description:Setting up required OSVR software.
diff --git a/Engine/Documentation/Source/Platforms/OSVR/QuickStart/2/OSVRQuickStart_2.INT.udn b/Engine/Documentation/Source/Platforms/OSVR/QuickStart/2/OSVRQuickStart_2.INT.udn
--- a/Engine/Documentation/Source/Platforms/OSVR/QuickStart/2/OSVRQuickStart_2.INT.udn
+++ b/Engine/Documentation/Source/Platforms/OSVR/QuickStart/2/OSVRQuickStart_2.INT.udn
@@ -1,4 +1,4 @@
-Availability:Public
+Availability:Docs
 Title:2. OSVR Project Creation 
 Crumbs:%ROOT%, Platforms/OSVR
 Description:Creating a UE4 project that can be used with OSVR.
diff --git a/Engine/Documentation/Source/Platforms/OSVR/QuickStart/3/OSVRQuickStart_3.INT.udn b/Engine/Documentation/Source/Platforms/OSVR/QuickStart/3/OSVRQuickStart_3.INT.udn
--- a/Engine/Documentation/Source/Platforms/OSVR/QuickStart/3/OSVRQuickStart_3.INT.udn
+++ b/Engine/Documentation/Source/Platforms/OSVR/QuickStart/3/OSVRQuickStart_3.INT.udn
@@ -1,4 +1,4 @@
-Availability:Public
+Availability:Docs
 Title:3.Using OSVR in UE4
 Crumbs:%ROOT%, Platforms/GoogleVR
 Description:Enabling the OSVR Plugin.
diff --git a/Engine/Documentation/Source/Platforms/OSVR/QuickStart/OSVRQuickStart.INT.udn b/Engine/Documentation/Source/Platforms/OSVR/QuickStart/OSVRQuickStart.INT.udn
--- a/Engine/Documentation/Source/Platforms/OSVR/QuickStart/OSVRQuickStart.INT.udn
+++ b/Engine/Documentation/Source/Platforms/OSVR/QuickStart/OSVRQuickStart.INT.udn
@@ -1,7 +1,7 @@
 Title:OSVR Development	
 Crumbs:%ROOT%, Platforms
 Description:
-Availability:Public
+Availability:Docs
 Navigation:platform
 Platform:OSVR
 Version: 4.12
