Setting pages private pending review

diff --git a/Engine/Documentation/Source/Engine/Tools/DevTools/DeveloperTools.INT.udn b/Engine/Documentation/Source/Engine/Tools/DevTools/DeveloperTools.INT.udn
--- a/Engine/Documentation/Source/Engine/Tools/DevTools/DeveloperTools.INT.udn
+++ b/Engine/Documentation/Source/Engine/Tools/DevTools/DeveloperTools.INT.udn
@@ -1,4 +1,4 @@
-Availability:Public
+Availability:Docs
 Title:Developer Tools
 Crumbs: %ROOT%, Engine
 Description:Descriptions of the options found under the Developer Tools file menu. 
diff --git a/Engine/Documentation/Source/Engine/Tools/DevTools/PixelInspector/PixelInspector.INT.udn b/Engine/Documentation/Source/Engine/Tools/DevTools/PixelInspector/PixelInspector.INT.udn
--- a/Engine/Documentation/Source/Engine/Tools/DevTools/PixelInspector/PixelInspector.INT.udn
+++ b/Engine/Documentation/Source/Engine/Tools/DevTools/PixelInspector/PixelInspector.INT.udn
@@ -1,4 +1,4 @@
-Availability:Public
+Availability:Docs
 Title:Pixel Inspector
 Crumbs: %ROOT%, Engine
 Description:The Pixel Inspector shows you what is happening in the various parts of the GBuffer at a given pixel.
