Merging

//depot/UE4/Engine/Documentation/...

to //UE4/Main/Engine/Documentation/...

diff --git a/Engine/Documentation/Source/Engine/Blueprints/BP_HowTo/ActorReference/ActorReference.INT.udn b/Engine/Documentation/Source/Engine/Blueprints/BP_HowTo/ActorReference/ActorReference.INT.udn
--- a/Engine/Documentation/Source/Engine/Blueprints/BP_HowTo/ActorReference/ActorReference.INT.udn
+++ b/Engine/Documentation/Source/Engine/Blueprints/BP_HowTo/ActorReference/ActorReference.INT.udn
@@ -8,7 +8,7 @@ skilllevel:Intermediate
 type:howto
 parent:Engine/Blueprints/BP_HowTo
 order:1
-tags:Handling Objects and Actors
+tags:Data Types
 
 [INCLUDE:Gameplay\HowTo\ReferenceAssets#referenceoverview]
 
diff --git a/Engine/Documentation/Source/Engine/Blueprints/BP_HowTo/Casting/Casting.INT.udn b/Engine/Documentation/Source/Engine/Blueprints/BP_HowTo/Casting/Casting.INT.udn
--- a/Engine/Documentation/Source/Engine/Blueprints/BP_HowTo/Casting/Casting.INT.udn
+++ b/Engine/Documentation/Source/Engine/Blueprints/BP_HowTo/Casting/Casting.INT.udn
@@ -7,8 +7,9 @@ Related: Engine/Blueprints/UserGuide/BlueprintCommsUsage
 Related: Engine/Blueprints/UserGuide
 skilllevel: Advanced
 version: 4.9
-Parent:Engine/Blueprints/Scripting
-tags:Handling Objects and Actors
+Parent:Engine/Blueprints/HowTo
+type:howto
+tags:Data Types
 order:3
 
 This page covers how to use Blueprint Casting to communicate with different Blueprints and share information. In this example we allow the player to shoot an object, making it smaller each time it is hit. Each time the object is hit, it also notifies the Character Blueprint and increases the player's movement speed. 
diff --git a/Engine/Documentation/Source/Engine/Blueprints/Blueprints.INT.udn b/Engine/Documentation/Source/Engine/Blueprints/Blueprints.INT.udn
--- a/Engine/Documentation/Source/Engine/Blueprints/Blueprints.INT.udn
+++ b/Engine/Documentation/Source/Engine/Blueprints/Blueprints.INT.udn
@@ -76,8 +76,11 @@ social-image:bp_social_image.png
 
 [EXCERPT:Intro]
 The **Blueprints Visual Scripting** system in Unreal Engine is a complete gameplay scripting system based on the concept of using 
-a node-based interface to create gameplay elements from within Unreal Editor. This system is extremely flexible and powerful as it 
-provides the ability for designers to use virtually the full range of concepts and tools generally only available to programmers.
+a node-based interface to create gameplay elements from within Unreal Editor. As with many common scripting languages, it is used to define object-oriented (OO) classes or objects in the engine. 
+As you use UE4, you'll often find that objects defined using Blueprint are colloquially referred to as just "Blueprints." 
+
+This system is extremely flexible and powerful as it provides the ability for designers to use virtually the full range of concepts and tools generally only available to programmers. In addition,
+Blueprint-specific markup available in Unreal Engine's C++ implementation enables programmers to create baseline systems that can be extended by designers.
 [/EXCERPT:Intro]
 
 ## Getting Started
@@ -126,6 +129,7 @@ provides the ability for designers to use virtually the full range of concepts a
 %Engine/Blueprints/Anatomy:TopicCompact%
 %Engine/Blueprints/Workflow:TopicCompact%
 %Engine/Blueprints/BP_HowTo:TopicCompact%
+%Engine/Blueprints/UserGuide/CheatSheet:TopicCompact%
 %Engine/Blueprints/BestPractices:TopicCompact%
 [OBJECT:TopicCompact]
 	[PARAM:image]
diff --git a/Engine/Documentation/Source/Engine/Blueprints/GettingStarted/BlueprintIntro.INT.udn b/Engine/Documentation/Source/Engine/Blueprints/GettingStarted/BlueprintIntro.INT.udn
--- a/Engine/Documentation/Source/Engine/Blueprints/GettingStarted/BlueprintIntro.INT.udn
+++ b/Engine/Documentation/Source/Engine/Blueprints/GettingStarted/BlueprintIntro.INT.udn
@@ -49,13 +49,16 @@ order:1
 
 [TOC (start:2)]
 
-[EXCERPT:BlueprintIntro]
-**Blueprints** are special assets that provide an intuitive, node-based interface that can be used to create new types of Actors 
-and script level events; giving designers and gameplay programmers the tools to quickly create and iterate gameplay from within 
-Unreal Editor without ever needing to write a line of code.
-[/EXCERPT:BlueprintIntro]
+[EXCERPT:Intro]
+The **Blueprints Visual Scripting** system in Unreal Engine is a complete gameplay scripting system based on the concept of using 
+a node-based interface to create gameplay elements from within Unreal Editor. As with many common scripting languages, it is used to define object-oriented (OO) classes or objects in the engine. 
+As you use UE4, you'll often find that objects defined using Blueprint are colloquially referred to as just "Blueprints." 
 
-For more detailed overview information about Blueprints, including types of Blueprints, Blueprint anatomy, and a discussion of how Blueprints relate to UE3's Kismet and UnrealScript, see [](Engine/Blueprints\Overview).
+This system is extremely flexible and powerful as it provides the ability for designers to use virtually the full range of concepts and tools generally only available to programmers. In addition,
+Blueprint-specific markup available in Unreal Engine's C++ implementation enables programmers to create baseline systems that can be extended by designers.
+[/EXCERPT:Intro]
+
+For more detailed overview information about Blueprints, including Blueprint anatomy and a discussion of how Blueprints relate to UE3's Kismet and UnrealScript, see [](Engine/Blueprints\Overview).
 
 
 ## How Do Blueprints Work?
@@ -142,15 +145,6 @@ You have read about Level Blueprints and Blueprint Classes, listed below are a h
 	[/PARAM]
 [/OBJECT]
 
-##Blueprint Samples and Tutorials
-
-Below are additional resources you can use to learn more about the Blueprint system. 
-
-* [](Resources/Showcases/BlueprintExamples)
-* [Blueprint Lift Tutorial](http://wiki.unrealengine.com/Blueprint_Lift_Tutorial)
-* [Blueprint Power-Up Tutorial](http://wiki.unrealengine.com/Blueprint_Power_Up_Tutorial)
-
-
 ## Blueprint Editors and Graphs
 
 Whether you are building a Level Blueprint or a Blueprint Class, you will be using [Blueprint Elements](Engine/Blueprints\UserGuide) assembled in a [Blueprint Editor](Engine/Blueprints/Editor).  
@@ -163,4 +157,15 @@ Whether you are building a Level Blueprint or a Blueprint Class, you will be usi
 Different types of **Blueprint Editor** are available depending on the type of Blueprint you are working with. The core feature of most **Blueprint Editors** is the **Graph** mode, with its central **Graph** tab 
 for laying out the network of your Blueprint.  
 
+##Blueprint Samples and Tutorials
+
+Below are additional resources you can use to learn more about the Blueprint system. 
+
+* [](Resources/Showcases/BlueprintExamples)
+* [Blueprint Lift Tutorial](http://wiki.unrealengine.com/Blueprint_Lift_Tutorial)
+* [Blueprint Power-Up Tutorial](http://wiki.unrealengine.com/Blueprint_Power_Up_Tutorial)
+
+
+
+
 
diff --git a/Engine/Documentation/Source/Engine/Blueprints/Scripting/Scripting.INT.udn b/Engine/Documentation/Source/Engine/Blueprints/Scripting/Scripting.INT.udn
--- a/Engine/Documentation/Source/Engine/Blueprints/Scripting/Scripting.INT.udn
+++ b/Engine/Documentation/Source/Engine/Blueprints/Scripting/Scripting.INT.udn
@@ -1,8 +1,8 @@
 Availability:Public
 Title: Basic Scripting
 Crumbs: %ROOT%, Engine, Engine/Blueprints
-Description: The User Guide is the go-to source to learn the different parts of Blueprints and nodes that are available to use within Blueprint graphs.
-version: 4.9
+Description: Get a general overview of the variables and execution flow of the BLueprints visual scripting system.
+version: 4.12
 parent:Engine/Blueprints
 
 [VAR:Topic]
@@ -44,25 +44,31 @@ parent:Engine/Blueprints
 [/OBJECT]
 [/VAR]
 
+Blueprints provide a visual approach to a scripting language. As such, the system shares many of the nuances of a standard written scripting language, such as data typed variables, arrays, structs, etc. 
+Execution flow also works as it does in a typical scripting language, although Blueprints require explicit linear execution for each node. Each of the pages below go into more detail about 
+different variable types, working with those variables, and execution of nodes within the graph. 
 
-
-## Data Types
+## Variables
 
 Variables can be created in a variety of different types, including data types such as Boolean, integer, and float. They are color-coded for easy identification within your Blueprint.
+Blueprint variables can also be reference types for holding things like Objects, Actors, and Classes. 
 
-[DIR(output:"listbutton" parent:"Engine/Blueprints/Scripting" tags:"Data Types")]
+[DIR(output:"listbutton" parent:"Engine/Blueprints" tags:"Data Types")]
 
 ## Execution Flow
 
 In Blueprints, the first node to execute is an event, and then execution flows through the white execution wire from left to right. You can visualize the execution flow
 while your game is running in the editor, which can help with debugging. Data also flows through wires colored to match the variable types. Input pins are evaluated
-when the node executes, tracing the data wires back from right to left until the final  result is calculated and supplied to the node.
+when the node executes, tracing the data wires back from right to left until the final result is calculated and supplied to the node.
+
+Nodes with execution pins (impure nodes) store the values of their output pins when they execute, while nodes without execution pins (pure nodes) reevaluate their outputs
+every time a node connected to their outputs executes. 
 
 [DIR(output:"listbutton" parent:"Engine/Blueprints/Scripting" tags:"Execution Flow")]
 
 ## Arrays & Structures
 
-The basic data types available in Blueprints can be organized into arrays and structures for more complex systems. 
+Arrays and structures are collections of multiple data types. They are available for more complex systems, and help with keeping like information together and streamlining processing of larger amounts of information.
 
 [DIR(output:"listbutton" parent:"Engine/Blueprints/Scripting" tags:"Arrays and Structures")]
 
@@ -70,19 +76,3 @@ The basic data types available in Blueprints can be organized into arrays and st
 Structs
 -->
 
-## Handling Objects and Actors
-
-Blueprint variables can also be reference types for holding things like Objects, Actors, and Classes. 
-
-[DIR(output:"listbutton" tags:"Handling Objects and Actors")]
-
-
-<!--
-References (pointers)
-Set and Get an Actor Reference	Engine/Blueprints/BP_HowTo/ActorReference
-Validity
-Casting
-Blueprint Casting	Engine/Blueprints/BP_HowTo/Casting
-Casting	Engine/Blueprints/UserGuide/CastNodes
--->
-
diff --git a/Engine/Documentation/Source/Engine/Blueprints/UserGuide/CastNodes/CastNodes.INT.udn b/Engine/Documentation/Source/Engine/Blueprints/UserGuide/CastNodes/CastNodes.INT.udn
--- a/Engine/Documentation/Source/Engine/Blueprints/UserGuide/CastNodes/CastNodes.INT.udn
+++ b/Engine/Documentation/Source/Engine/Blueprints/UserGuide/CastNodes/CastNodes.INT.udn
@@ -6,7 +6,7 @@ Related: Engine/Blueprints/UserGuide/BlueprintCommsUsage
 Related: Engine/Blueprints/HowTo
 version: 4.9
 Parent:Engine/Blueprints/Scripting
-tags:Handling Objects and Actors
+tags:Data Types
 order:2
 
 [VAR:Topic]
diff --git a/Engine/Documentation/Source/Engine/Blueprints/UserGuide/CheatSheet/BlueprintEditorCheatSheet.INT.udn b/Engine/Documentation/Source/Engine/Blueprints/UserGuide/CheatSheet/BlueprintEditorCheatSheet.INT.udn
--- a/Engine/Documentation/Source/Engine/Blueprints/UserGuide/CheatSheet/BlueprintEditorCheatSheet.INT.udn
+++ b/Engine/Documentation/Source/Engine/Blueprints/UserGuide/CheatSheet/BlueprintEditorCheatSheet.INT.udn
@@ -3,7 +3,6 @@ Title:Blueprint Editor Cheat Sheet
 Crumbs: %ROOT%, Engine, Engine/Blueprints
 Description:Reference covering Blueprint shortcuts and useful actions.
 version: 4.9
-parent:Engine/Blueprints/Workflow
 order:1
 
 [VAR:Topic]
diff --git a/Engine/Documentation/Source/Engine/Blueprints/Workflow/Workflow.INT.udn b/Engine/Documentation/Source/Engine/Blueprints/Workflow/Workflow.INT.udn
--- a/Engine/Documentation/Source/Engine/Blueprints/Workflow/Workflow.INT.udn
+++ b/Engine/Documentation/Source/Engine/Blueprints/Workflow/Workflow.INT.udn
@@ -1,7 +1,7 @@
 Availability:Public
-Title: Working with Blueprints
+Title: Blueprint Workflow Tools
 Crumbs: %ROOT%, Engine, Engine/Blueprints
-Description: The User Guide is the go-to source to learn the different parts of Blueprints and nodes that are available to use within Blueprint graphs.
+Description: Tools like commenting and breakpoints enable you to follow the flow of your code and make notes for yourself or teammates working with the same system.
 version: 4.9
 parent:Engine/Blueprints
 order:3
diff --git a/Engine/Documentation/Source/Engine/LevelStreaming/HowTo/HowToStreamLevels.INT.udn b/Engine/Documentation/Source/Engine/LevelStreaming/HowTo/HowToStreamLevels.INT.udn
--- a/Engine/Documentation/Source/Engine/LevelStreaming/HowTo/HowToStreamLevels.INT.udn
+++ b/Engine/Documentation/Source/Engine/LevelStreaming/HowTo/HowToStreamLevels.INT.udn
@@ -2,6 +2,7 @@ Availability: Docs
 Title:Level Streaming How-Tos
 Crumbs:%ROOT%, Engine, Engine/UI/LevelEditor
 Description:How to dynamically stream in levels using Level Streaming Volumes, Blueprints, and C++ code
+Version: 4.12
 parent:Engine/LevelStreaming
 order:10
 tags:Level Streaming
@@ -9,7 +10,8 @@ tags:Level Streaming
 [TOC]
 
 
-These three how-tos all solve the following scenario, and have a shared set of steps to load in the two Levels you are going to work with. Follow these steps, then select the how-to you want to continue with.
+These three how-tos all solve the following scenario, and have a shared set of steps to load in the two Levels you are going to work with. Follow these steps, then select the how-to 
+you want to continue with.
 
 [EXCERPT:LevelStreamScenario]
 ## Level Streaming Scenario
@@ -29,7 +31,8 @@ will be loaded and visible.
 
 ![](StreamingLevelVisible.png)(w:500)
 
-As part of the setup, we have two levels, **SunTemple_Persistent** and **SunTemple_Streaming**. Our **Player Start** is in **SunTemple_Persistent**, and our player in the game is represented by a _Character_.
+As part of the setup, we have two levels, **SunTemple_Persistent** and **SunTemple_Streaming**. Our **Player Start** is in **SunTemple_Persistent**, and our player in the game is 
+represented by a _Character_.
 
 1. Open **SunTemple_Persistent** in the **Content Browser**.
 
diff --git a/Engine/Documentation/Source/Engine/LevelStreaming/HowTo/StreamWithVolumes/StreamLSV.INT.udn b/Engine/Documentation/Source/Engine/LevelStreaming/HowTo/StreamWithVolumes/StreamLSV.INT.udn
--- a/Engine/Documentation/Source/Engine/LevelStreaming/HowTo/StreamWithVolumes/StreamLSV.INT.udn
+++ b/Engine/Documentation/Source/Engine/LevelStreaming/HowTo/StreamWithVolumes/StreamLSV.INT.udn
@@ -2,6 +2,7 @@ Availability: Docs
 Title:Stream Sublevels with Level Streaming Volumes
 Crumbs:%ROOT%, Engine, Engine/UI/LevelEditor
 Description:How to use volumes to control the streaming of Levels based on the player's viewpoint.
+Version: 4.12
 parent:Engine/LevelStreaming/HowTo
 type:howto
 order:1
diff --git a/Engine/Documentation/Source/Engine/LevelStreaming/Overview/Overview.INT.udn b/Engine/Documentation/Source/Engine/LevelStreaming/Overview/Overview.INT.udn
--- a/Engine/Documentation/Source/Engine/LevelStreaming/Overview/Overview.INT.udn
+++ b/Engine/Documentation/Source/Engine/LevelStreaming/Overview/Overview.INT.udn
@@ -2,6 +2,7 @@ Availability: Docs
 Title:Level Streaming Overview
 Crumbs:%ROOT%, Engine, Engine/UI/LevelEditor
 Description:Streaming Levels can be loaded with Level Streaming Volumes or programmatically with Blueprints or C++.
+Version: 4.12
 parent:Engine/LevelStreaming
 type:reference
 order:1
@@ -52,14 +53,14 @@ tags:Level Streaming
 
 ## The Persistent Level
 
-The first step towards getting your levels to blend together seamlessly is creating the persistent level. This can be thought of as a master level that is primarily used to govern which 
-levels will be streamed in and out. 
+The first step towards getting your levels to blend together seamlessly is creating the persistent level. This can be thought of as a master level that is primarily used to govern 
+which levels will be streamed in and out. 
 
 
 ## Level Streaming Volumes
 
-Level streaming volumes make controlling level streaming very easy.  The idea is simple: load/unload requests for a streaming level are issued based on whether the viewpoint is inside any of the 
-LevelStreamingVolume volumes associated with a level.
+Level streaming volumes make controlling level streaming very easy. The idea is simple: load/unload requests for a streaming level are issued based on whether the viewpoint is 
+inside any of the LevelStreamingVolume volumes associated with a level.
 
 Specifically, level streaming volumes can be used in two ways:
 
@@ -69,57 +70,59 @@ Specifically, level streaming volumes can be used in two ways:
 * In the **editor**, level streaming volumes can be used to preview level streaming by automatically hiding/unhiding levels based on the location of the perspective viewport's camera.
 
 
-Volume-based level streaming is simple to use and does not require any scripting, making it an ideal way to control level streaming.  Furthermore, volume-based level streaming is much easier to 
-maintain than scripted streaming: when demands on the loading system change, level loading/unloading behavior can be modified by simply resizing the streaming volumes.
+Volume-based level streaming is simple to use and does not require any scripting, making it an ideal way to control level streaming. Furthermore, volume-based level streaming is 
+much easier to maintain than scripted streaming: when demands on the loading system change, level loading/unloading behavior can be modified by simply resizing the streaming volumes.
 
 ### Associating Streaming Volumes With Levels
 
-Volume-based level streaming works as follows: each streaming level can have associated with it a set of volumes of type LevelStreamingVolume.  Each frame, the engine iterates over each level and 
-checks to see if the player viewpoint is inside any of the Level Streaming Volumes associated with that level.  If the viewpoint is inside at least one LevelStreamingVolume, a request is issued to 
-begin loading that level.  If the viewpoint is outside all Level Streaming Volumes, the level is marked for unloading.
+Volume-based level streaming works as follows: each streaming level can have associated with it a set of volumes of type LevelStreamingVolume. Each frame, the engine iterates over 
+each level and checks to see if the player viewpoint is inside any of the Level Streaming Volumes associated with that level. If the viewpoint is inside at least one 
+LevelStreamingVolume, a request is issued to begin loading that level. If the viewpoint is outside all Level Streaming Volumes, the level is marked for unloading.
 
 * [](Engine/LevelStreaming/HowTo/StreamWithVolumes)
 
 ### Important Details
 
-* All Level Streaming Volumes must exist in the persistent level.  Level Streaming Volumes that live in other levels cannot be used for level streaming, and will generate warnings when the map is checked 
-for errors.
+* All Level Streaming Volumes must exist in the persistent level. Level Streaming Volumes that live in other levels cannot be used for level streaming, and will generate warnings 
+when the map is checked for errors.
 
 * If a level has any streaming volumes associated with it, other methods of streaming the level will not behave correctly.
 
-* A single LevelStreamingVolume can affect multiple levels.  Similarly, a single level can be affected by multiple Level Streaming Volumes.
+* A single LevelStreamingVolume can affect multiple levels. Similarly, a single level can be affected by multiple Level Streaming Volumes.
 
-* Volume-based streaming works for split screen.  The viewpoints of all local players are considered before any loading/unloading requests are issued.
+* Volume-based streaming works for split screen. The viewpoints of all local players are considered before any loading/unloading requests are issued.
 
 
 
 
 ### Testing Your Streaming Volume Setup
 
-It is critical that volume-based level streaming be tested in game on the target platform.  Streaming in **Play in Editor** mode (PIE) will show where the loads/unloads will happen, but streaming in 
-PIE is not representative of real, in-game loading-unloading.  This is because in PIE, the levels are already in memory, and so "loading" a level is simply a matter of unhiding it and is thus instantaneous.
+It is critical that volume-based level streaming be tested in game on the target platform. Streaming in **Play in Editor** mode (PIE) will show where the loads/unloads will happen, 
+but streaming in PIE is not representative of real, in-game loading-unloading. This is because in PIE, the levels are already in memory, and so "loading" a level is simply a matter 
+of unhiding it and is thus instantaneous.
 
-Running the level in the standalone game on the target platform is critical in making sure your streaming setups work.  Note that on some platforms, it can take several seconds to stream in a level.  
-Size your Level Streaming Volumes appropriately, so that the level is loaded by the time the player can reach it.  Level loading behaviour can be modified by resizing the Level Streaming Volumes associated 
-with a level.  Growing a volume causes associated levels to load sooner and unload later, while shrinking a volume causes later loads and earlier unloads.
+Running the level in the standalone game on the target platform is critical in making sure your streaming setups work. Note that on some platforms, it can take several seconds to 
+stream in a level. Size your Level Streaming Volumes appropriately, so that the level is loaded by the time the player can reach it. Level loading behaviour can be modified by 
+resizing the Level Streaming Volumes associated with a level. Growing a volume causes associated levels to load sooner and unload later, while shrinking a volume causes later loads 
+and earlier unloads.
 
 
 ### Level Streaming Volumes for Previs
 
-Volume-based level streaming can be easily previewed by enabling the "Level Streaming Volume Previs" button on the perspective viewport toolbar.  When this mode is enabled, levels will be hide/unhide 
-based on the location of the perspective viewport camera.  Note that this hiding/unhiding will affect the set of visible levels in the Level Browser.
+Volume-based level streaming can be easily previewed by enabling the "Level Streaming Volume Previs" button on the perspective viewport toolbar. When this mode is enabled, levels 
+will be hidden/shown based on the location of the perspective viewport camera. Note that this hiding/showing will affect the set of visible levels in the Level Browser.
 
-A LevelStreamingVolume can be marked for editor previs only by setting the bEditorPreVisOnly flag on the LevelStreamingVolume.  In this manner, volume-based level streaming can still be used for editor 
-previs while at the same time using another streaming method for in-game streaming.
+A LevelStreamingVolume can be marked for editor previs only by setting the bEditorPreVisOnly flag on the LevelStreamingVolume. In this manner, volume-based level streaming can 
+still be used for editor previs while at the same time using another streaming method for in-game streaming.
 
 
 
 ### Cost of Level Streaming Volumes
 
-Each frame, UWorld::ProcessLevel Streaming Volumes iterates over each streaming level and for each level, that level will begin loading if any local player is within any of the volumes associated with 
-that level.  Likewise, the level will be begin unloading if all local players are outside all volumes.
+Each frame, `UWorld::ProcessLevelStreamingVolumes` iterates over each streaming level and for each level, that level will begin loading if any local player is within any of the 
+volumes associated with that level. Likewise, the level will be begin unloading if all local players are outside all volumes.
 
-UWorld::ProcessLevel Streaming Volumes exploits coherency in the following manner: for each level, the volume that most recently contained a player is cached.  This cached volume is checked first, so 
+`UWorld::ProcessLevelStreamingVolumes` exploits coherency in the following manner: for each level, the volume that most recently contained a player is cached.  This cached volume is checked first, so 
 that levels a player is in or returns to are quickly accepted.
 
 Volumes of any shape are fine, although obviously the fewer the better.  An upper bound on the cost of level streaming volumes can be approximated by the sum of the number of level streaming volumes 
@@ -132,26 +135,28 @@ per frame, while the "Volume Streaming Tick" stat tracks the amount of time spen
 
 ### Adding Hysteresis to Unloading Requests
 
-A player moving back and forth across a LevelStreamingVolume boundary causes spurious load/unload requests to be issued.  To address this, hysteresis has been added to unloading requests.   No hysteresis 
-exists for loading requests, because if a level needs to be loaded, we always want it loaded ASAP.
+A player moving back and forth across a LevelStreamingVolume boundary causes unneeded load/unload requests to be issued. To address this, hysteresis has been added to unloading 
+requests. No hysteresis exists for loading requests, because if a level needs to be loaded, we always want it loaded as soon as possible.
 
-The amount of unloading hysteresis can be adjusted by modifying the MinTimeBetweenVolumeUnloadRequests property of the LevelStreaming object (LevelStreaming properties can be accessed from the Level 
-Browser context menu).  The default unloading hysteresis is 2.0 seconds.
+The amount of unloading hysteresis can be adjusted by modifying the **Min Time Between Volume Unload Requests** property of the streaming Level in the [Levels window](Engine/Levels/LevelsWindow). 
+The default unloading hysteresis is 2.0 seconds.
 
 
 
 ### Disabling Level Streaming Volumes
 
-There is a property on Level Streaming Volumes called bDisabled.  When set to true, the volume will be ignored by the streaming volume code, both in-game and in the editor.  bDisabled can be used to disable 
-a LevelStreamingVolume without disassociating it from the level.  
+There is a property on Level Streaming Volumes called **bDisabled**. When set to true, the volume will be ignored by the streaming volume code, both in-game and in the editor. **bDisabled** 
+can be used to disable a Level Streaming Volume without dissociating it from the Level.  
 
-As an example of where the bDisabled flag is useful, imagine a door leading to a Level whose streaming is controlled with streaming volumes. The streaming volume extends out past the door, so that the 
-Level will be streamed in by the time the player can reach the door and open it.  Initially, however, the door is locked, and will become unlocked when the player achieves an objective in another part of the Level.  
-So, even though the streaming volume extends past the door, we don't want the Level on the other side of the door to stream in unless the door is actually unlocked ("openable").
+As an example of where the **bDisabled** flag is useful, imagine a door leading to a Level whose streaming is controlled with streaming volumes. The streaming volume extends out past 
+the door, so that the Level will be streamed in by the time the player can reach the door and open it. Initially, however, the door is locked, and will become unlocked when the 
+player achieves an objective in another part of the Level. So, even though the streaming volume extends past the door, we don't want the Level on the other side of the door to 
+stream in unless the door is actually unlocked ("openable").
 
-Blueprints and C++ could both be used to toggle the **Disabled** state of **Level Streaming Volumes**.
+Blueprints and C++ could both be used to toggle the disabled state of **Level Streaming Volumes**.
 
 ## Programmatic Level Streaming
 
-If you want to set up more complex behavior for loading and unloading Levels, similar to the door unlocking example above, you can also set up Blueprint or C++ logic for Level loading and unloading. The key functions are
-**Load Stream Level** and **Unload Stream Level**, and you designate the Level to load by its name. See [](Engine/LevelStreaming/HowTo/StreamWithBlueprints) and [](Engine/LevelStreaming/HowTo/StreamWithC++) for examples. 
\ No newline at end of file
+If you want to set up more complex behavior for loading and unloading Levels, similar to the door unlocking example above, you can also set up Blueprint or C++ logic for Level 
+loading and unloading. The key functions are **Load Stream Level** and **Unload Stream Level**, and you designate the Level to load by its name. See 
+[](Engine/LevelStreaming/HowTo/StreamWithBlueprints) and [](Engine/LevelStreaming/HowTo/StreamWithC++) for examples. 
\ No newline at end of file
diff --git a/Engine/Documentation/Source/Engine/Levels/CollaborateWithSublevels/CollaborateWithSublevels.INT.udn b/Engine/Documentation/Source/Engine/Levels/CollaborateWithSublevels/CollaborateWithSublevels.INT.udn
--- a/Engine/Documentation/Source/Engine/Levels/CollaborateWithSublevels/CollaborateWithSublevels.INT.udn
+++ b/Engine/Documentation/Source/Engine/Levels/CollaborateWithSublevels/CollaborateWithSublevels.INT.udn
@@ -3,15 +3,15 @@ Title:Collaborate with Sublevels
 Crumbs: %ROOT%, Engine
 Description:You can use sublevels to collaborate with other designers working in the same space. 
 Related: Engine/LevelStreaming
-Version: 4.9
+Version: 4.12
 parent:Engine/Levels
 order:5
 
 Levels are binary assets, which can lead to some challenges when collaborating with teammates working on the same area of the project.
 
-However, by setting up **Always Streaming** sublevels, multiple people can work in the same game "space" at the same time. For instance, you might have a sublevel that contains all the lights
-for your scene, as well as a sublevel that contains just the Static Meshes laying out your environment. Each sublevel can still only have one person working on it at a time, but with careful
-segmenting of your scene's Actors, you can minimize overlapping work areas. 
+However, by setting up **Always Streaming** sublevels, multiple people can work in the same game "space" at the same time. For instance, you might have a sublevel that contains all 
+the lights for your scene, as well as a sublevel that contains just the Static Meshes laying out your environment. Each sublevel can still only have one person working on it at a 
+time, but with careful segmenting of your scene's Actors, you can minimize overlapping work areas. 
 
 [REGION:tip]
 You can even have sublevels' Level Blueprints communicate with each other, so that allows a switch in one sublevel, for instance, to toggle a light in another.
diff --git a/Engine/Documentation/Source/Engine/Levels/HowTo/ChangeDefaultLevel/ChangeDefaultLevel.INT.udn b/Engine/Documentation/Source/Engine/Levels/HowTo/ChangeDefaultLevel/ChangeDefaultLevel.INT.udn
--- a/Engine/Documentation/Source/Engine/Levels/HowTo/ChangeDefaultLevel/ChangeDefaultLevel.INT.udn
+++ b/Engine/Documentation/Source/Engine/Levels/HowTo/ChangeDefaultLevel/ChangeDefaultLevel.INT.udn
@@ -3,7 +3,7 @@ Title:Change Default Levels
 Crumbs: %ROOT%, Engine
 Description:How to set the default editor and game Levels for your project
 SkillLevel: Beginner
-Version: 4.9
+Version: 4.12
 parent:Engine/Levels
 order:4
 type:howto
diff --git a/Engine/Documentation/Source/Engine/Levels/Levels.INT.udn b/Engine/Documentation/Source/Engine/Levels/Levels.INT.udn
--- a/Engine/Documentation/Source/Engine/Levels/Levels.INT.udn
+++ b/Engine/Documentation/Source/Engine/Levels/Levels.INT.udn
@@ -52,7 +52,7 @@ order:4
 
 ![](Level_Intro.png)
 
-When playing a video game, every object that you see or interact with in the world resides in what is known as a **Level**. In Unreal Engine 4 terms, a Level is made up of a 
+When playing a video game, every object that you see or interact with resides in what is known as a **Level**. In Unreal Engine 4 terms, a Level is made up of a 
 collection of Static Meshes, Volumes, Lights, Blueprints and more all working together to bring the desired experience to the player. Levels in UE4 can range in size from massive 
 terrain-based worlds to very small levels that contain a few Actors. 
 
diff --git a/Engine/Documentation/Source/Engine/Levels/WorldSettings/WorldSettings.INT.udn b/Engine/Documentation/Source/Engine/Levels/WorldSettings/WorldSettings.INT.udn
--- a/Engine/Documentation/Source/Engine/Levels/WorldSettings/WorldSettings.INT.udn
+++ b/Engine/Documentation/Source/Engine/Levels/WorldSettings/WorldSettings.INT.udn
@@ -1,7 +1,7 @@
 Availability: Docs
 Title:World Settings
 Description:The **World Settings** tab enables you to set and override Level-specific settings like base Game Mode and gravity. 
-Version: 4.9
+Version: 4.12
 parent:Engine/Levels
 order:3
 
