#UE4 Docs: added link to Sequencer sample

diff --git a/Engine/Documentation/Source/Engine/Sequencer/Sequencer.INT.udn b/Engine/Documentation/Source/Engine/Sequencer/Sequencer.INT.udn
--- a/Engine/Documentation/Source/Engine/Sequencer/Sequencer.INT.udn
+++ b/Engine/Documentation/Source/Engine/Sequencer/Sequencer.INT.udn
@@ -6,6 +6,7 @@ Navigation:topic
 parent:Engine
 order:3
 social-image:sequencer_social.png
+related: Engine/Matinee
 
 [VAR:Topic]
 [OBJECT:Topic]
@@ -55,37 +56,18 @@ Similar to Matinee, Sequencer uses various specialized **Tracks** which are used
 Add a couple of Skeletal Mesh Tracks to Sequencer, then add Animation Tracks to animate those Skeletal Meshes, throw in a Camera Track and Camera Cuts track to provide a perspective and cycle between cameras and you have the foundation for setting up a cinematic sequence. 
 
 On this page you will find links to documenation regarding different aspects of the Sequencer Editor. 
-If you are new to Sequencer, it is recommended that you check out the **Esentials** section which includes a high-level overview of the Sequencer Editor, an Editor Reference which breaks down components of the Sequencer Editor, a Quick Start guide to quickly get you up and running, multiple How-to's which are one-off pages of how to perform certain tasks with Sequencer as well as a section on Workflow tips. 
+If you are new to Sequencer, it is recommended that you check out the **Essentials** section which includes a high-level overview of the Sequencer Editor, an Editor Reference which breaks down components of the Sequencer Editor, a Quick Start guide to quickly get you up and running, multiple How-to's which are one-off pages of how to perform certain tasks with Sequencer as well as a section on Workflow tips. 
 
+The Sequencer Subway example project, available from the Learn tab in the Epic Games Launcher, shows how to create film-quality cinematic sequences using Unreal Engine 4's Sequencer Editor. 
 
-[OBJECT:TopicButtonList]
-	[PARAM:icon]
-		![](%ROOT%/start_icon.png)(convert:false)
-	[/PARAM]
-	[PARAM:title]
-		Essentials
-	[/PARAM]
-	[PARAM:links]
-		* [](Engine\Sequencer\Overview "%Engine\Sequencer\Overview:description%")
-		* [](Engine\Sequencer\ReferenceEditor "%Engine\Sequencer\ReferenceEditor:description%")
-		* [](Engine\Sequencer\QuickStart "%Engine\Sequencer\QuickStart:description%")
-		* [](Engine\Sequencer\HowTo "%Engine\Sequencer\HowTo:description%")
-		* [](Engine\Sequencer\Workflow\ "%Engine\Sequencer\Workflow:description%")
-	[/PARAM]
-[/OBJECT]
+## Essentials
 
-[OBJECT:TopicButtonList]
-	[PARAM:icon]
-		![](%ROOT%/related_icon.png)(convert:false)
-	[/PARAM]
-	[PARAM:title]
-		Related Links
-	[/PARAM]
-	[PARAM:links]
-		* [Matinee Editor](Engine/Matinee "%Engine/Matinee:description%")
-	[/PARAM]
-[/OBJECT]
+[DIR(output:"listbutton" parent:"Engine/Sequencer" org:"hierarchy" end:"1")]
+
+
+## Samples
 
+[DIR(output:"topiccompact" parent:"Resources/Showcases" tags:"Sequencer")]
 
 <!-- Killing this for now while the current does get updates instead.
 * [](Engine\Sequencer\QuickStart "%Engine\Sequencer\QuickStart:description%")
diff --git a/Engine/Documentation/Source/Resources/Showcases/SequencerSubway/SequencerSubway.INT.udn b/Engine/Documentation/Source/Resources/Showcases/SequencerSubway/SequencerSubway.INT.udn
--- a/Engine/Documentation/Source/Resources/Showcases/SequencerSubway/SequencerSubway.INT.udn
+++ b/Engine/Documentation/Source/Resources/Showcases/SequencerSubway/SequencerSubway.INT.udn
@@ -5,6 +5,7 @@ Availability: Public
 Related: Engine/Sequencer
 Related: Engine/Sequencer/HowTo/TracksShot
 Related: Engine/Sequencer/HowTo/ViewportOptions
+tags:Sequencer
 
 [VAR:Topic]
 [OBJECT:Topic]
