#UE4 Docs: updated metadata

diff --git a/Engine/Documentation/Source/Engine/Sequencer/Workflow/Workflow.INT.udn b/Engine/Documentation/Source/Engine/Sequencer/Workflow/Workflow.INT.udn
--- a/Engine/Documentation/Source/Engine/Sequencer/Workflow/Workflow.INT.udn
+++ b/Engine/Documentation/Source/Engine/Sequencer/Workflow/Workflow.INT.udn
@@ -3,8 +3,8 @@ Title: Working with Sequencer
 Crumbs: %ROOT%, Engine
 Description:Tips, Tricks and other Considerations when working with Sequencer.
 Navigation:topic
-parent:Engine\Sequencer
-order:1
+parent:Engine/Sequencer
+order:5
 Version: 4.12
 
 There are several built-in shortcuts and keybindings you can use to speed up the rate at which you create your content in Sequencer. There are also several different workflows you can take depending upon the composition of your team. 
