#UE4 Docs: minor formatting and text updates

diff --git a/Engine/Documentation/Source/Engine/Rendering/LightingAndShadows/ReflectionEnvironment/ReflectionEnvironment.INT.udn b/Engine/Documentation/Source/Engine/Rendering/LightingAndShadows/ReflectionEnvironment/ReflectionEnvironment.INT.udn
--- a/Engine/Documentation/Source/Engine/Rendering/LightingAndShadows/ReflectionEnvironment/ReflectionEnvironment.INT.udn
+++ b/Engine/Documentation/Source/Engine/Rendering/LightingAndShadows/ReflectionEnvironment/ReflectionEnvironment.INT.udn
@@ -158,23 +158,25 @@ to see results.
 	[/PARAM]
 [/OBJECT]
 
-## Higher  Quality Reflections
+## High Quality Reflections
 
-While the default reflection quality settings strike a good balance between performance and visual quality, there could be instances where you want to achieve even higher quality reflections. In the following section, we will take a look at what you will need to do in order to achieve the highest quality reflections you can in UE4. 
+While the default reflection quality settings strike a good balance between performance and visual quality, there could be instances where you want to achieve even higher quality reflections. The following sections describe the available methods for achieving high quality reflections.
 
 ### High Precision Static Mesh Vertex Normal and Tangent Encoding
 
-An important factor in achieving high quality reflections is how accurately the vertex normal and tangent can be represented. Very high density meshes may lead to adjacent vertices quantizing to the same vertex normal and tangent values. This can lead to blocky jumps in normal orientation. We added the option to encode normals and tangents as 16 bits per channel vectors. This allows developers make the trade off between higher quality and how much additional memory is used encoding vertex buffers. You can enable High Precision Static Mesh Vertex Normal and Tangent Encoding in your project by doing the following:
+An important factor in achieving high quality reflections is how accurately the vertex normal and tangent can be represented. Very high density meshes may lead to adjacent vertices quantizing to the same vertex normal and tangent values. This can lead to blocky jumps in normal orientation. We added the option to encode normals and tangents as 16 bits per channel vectors which enables developers to make the trade off between higher quality and how much additional memory is used encoding vertex buffers. 
 
-1. Double-click on a Static Mesh in the **Content Browser** to open it up in the **Static Mesh Editor**.
+**To enable High Precision Static Mesh Vertex Normal and Tangent Encoding:**
 
-1. Once the Static Mesh editor is opened, go the **Details** panel and expand the **LOD0** option.
+1. In the **Content Browser**, double-click on a Static Mesh to open it up in the **Static Mesh Editor**.
+
+1. In the Static Mesh Editor, go the **Details** panel and expand the **LOD0** option.
 
 	[REGION:lightbox]
 	[![](SM_HP_Vertex_Normals_00.png)(w:400 h:350)](SM_HP_Vertex_Normals_00.png)
 	[/REGION]
 
-1. At the bottom of the LOD0, there is a section called **Build Settings**. Click on the small triangle next Build Settings to expand the Build Settings options.
+1. At the bottom of the **LOD0**, there is a section called **Build Settings**. Click on the small triangle next Build Settings to expand the Build Settings options.
 
 	[REGION:lightbox]
 	[![](SM_HP_Vertex_Normals_01.png)(w:400 h:350)](SM_HP_Vertex_Normals_01.png)
@@ -186,7 +188,11 @@ An important factor in achieving high quality reflections is how accurately the
 	[![](SM_HP_Vertex_Normals_02.png)(w:400 h:350)](SM_HP_Vertex_Normals_02.png)
 	[/REGION]
 
-Once the option has been enabled, your view will automatically update to reflect the changes. Note that the quality of the reflection that is viewed is directly related to how densely tessellated the Static Mesh is. Static Meshes that have less tessellation will have more stretching artifacts in the reflection than Static Meshes that have more tessellation.     
+The viewport will automatically update to reflect the changes. 
+
+[REGION:note]
+The quality of the reflection that is viewed is directly related to how densely tessellated the Static Mesh is. Static Meshes that have less tessellation will have more stretching artifacts in the reflection than Static Meshes that have more tessellation.
+[/REGION] 
 
 [OBJECT:ComparisonSlider]
 [PARAM:before]
@@ -204,9 +210,11 @@ Once the option has been enabled, your view will automatically update to reflect
 
 ### High Precision GBuffer Normal Encoding
 
-Enabling the High Precision GBuffer Normal Encoding option will allow the GBuffer to use a higher precision Normal encoding. This higher precision GBuffer Normal encoding encodes the Normal vector into three channels with each channel having 16 bits per. Using this higher precision encoding allows techniques like Screen Space Reflections (SSR) to rely on high precision normals. To enable this option in your UE4 project you will need to do the following:
+Enabling the High Precision GBuffer Normal Encoding option will allow the GBuffer to use a higher precision Normal encoding. This higher precision GBuffer Normal encoding encodes the Normal vector into three channels with each channel having 16 bits per. Using this higher precision encoding allows techniques like Screen Space Reflections (SSR) to rely on high precision normals. 
 
-1. Open up your **Project Settings** by going to the **Main Toolbar** and then selecting **Edit** >  **Project Settings**.
+**To enable High Precision GBuffer Normal Encoding:**
+
+1. In the **Main Toolbar**, select **Edit** >  **Project Settings** to open the **Project Settings**.
 
 	[REGION:lightbox]
 	[![](ARCS_00.png)(w:400 h:350)](ARCS_00.png)
@@ -222,16 +230,16 @@ Enabling the High Precision GBuffer Normal Encoding option will allow the GBuffe
 	Keep in mind that this encoding requires increased GPU memory and enabling this will have a direct impact on your project's performance.
 	[/region]
 
-1. Since changing the GBuffer format does not require you to restart the Editor, so you can quickly change between the different GBuffer formats to see the impact they will have on the reflection visuals. In the image below we can see how changing the GBuffer format from the Default to High Precision Normals changes the look and quality of the reflection.
+Since changing the GBuffer format does not require you to restart the Editor, you can quickly change between the different GBuffer formats to see the impact they will have on the reflection visuals. In the image below we can see how changing the GBuffer format from the Default to High Precision Normals changes the look and quality of the reflection.
 
-	[OBJECT:ComparisonSlider]
-	[PARAM:before]
-	![Default GBuffer Format](Default_GBuffer_Format.png)(h:600)
-	[/PARAM]
-	[PARAM:after]
-	![High Precision GBuffer Format](Hight_P_GBuffer_Format.png)(h:600)
-	[/PARAM]
-	[/OBJECT]
+[OBJECT:ComparisonSlider]
+[PARAM:before]
+![Default GBuffer Format](Default_GBuffer_Format.png)(h:600)
+[/PARAM]
+[PARAM:after]
+![High Precision GBuffer Format](Hight_P_GBuffer_Format.png)(h:600)
+[/PARAM]
+[/OBJECT]
 
 ## Reflection capture shapes
 
