#UE4 docs added warning to pages about dynamic shadows, corrected some tool tips

diff --git a/Engine/Documentation/Source/Engine/HLOD/HLODLanding.INT.udn b/Engine/Documentation/Source/Engine/HLOD/HLODLanding.INT.udn
--- a/Engine/Documentation/Source/Engine/HLOD/HLODLanding.INT.udn
+++ b/Engine/Documentation/Source/Engine/HLOD/HLODLanding.INT.udn
@@ -6,6 +6,10 @@ version:4.12
 
 **Hierarchical Level of Detail** meshes (HLODs) have been developed as a way to increase performance and keep the quality bar high for projects using Unreal Engine 4.11 or later. HLODs can replace multiple static mesh actors with a group of static mesh actors at long view distances. This helps reduce the number of actors that need to be rendered for the scene, increasing performance by lowering the number of draw calls per frame. 
 
+[REGION:warning]
+HLOD doesn't currently work well with dynamic shadows (casting or recieving them), and as such, you should use static lighting and shadowing for the meshes that the HLOD system generates.
+[/REGION]
+
 [OBJECT:TopicButtonList]
 	[PARAM:icon]
 		![](%ROOT%/tech_icon.png)(convert:false)
diff --git a/Engine/Documentation/Source/Engine/HLOD/HowTo/BuildingHLODs/BuildingHLODs.INT.udn b/Engine/Documentation/Source/Engine/HLOD/HowTo/BuildingHLODs/BuildingHLODs.INT.udn
--- a/Engine/Documentation/Source/Engine/HLOD/HowTo/BuildingHLODs/BuildingHLODs.INT.udn
+++ b/Engine/Documentation/Source/Engine/HLOD/HowTo/BuildingHLODs/BuildingHLODs.INT.udn
@@ -10,6 +10,10 @@ To build your Hierarchical Level of Detail (HLOD) Meshes, there are 2 major step
 
 For each LOD level you'll need to specify the settings for that LOD (such as 
 
+[REGION:warning]
+HLOD doesn't currently work well with dynamic shadows (casting or recieving them), and as such, you should use static lighting and shadowing for the meshes that the HLOD system generates.
+[/REGION]
+
 ## Steps
 
 Below are the steps to generate your HLOD meshes broken into 2 major sections: Generate Clusters and Generate Proxy Meshes.
@@ -42,7 +46,7 @@ To visualize how the generated clusters work in the editor you can use the "Forc
 
 ![image alt text](Engine/HLOD/Overview/image_8.png)(h:17 w:624 convert:false)
 
-Should you run into any issues you don’t like with the generated cluster you can expand the specified cluster and select the offending static mesh actor. By right-clicking on the actor you have the option to Remove or Exclude it from the generation of the cluster. 
+Should you run into any issues you don’t like with the generated cluster you can expand the specified cluster and select the offending static mesh actor. You can then click and drag it to another cluster or by right-clicking on the actor you have the option to **Remove** or **Exclude** it from the generation of the cluster. 
 
 ![image alt text](Engine/HLOD/Overview/image_9.png)(h:101 w:392 convert:false)
 
diff --git a/Engine/Documentation/Source/Engine/HLOD/HowTo/Enabling/EnablingHLODs.INT.udn b/Engine/Documentation/Source/Engine/HLOD/HowTo/Enabling/EnablingHLODs.INT.udn
--- a/Engine/Documentation/Source/Engine/HLOD/HowTo/Enabling/EnablingHLODs.INT.udn
+++ b/Engine/Documentation/Source/Engine/HLOD/HowTo/Enabling/EnablingHLODs.INT.udn
@@ -8,6 +8,10 @@ related:Engine/HLOD/Reference
 
 The **Hierarchical Level of Detail System** (HLOD) must be enabled on each Level that wants to make use of it. Follow the steps below to enable the Hierarchical LOD System.
 
+[REGION:warning]
+HLOD doesn't currently work well with dynamic shadows (casting or recieving them), and as such, you should use static lighting and shadowing for the meshes that the HLOD system generates.
+[/REGION]
+
 ## Steps
 
 1. Enable by going to **World Settings -> LODSystem -> Enable Hierarchical LODSystem**
diff --git a/Engine/Documentation/Source/Engine/HLOD/Overview/HLODOverview.INT.udn b/Engine/Documentation/Source/Engine/HLOD/Overview/HLODOverview.INT.udn
--- a/Engine/Documentation/Source/Engine/HLOD/Overview/HLODOverview.INT.udn
+++ b/Engine/Documentation/Source/Engine/HLOD/Overview/HLODOverview.INT.udn
@@ -5,6 +5,10 @@ availability:docs
 
 HLODs, in their simplest form, combine pre-existing **Static Mesh Actors** into a single Static Mesh Actor and material, with atlased textures, to reduce draw calls which in-turn helps with your performance of your game. There is a robust system in place to help you set specifically what needs to be generated when you create these HLOD actors for your scenes and covers more options that is normally available for LODs for a single actor. 
 
+[REGION:warning]
+HLOD doesn't currently work well with dynamic shadows (casting or recieving them), and as such, you should use static lighting and shadowing for the meshes that the HLOD system generates.
+[/REGION]
+
 ## Proxy Meshes
 
 Once your proxy meshes have been generated you can find these located under the HLOD folder generated in your project. This will include the LODActors referenced in the HLOD Outliner. 
diff --git a/Engine/Documentation/Source/Engine/HLOD/Reference/HLODReference.INT.udn b/Engine/Documentation/Source/Engine/HLOD/Reference/HLODReference.INT.udn
--- a/Engine/Documentation/Source/Engine/HLOD/Reference/HLODReference.INT.udn
+++ b/Engine/Documentation/Source/Engine/HLOD/Reference/HLODReference.INT.udn
@@ -5,6 +5,10 @@ availability:docs
 
 This page details the buttons and properties that exist on the Hierarchical LOD Outliner.
 
+[REGION:warning]
+HLOD doesn't currently work well with dynamic shadows (casting or recieving them), and as such, you should use static lighting and shadowing for the meshes that the HLOD system generates.
+[/REGION]
+
 ## Property and Interface Reference
 
 ![HLOD Outliner Fully Expanded](image_2.png)
@@ -49,9 +53,9 @@ Below, broken out by major section, are the properties that can be found in a **
 | **Merge Physics Data:** || Whether to merge the collision primitives when combining the actors. |
 | **Merge Materials** || Whether to merge source materials into one flat material. |
 | [REGION:tablesection]**[Material Settings See Below](#materialsettings)**[/REGION] |||
-| **Bake Vertex Data ** || This will bake any vertex data found into the LOD. |
-| **Calculate Correct LODModel** || Calculates the correct LOD Model. |
-| **Export Specific LOD** || Exports the specified LOD. |
+| **Bake Vertex Data ** || Incorporates painted vertex colours into model/baked out material. |
+| **Calculate Correct LODModel** || Calculates the correct LOD model based upon the screen size value (similar to how regular LODs are determined for static meshes in the editor/game). |
+| **Export Specific LOD** || Forces the system to pick the given LOD level for all models (if the LOD level exists, otherwise picks the maximum closest value). |
 
 ### Material Settings
 
@@ -96,7 +100,7 @@ Texture Sizing Type Description
 | Type | %spacer2% | Description |
 | --- | :---: | :--- |	
 | **Use TextureSize for all material properties** || Will use the value entered into **Texture Size** for sizing. |
-| **Use automatically biased texture sizes based on TextureSize** || Will use the value entered into **Texture Size**, but bias it based on usage. |
+| **Use automatically biased texture sizes based on TextureSize** || This creates a biased set of **{normal:diffuse:otherproperties}** textures using the ratio of **{1:0.5:0.25}** so if the texture size is set to 512 it will output a 512 normal map, 256 diffuse map and 128 map for roughness specular etc (diffuse is clamped to 32x32 and roughness specular etc to 16x16 |
 | **Use per property manually overridden texture sizes** || Enables the per texture overrides, and uses those values instead of **Texture Size**. |
 | **Use Simplygon’s automatic texture sizing** || Uses **Simplygon's** sizing. |
 [/REGION]
