#UE4 Docs: minor formatting and text updates

diff --git a/Engine/Documentation/Source/Engine/Sequencer/QuickStart/1/SequencerQS_1.INT.udn b/Engine/Documentation/Source/Engine/Sequencer/QuickStart/1/SequencerQS_1.INT.udn
--- a/Engine/Documentation/Source/Engine/Sequencer/QuickStart/1/SequencerQS_1.INT.udn
+++ b/Engine/Documentation/Source/Engine/Sequencer/QuickStart/1/SequencerQS_1.INT.udn
@@ -6,24 +6,12 @@ SkillLevel:Beginner
 Version:4.12
 Related:Engine\Sequencer
 Related:Engine\Sequencer\Overview
-checkpoint: editorqs
+checkpoint: sequencerqs
+parent:Engine/Sequencer/QuickStart
+type:sequential
+order:1
  
-[VAR:Steps]
-[OBJECT:Navigation]
-	[PARAM:previous]
-	[/PARAM]
-	[PARAM:current]
-	[/PARAM]
-	[PARAM:home]
-		[Quick Start Home](Engine\Sequencer\QuickStart)
-	[/PARAM]
-	[PARAM:next]
-		[Next Step](Engine\Sequencer\QuickStart\2 "%Engine\Sequencer\QuickStart\2:title%")
-	[/PARAM]
-[/OBJECT]
-[/VAR]
-
-%Steps%
+[NAV]
 
 Before we start crafting our scene, we will create some of the assets we'll need as well as set up the Editor UI for previewing our cinematic in real-time. 
 For the Editor Layout, feel free to experiment with the different options and choose the one that you feel most comfortable with. 
@@ -55,7 +43,7 @@ For the duration of this guide, we are using the **Blueprint Third Person Templa
 
 	Above we have docked the Sequencer Editor next to the **Content Browser**. This allows us to either click the **Content Browser** tab to see the Content Browser or click the **Sequencer** tab to open the Sequencer Editor. 
 
-1.	In the **Content Browser**, **Right-click** on and duplicate your **Master** Level Sequence and call it **Shot01**. 
+1.	In the **Content Browser**, **Right-click** on the **Master** Level Sequence and choose **Duplicate**. Name the new Level Sequence **Shot01**. 
 
 	![](SequencerQS_1_3.png)
 
@@ -77,13 +65,11 @@ For the duration of this guide, we are using the **Blueprint Third Person Templa
 
 	![](SequencerQS_1_7.png)
 
-1.	In the right Viewport, press **G** then press **Ctrl+Shift+T**. 
+1.	In the right Viewport, press **G** then press **Ctrl+Shift+T**. Pressing the **G** key enters [Game View](Engine/UI/LevelEditor/Viewports/Basics) which hides all editor icons and represents what the game (or cinematic) will look like. 
+	Pressing **Ctrl+Shift+T** toggles the display of the Viewport Toolbar (which includes things like grid snapping, view modes and transform gizmos which you may not need for a cinematic preview). 
 
 	![](SequencerQS_1_8.png)
 
-	Pressing the **G** key enters [Game View](Engine/UI/LevelEditor/Viewports/Basics) which hides all editor icons and represents what the game (or cinematic) will look like. 
-	Pressing **Ctrl+Shift+T** toggles the display of the Viewport Toolbar (which includes things like grid snapping, view modes and transform gizmos which you may not need for a cinematic preview). 
-
 	[REGION:tip]
 	You can press **G** again to re-enable the editor icons or press **Ctrl+Shift+T** to display the Viewport Toolbar again. 
 	[/REGION]
@@ -96,6 +82,6 @@ We have created our **Master** Level Sequence and our first **Shot** as well as
 
 Up next we will work on our first shot, add a camera and adjust focus as well as generate movement for the camera.
 
-%Steps%
+[NAV]
 
 
diff --git a/Engine/Documentation/Source/Engine/Sequencer/QuickStart/2/SequencerQS_2.INT.udn b/Engine/Documentation/Source/Engine/Sequencer/QuickStart/2/SequencerQS_2.INT.udn
--- a/Engine/Documentation/Source/Engine/Sequencer/QuickStart/2/SequencerQS_2.INT.udn
+++ b/Engine/Documentation/Source/Engine/Sequencer/QuickStart/2/SequencerQS_2.INT.udn
@@ -6,25 +6,12 @@ SkillLevel:Beginner
 Version:4.12
 Related:Engine\Sequencer
 Related:Engine\Sequencer\Overview
-checkpoint: editorqs
+checkpoint: sequencerqs
+parent:Engine/Sequencer/QuickStart
+type:sequential
+order:2
  
-[VAR:Steps]
-[OBJECT:Navigation]
-	[PARAM:previous]
-		[Previous Step](Engine\Sequencer\QuickStart\1 "%Engine\Sequencer\QuickStart\1:title%")
-	[/PARAM]
-	[PARAM:current]
-	[/PARAM]
-	[PARAM:home]
-		[Quick Start Home](Engine/QuickStart)
-	[/PARAM]
-	[PARAM:next]
-		[Next Step](Engine\Sequencer\QuickStart\3 "%Engine\Sequencer\QuickStart\3:title%")
-	[/PARAM]
-[/OBJECT]
-[/VAR]
-
-%Steps%
+[NAV]
 
 With our project ready to go, in this step we will create the first shot in our cinematic. 
 We'll add a camera to that shot and generate movement as well as see how to adjust the focus of a camera through **Camera Settings**. 
@@ -41,8 +28,7 @@ We'll add a camera to that shot and generate movement as well as see how to adju
 
 	![](SequencerQS_2_1b.png)
 
-	This type of track is used to define the perspective (or camera) to shoot the scene from at any given time. 
-	When using the **Add Camera** button, the camera that is created is automatically assigned as the Camera Cut to use for you. 
+	This type of track is used to define the perspective (or camera) to shoot the scene from at any given time. When using the **Add Camera** button, the camera that is created is automatically assigned as the Camera Cut to use for you. 
 
 1.	In the right Viewport, hold the **Right-mouse** button and use the **WASD** keys to move the camera to a position like below. 
 
@@ -60,42 +46,30 @@ We'll add a camera to that shot and generate movement as well as see how to adju
 
 	Notice how the Active Viewport is the one the camera is locked to when clicking the Unlock/Lock Camera button. 
 
-1.	With the **CineCameraActor** selected in **Sequencer**, press the **S** key. 
+1.	With the **CineCameraActor** selected in **Sequencer**, press the **S** key. This is a [**Workflow Shortcut**](Engine\Sequencer\Workflow\CheatSheet\) that can be used to add a keyframe for the currently selected Actor's Transform (Location, Rotation and Scale). Since the Timeline is at 0 by default, we have initialized the position of the camera by adding the keyframe. When we scrub ahead and add another keyframe with the camera at a different position, the camera will move from one keyframe to the next over time.
 
-	![](SequencerQS_2_5.png)
+	![](SequencerQS_2_5.png)  
 
-	This is a [**Workflow Shortcut**](Engine\Sequencer\Workflow\CheatSheet\) that can be used to add a keyframe for the currently selected Actor's Transform (Location, Rotation and Scale). 
-	Since the Timeline is at 0 by default, we have initialized the position of the camera by adding the keyframe. 
-	When we scrub ahead and add another keyframe with the camera at a different position, the camera will move from one keyframe to the next over time.  
-
-1.	With the **CineCameraActor** selected, in the **Details** panel under **Current Camera Settings** expand **Focus Settings**.
-
-1.	Click the **Eye Dropper** icon next to **Manual Focus**. 
+1.	With the **CineCameraActor** selected, in the **Details** panel under **Current Camera Settings** expand **Focus Settings** and click the **Eye Dropper** icon next to **Manual Focus** to set the focus depth based on the scene sample location we provide. 
 
 	![](SequencerQS_2_6.png)
 
-	This will allow us to set the focus depth based on the scene sample location we provide. 
-
 1.	In the left Viewport, move the **Eye Dropper** icon over the character in the level and **Left-click** to select it. 
 
 	![](SequencerQS_2_7.png)
 
 	After clicking, the character will come into focus. 
 
-1.	With the scene sampled, click the **Add Key** button next to the **Manual Focus Distance** property. 
+1.	With the scene sampled, click the **Add Key** button next to the **Manual Focus Distance** property. This will add a key for the value of the property to the open Level Sequence. Any properties that have an **Add Key** icon, you can click on with a Level Sequence open to keyframe the value of the property in Sequencer. 
 
 	![](SequencerQS_2_8.png)
 
-	This will add a key for the value of the property to the open Level Sequence. 
-	Any properties that have an **Add Key** icon, you can click on with a Level Sequence open to keyframe the value of the property in Sequencer. 
-
 1.	In the Cinematic Viewport, click the **To End** button to move the timeline to the end of the Sequence. 
 
 	![](SequencerQS_2_8b.png)
 
 	[REGION:warning]
-	As of 4.12.2 there is an issue when scrubbing the timeline with Spawnable Cameras. If you scrub the timeline to the end of the sequence, the camera will become despawned and will lose focus in the viewport. 
-	You can re-gain focus in the viewport by clicking the **To End** button and locking the viewport to the camera. 
+	As of 4.12.2 there is an issue when scrubbing the timeline with Spawnable Cameras. If you scrub the timeline to the end of the sequence, the camera will become despawned and will lose focus in the viewport. You can re-gain focus in the viewport by clicking the **To End** button and locking the viewport to the camera. 
 
 	![](LockCamera.png)
 
@@ -120,11 +94,8 @@ We'll add a camera to that shot and generate movement as well as see how to adju
 
 ## End Result
 
-The camera work for our first shot of our cinematic is complete! 
-We added a camera to our shot and assigned a **Camera Cuts** track to indicate what camera to shoot from, we added some camera movement for the shot and adjusted camera settings so our character is in focus. 
+The camera work for our first shot of our cinematic is complete! We added a camera to our shot and assigned a **Camera Cuts** track to indicate what camera to shoot from, we added some camera movement for the shot and adjusted camera settings so our character is in focus. 
 
 In the next step, we will add and animate a character and preview the end result of our first shot. 
 
-%Steps%
-
-
+[NAV]
\ No newline at end of file
diff --git a/Engine/Documentation/Source/Engine/Sequencer/QuickStart/3/SequencerQS_3.INT.udn b/Engine/Documentation/Source/Engine/Sequencer/QuickStart/3/SequencerQS_3.INT.udn
--- a/Engine/Documentation/Source/Engine/Sequencer/QuickStart/3/SequencerQS_3.INT.udn
+++ b/Engine/Documentation/Source/Engine/Sequencer/QuickStart/3/SequencerQS_3.INT.udn
@@ -6,29 +6,14 @@ SkillLevel:Beginner
 Version:4.12
 Related:Engine\Sequencer
 Related:Engine\Sequencer\Overview
-checkpoint: editorqs
+checkpoint: sequencerqs
+parent:Engine/Sequencer/QuickStart
+type:sequential
+order:3
  
-[VAR:Steps]
-[OBJECT:Navigation]
-	[PARAM:previous]
-		[Previous Step](Engine\Sequencer\QuickStart\2 "%Engine\Sequencer\QuickStart\2:title%")
-	[/PARAM]
-	[PARAM:current]
-	[/PARAM]
-	[PARAM:home]
-		[Quick Start Home](Engine\Sequencer\QuickStart)
-	[/PARAM]
-	[PARAM:next]
-		[Next Step](Engine\Sequencer\QuickStart\4 "%Engine\Sequencer\QuickStart\4:title%")
-	[/PARAM]
-[/OBJECT]
-[/VAR]
-
-%Steps%
+[NAV]
 
-In the previous step we added a camera and provided a path for our camera to take. 
-Up next we will replace the Character Blueprint that is used with the template with a **Skeletal Mesh** Actor that we can apply animations to. 
-By the end of this step, our first shot will be complete. 
+In the previous step, we added a camera and provided a path for our camera to take. In this step, we replace the Character Blueprint that is used with the template with a **Skeletal Mesh** Actor that we can apply animations to, completing our first shot.
 
 ## Steps
 
@@ -36,17 +21,15 @@ By the end of this step, our first shot will be complete.
 	
 	![](SequencerQS_3_1.png)
 
-1. In the left Viewport, **Right-click** on the character in the level then select **Replace Selected Actors with** option. 
+1. In the left Viewport, **Right-click** on the character in the level and choose **Replace Selected Actors with**. This will replace the default **Character Blueprint** used for gameplay with a **Skeletal Mesh Actor**.
 
-	![](SequencerQS_3_2.png)
+	![](SequencerQS_3_2.png) 
 
-	This will replace the default **Character Blueprint** used for gameplay with a **Skeletal Mesh Actor**. 
-
-1.	In the left Viewport, press the **End** key to drop the Skeletal Mesh down to the ground, then press **E** and rotate it as shown below. 
+1.	In the left Viewport, press the **End** key to drop the Skeletal Mesh down to the ground. Press **E** and rotate it as shown below. 
 
 	![](SequencerQS_3_3.png)
 
-1.	With the Skeletal Mesh still selected, in **Sequencer** click the **Add** button and it to the Level Sequence. 
+1.	With the Skeletal Mesh still selected, in **Sequencer** click the **Add** button to add it to the Level Sequence. 
 
 	![](SequencerQS_3_4.png)
 
@@ -54,21 +37,17 @@ By the end of this step, our first shot will be complete.
 
 	![](SequencerQS_3_5a.png)
 
-1.	For the **ThirdPersonCharacter** Track, click the **+ Animation** button and search for and add the **ThirdPersonIdle**.
-
-	![](SequencerQS_3_5.png)
+1.	For the **ThirdPersonCharacter** Track, click the **+ Animation** button. Search for the **ThirdPersonIdle** animation and add it to the track. The Skeletal Mesh in the level will assume an Idle pose.
 
-	The Skeletal Mesh in the level will assume an Idle pose. 
+	![](SequencerQS_3_5.png) 
 
-1.	**Left-click** and drag the end of the **ThirdPersonIdle** so it loops and fills the entire 150 frames. 
+1.	**Left-click** and drag the end of the **ThirdPersonIdle** animation so it loops and fills the entire 150 frames. 
 
 	![](SequencerQS_3_6.png)
 
-1.	On the **Camera Cuts** track, click the **Unlock/Lock Camera** button. 
+1.	On the **Camera Cuts** track, click the **Unlock/Lock Camera** button. Locking the Camera on a Camera Cuts track, will only affect Cinematic Viewports.
 
-	![](SequencerQS_3_6b.png)	
-
-	Locking the Camera on a Camera Cuts track, will only affect Cinematic Viewports. 
+	![](SequencerQS_3_6b.png) 
 
 1.	Click the **Preview Play** button in the Cinematic Preview window. 
 
@@ -91,8 +70,6 @@ vx-HlYIbpOI
 [/OBJECT]
 
 
-In the next step we will work on another shot for our scene before we add both shots to our Master track. 
-
-%Steps%
-
+In the next step, we will work on another shot for our scene before we add both shots to our Master track. 
 
+[NAV]
\ No newline at end of file
diff --git a/Engine/Documentation/Source/Engine/Sequencer/QuickStart/4/SequencerQS_4.INT.udn b/Engine/Documentation/Source/Engine/Sequencer/QuickStart/4/SequencerQS_4.INT.udn
--- a/Engine/Documentation/Source/Engine/Sequencer/QuickStart/4/SequencerQS_4.INT.udn
+++ b/Engine/Documentation/Source/Engine/Sequencer/QuickStart/4/SequencerQS_4.INT.udn
@@ -6,25 +6,12 @@ SkillLevel:Beginner
 Version:4.12
 Related:Engine\Sequencer
 Related:Engine\Sequencer\Overview
-checkpoint: editorqs
+checkpoint: sequencerqs
+parent:Engine/Sequencer/QuickStart
+type:sequential
+order:4
  
-[VAR:Steps]
-[OBJECT:Navigation]
-	[PARAM:previous]
-		[Previous Step](Engine\Sequencer\QuickStart\3 "%Engine\Sequencer\QuickStart\3:title%")
-	[/PARAM]
-	[PARAM:current]
-	[/PARAM]
-	[PARAM:home]
-		[Quick Start Home](Engine\Sequencer\QuickStart)
-	[/PARAM]
-	[PARAM:next]
-		[Next Step](Engine\Sequencer\QuickStart\5 "%Engine\Sequencer\QuickStart\5:title%")
-	[/PARAM]
-[/OBJECT]
-[/VAR]
-
-%Steps%
+[NAV]
 
 Now that you have completed your first shot, let's take what we have learned and create another!
 
@@ -32,47 +19,47 @@ This time we will push the camera in from the front of our character in on their
 
 ## Steps
 
-1.	In the **Content Browser**, select the **Shot01** Level Sequence then press **Shift+W** to duplicate it. 
+1.	In the **Content Browser**, select the **Shot01** Level Sequence and press **Shift+W** to duplicate it. 
 
 	![](SequencerQS_4_1.png)
 
-	Our Shot01 already has most of the existing content we will want to also use in our second shot. 
+	Shot01 already has most of the existing content we want to use in the second shot. 
 
-1.	Inside **Shot02**, delete the existing **CineCameraActor** then click the **Add Camera** button. 
+1.	Inside **Shot02**, delete the existing **CineCameraActor** and click the **Add Camera** button to automatically assign it on the Camera Cuts track and enter the camera in Pilot Mode. 
 
 	![](SequencerQS_4_2.png)
 
-	This will automatically assign it on the Camera Cuts track and enter the camera in Pilot Mode. 
-
 1.	Move the camera around to the front of the Skeletal Mesh Actor in the level. 
 
 	![](SequencerQS_4_3.png)
 
+	[REGION:note]
 	The name of your CineCameraActor may differ from above which is perfectly fine. 
-	
-1.	Select the **CineCameraActor** track in Sequencer, then press **S** to add a keyframe the Camera's Transform.
+	[/REGION]
+
+1.	Select the **CineCameraActor** track in Sequencer and press **S** to add a keyframe the Camera's Transform.
 
 	![](SequencerQS_4_3b.png)
 
-1.	In the **Details** panel for the camera, use the **Eye Dropper** to focus in on the character then add a keyframe. 
+1.	In the **Details** panel for the camera, use the **Eye Dropper** to focus in on the character. Add a keyframe. 
 
 	![](SequencerQS_4_4.png)
 
+	[REGION:note]
 	Your value might be different from above based on where you placed the camera. 
+	[/REGION]
 
-1.	Use the **To End** button in the Cinematic Preview to jump to the end of the Sequence. 
+1.	Click the **To End** button in the Cinematic Preview to jump to the end of the Sequence. 
 
 1.	Move the camera in on the character then add a keyframe (**S**) for the position. 
 
 	![](SequencerQS_4_5.png)
 
-1.	In the **Details** panel for the camera, use the **Eye Dropper** to focus in on the character then add a keyframe. 
-
-	![](SequencerQS_4_6.png)
+1.	In the **Details** panel for the camera, use the **Eye Dropper** to focus in on the character. Add a keyframe. You should have keyframes at the start for the initial position and focus value as well as at the end.
 
-	You should have keyframes at the start for the initial position and focus value as well as at the end. 
+	![](SequencerQS_4_6.png) 
 
-1.	You can use the **Camera Lock** toggle button on the **Camera Cuts** track to enable the Cinematic Preview. 
+1.	Click the **Camera Lock** toggle button on the **Camera Cuts** track to enable the Cinematic Preview. 
 
 	![](SequencerQS_4_7.png)
 
@@ -81,10 +68,6 @@ This time we will push the camera in from the front of our character in on their
 
 We have set up our second shot of our sequence where the camera starts from far away and pushes in on the character. 
 
-In our next step we will tie our shots together using the **Shot Track** in our Master Level Sequence. 
-We will also create a **Take** (or variation) of one of our shots that we can quickly swap in or out. 
-You will also see how we can arrange the order of shots in our scene as we get a little bit closer to finishing our sample scene. 
-
-%Steps%
-
+In our next step, we tie our shots together using the **Shot Track** in our Master Level Sequence. We also create a **Take** (or variation) of one of our shots that we can quickly swap in or out. You will see how we can arrange the order of shots in our scene as we get a little bit closer to finishing our sample scene. 
 
+[NAV]
\ No newline at end of file
diff --git a/Engine/Documentation/Source/Engine/Sequencer/QuickStart/5/SequencerQS_5.INT.udn b/Engine/Documentation/Source/Engine/Sequencer/QuickStart/5/SequencerQS_5.INT.udn
--- a/Engine/Documentation/Source/Engine/Sequencer/QuickStart/5/SequencerQS_5.INT.udn
+++ b/Engine/Documentation/Source/Engine/Sequencer/QuickStart/5/SequencerQS_5.INT.udn
@@ -6,34 +6,20 @@ SkillLevel:Beginner
 Version:4.12
 Related:Engine\Sequencer
 Related:Engine\Sequencer\Overview
-checkpoint: editorqs
+checkpoint: sequencerqs
+parent:Engine/Sequencer/QuickStart
+type:sequential
+order:5
  
-[VAR:Steps]
-[OBJECT:Navigation]
-	[PARAM:previous]
-		[Previous Step](Engine\Sequencer\QuickStart\4 "%Engine\Sequencer\QuickStart\4:title%")
-	[/PARAM]
-	[PARAM:current]
-	[/PARAM]
-	[PARAM:home]
-		[Quick Start Home](Engine/QuickStart)
-	[/PARAM]
-	[PARAM:next]
-		[Next Step](Engine\Sequencer\QuickStart\6 "%Engine\Sequencer\QuickStart\6:title%")
-	[/PARAM]
-[/OBJECT]
-[/VAR]
-
-%Steps%
+[NAV]
 
-In the previous steps we created our shots, now we will tie them together using the **Shot Track** which allows you to add Level Sequences along the timeline for playback. 
-We will also create a **Take** of one of those Shots, which will allow us to create a variant of that shot or use it as an additional Shot in the cinematic. 
+In the previous steps, we created our shots. Now, we tie them together using the **Shot Track** which enables you to add Level Sequences along the timeline for playback. We also create a **Take** of one of those Shots, which enables us to create a variant of that shot or use it as an additional Shot in the cinematic. 
 
-At the end of this step, you will have created your own multi-shot cinematic with Sequencer! 
+At the end of this step, you will have your own multi-shot cinematic! 
 
 ## Steps
 
-1. In the **Content Browser** open the **Master** Level Sequence. 
+1. In the **Content Browser**, open the **Master** Level Sequence. 
 
 	![](SequencerQS_5_1.png)
 
@@ -45,45 +31,37 @@ At the end of this step, you will have created your own multi-shot cinematic wit
 
 	![](SequencerQS_5_3.png)
 
-1. Add **Shot02** to the **Shot Track**. 
-
-	![](SequencerQS_5_4.png)
+1. Add **Shot02** to the **Shot Track**. Notice that Shot02 extends past the red maker (which is the end of playback range). Drag the red marker to the end of Shot02 or automatically keep added sections within playback range by enabling the option in the **General Options** menu.
 
-	Notice that Shot02 extends past the red maker (which is the end of playback range). 
-	You can drag the red marker to the end of Shot02 or automatically keep added sections within playback range by enabling the option in the **General Options** menu. 
+	![](SequencerQS_5_4.png) 
 
-1. Click the **General Options** button then select **Keep Playback Range in Section Bounds**. 
+1. Click the **General Options** button and select **Keep Playback Range in Section Bounds**. This pushes the end marker out for us anytime we add content and will ensure that everything is within playback range. 
 
 	![](SequencerQS_5_5.png)
 
-	This will push the end marker out for us anytime we add content and will ensure that everything is within playback range. 
-
 1. **Left-click** and drag **Shot02** down a level within the **Shot Track**. 
 
 	![](SequencerQS_5_6.png)
 
 	[REGION:tip]
-	You can drag shots around to arrange the playback order or for organization purposes. 
-	A typical best practice is to alternate between tracks within a Shot Track to illustrate where cuts occur rather than leaving them all on the same single track. 
+	You can drag shots around to arrange the playback order or for organization purposes. A typical best practice is to alternate between tracks within a Shot Track to illustrate where cuts occur rather than leaving them all on the same single track. 
 
 	![](ShotCuts.png)
 
-	Below we have the same number of shots without using an adddtional Track within the Shots Track. 
+	Below, we have the same number of shots without using an adddtional Track within the Shots Track. 
 
 	![](ShotCutsB.png)
 	[/REGION]
 
-1. Click the **Lock Viewport** toggle button to lock the camera to the **Shot Track**. 
-
-	![](SequencerQS_5_7.png)
+1. Click the **Lock Viewport** toggle button to lock the camera to the **Shot Track**. Like the **Camera Cuts** Camera Lock button, the Shot Track will only affect the Cinematic Viewport.
 
-	Like the **Camera Cuts** Camera Lock button, the Shot Track will only affect the Cinematic Viewport. 
+	![](SequencerQS_5_7.png) 
 
-1. **Right-click** on **Shot02** then select the **New Take** option. 
+1. **Right-click** on **Shot02** and select the **New Take** option. 
 
 	![](SequencerQS_5_8.png)
 
-	A **Duplicate Asset As** window will appear, click **Save**. 
+1. A **Duplicate Asset As** window will appear, click **Save**. 
 
 	![](SequencerQS_5_9.png)
 
@@ -96,31 +74,26 @@ At the end of this step, you will have created your own multi-shot cinematic wit
 	![](SequencerQS_5_10.png)
 
 	[REGION:note]
-	When creating a Take it does not delete the Level Sequence but creates a duplicate of it which you can then alter. 
-	Currently when creating a Take for the first time it replaces the Shot with the new Take (the original Shot still exists in the Content Browser but is not listed as a "Take" in the Takes Menu - see below). 
+	When creating a Take, it does not delete the Level Sequence but creates a duplicate of it which you can then alter. Currently, when creating a Take for the first time, it replaces the Shot with the new Take (the original Shot still exists in the Content Browser but is not listed as a "Take" in the Takes Menu - see below). 
 	[/REGION]
 
-1.	**Right-click** on the newly created shot and highlight **Takes**. 
+1.	**Right-click** on the newly created shot and highlight **Takes**. We are currently using **Take 3** in our cinematic. We can switch to Take 2 by select it from this menu. 
 
 	![](SequencerQS_5_11.png)
 
-	Above we are currently using **Take 3** in our cinematic. We can switch to Take 2 by select it from this menu. 
-
-1.	**Double-click** on the **Shot0002_003** to open it, then **Left-click** drag to select and delete each of the keyframes within it. 
+1.	**Double-click** on the **Shot0002_003** to open it. **Left-click** drag to select and delete each of the keyframes within it. We are going to change the current camera movement to a different motion for this take. 
 
 	![](SequencerQS_5_12.png)
 
-	We are going to change the current camera movement to a different motion for this take. 
-
-1.	Lock the Viewport to the camera, the move it in near the character's face and press **S** to keyframe the position. 
+1.	Lock the Viewport to the camera. Move it near the character's face and press **S** to keyframe the position. 
 
 	![](SequencerQS_5_14.png)
 
-1.	In the **Details** panel for the camera, adjust the **Manual Focus Distance** then add a keyframe for the adjusted value. 
+1.	In the **Details** panel for the camera, adjust the **Manual Focus Distance** and add a keyframe for the adjusted value. 
 
 	![](SequencerQS_5_15.png)
 
-1.	Jump to the end of the sequence, the move the camera back away from the character and keyframe it. 
+1.	Jump to the end of the sequence. Move the camera back away from the character and keyframe it. 
 
 	![](SequencerQS_5_16.png)
 
@@ -131,7 +104,7 @@ At the end of this step, you will have created your own multi-shot cinematic wit
 
 ## End Result
 
-You should now have a two shot cinematic, however you can use different **Takes** from the Take Menu. 
+You now have a two shot cinematic, however you can use different **Takes** from the Take Menu. 
 
 [OBJECT:EmbeddedVideo]
 [PARAMLITERAL:width]
@@ -145,11 +118,8 @@ Z2E4lcodFww
 [/PARAMLITERAL]
 [/OBJECT]
 
-You can also re-arrange the Shots by dragging them to new locations on the Shot Track. You can also add a Take as an additional Shot, giving you three Shots. 
-While this is only a small fraction of what is possible with Sequencer, it highlights some of the core aspects of the tool that you can use to craft your scenes. 
-
-In the next and final step, you will find some "On Your Own" challenges you can attempt to add to this scene and links to additional documentation. 
-
-%Steps%
+You can re-arrange the Shots by dragging them to new locations on the Shot Track or add a Take as an additional Shot, giving you three Shots. While this is only a small fraction of what is possible with Sequencer, it highlights some of the core aspects of the tool that you can use to craft your scenes. 
 
+In the next and final step, there are some "On Your Own" challenges you can attempt to add to this scene and links to additional documentation. 
 
+[NAV]
\ No newline at end of file
diff --git a/Engine/Documentation/Source/Engine/Sequencer/QuickStart/6/SequencerQS_6.INT.udn b/Engine/Documentation/Source/Engine/Sequencer/QuickStart/6/SequencerQS_6.INT.udn
--- a/Engine/Documentation/Source/Engine/Sequencer/QuickStart/6/SequencerQS_6.INT.udn
+++ b/Engine/Documentation/Source/Engine/Sequencer/QuickStart/6/SequencerQS_6.INT.udn
@@ -6,24 +6,12 @@ SkillLevel:Beginner
 Version:4.12
 Related:Engine\Sequencer
 Related:Engine\Sequencer\Overview
-checkpoint: editorqs
+checkpoint: sequencerqs
+parent:Engine/Sequencer/QuickStart
+type:sequential
+order:6
  
-[VAR:Steps]
-[OBJECT:Navigation]
-	[PARAM:previous]
-		[Previous Step](Engine\Sequencer\QuickStart\5 "%Engine\Sequencer\QuickStart\5:title%")
-	[/PARAM]
-	[PARAM:current]
-	[/PARAM]
-	[PARAM:home]
-		[Quick Start Home](Engine\Sequencer\QuickStart)
-	[/PARAM]
-	[PARAM:next]
-	[/PARAM]
-[/OBJECT]
-[/VAR]
-
-%Steps%
+[NAV]
 
 Now that you have created a sample cinematic with Sequencer, you can try several of the How-to's to add to this scene. 
 
@@ -55,14 +43,9 @@ Optionally, you can **Render** out your cinematic as a video file:
 
 	![](ExportMovie.png)
 
-1. Under **General**, define the **Output Directory** then click **Capture Movie** to start rendering.
-
-	![](ExportMovieB.png)
-
-	A preview window will appear while rendering and will close when complete. 
-
-
+1. Under **General**, define the **Output Directory** then click **Capture Movie** to start rendering. A preview window will appear while rendering and will close when complete.
 
-%Steps%
+	![](ExportMovieB.png) 
 
 
+[NAV]
\ No newline at end of file
diff --git a/Engine/Documentation/Source/Engine/Sequencer/QuickStart/SequencerQuickStart.INT.udn b/Engine/Documentation/Source/Engine/Sequencer/QuickStart/SequencerQuickStart.INT.udn
--- a/Engine/Documentation/Source/Engine/Sequencer/QuickStart/SequencerQuickStart.INT.udn
+++ b/Engine/Documentation/Source/Engine/Sequencer/QuickStart/SequencerQuickStart.INT.udn
@@ -6,6 +6,7 @@ Navigation:topic
 parent:Engine/Sequencer
 order:1
 Version: 4.12
+type:quickstart
 
 
 [VAR:Topic]
