#UE4 Docs: minor text update

diff --git a/Engine/Documentation/Source/Engine/UMG/UserGuide/WidgetTypeReference/WebBrowser/WebBrowser.INT.udn b/Engine/Documentation/Source/Engine/UMG/UserGuide/WidgetTypeReference/WebBrowser/WebBrowser.INT.udn
--- a/Engine/Documentation/Source/Engine/UMG/UserGuide/WidgetTypeReference/WebBrowser/WebBrowser.INT.udn
+++ b/Engine/Documentation/Source/Engine/UMG/UserGuide/WidgetTypeReference/WebBrowser/WebBrowser.INT.udn
@@ -21,7 +21,7 @@ This widget is used to display an in-game web browser that can host content from
 
 ![](WebBrowserWidget.png)(w:600)
 
-The Web Browser widget works not just for desktop machines, but for Android and iOS as well. 
+The Web Browser widget is supported on desktop platforms as well as Android and iOS mobile devices. 
 
 [REGION:raw]
 ![](Devices.png)
