#UE4 Docs: minor formatting and text updates

diff --git a/Engine/Documentation/Source/Engine/Foliage/Foliage.INT.udn b/Engine/Documentation/Source/Engine/Foliage/Foliage.INT.udn
--- a/Engine/Documentation/Source/Engine/Foliage/Foliage.INT.udn
+++ b/Engine/Documentation/Source/Engine/Foliage/Foliage.INT.udn
@@ -292,7 +292,7 @@ The example Material below multiplies the Material mask by the fading factor to
 
 ![Culling Material Example](Foliage_MaterialSetup.png)
 
-### LOD
+## LOD
 
 Static Mesh LOD is supported on foliage. There are a few caveats:
 
@@ -300,7 +300,7 @@ Static Mesh LOD is supported on foliage. There are a few caveats:
 * Note that light and shadow maps are shared between all LOD levels, so all LODs need to use the same unwrapping.
 * Currently the entire cluster of instances change LODs simultaneously. We may add support for distance-based fading per instance in the future.
 
-### Lighting
+## Lighting
 
 Each individual mesh instance has its own shadow and/or light map generated by Lightmass, and these are tiled together for each precomputed batch. There are several settings on the Static Mesh that should be checked for precomputed lighting to operate well with instanced foliage. Lightmass is less forgiving when generating shadow maps for instanced meshes, and incorrect settings could lead to black meshes after rebuilding lighting.
 
@@ -310,22 +310,21 @@ Each individual mesh instance has its own shadow and/or light map generated by L
 
 
 
-### Foliage Scalability
+## Foliage Scalability
 
 Foliage Static Meshes can  use the [Scalability System](Engine/Performance/Scalability/ScalabilityReference/) to increase or decrease the amount of Foliage Instances that are currently being rendered on screen. To add this functionality to your project you will need to do the following:
 
 [region:note]
-In Order to use the Foliage Painting mode, you will need to make sure that your level has either a Landscape Terrain or Static Meshes with Collision enabled. 
+In Order to use the Foliage Painting mode, your level must contain either a Landscape Terrain or Static Meshes with Collision enabled. 
 [/region]
 
-
-1. Right-click in the Content Browser and from the menu that is displayed, under the **Create Advanced Asset** section, expand the **Miscellaneous** section and select the **Foliage Type** option.
+1. Right-click in the Content Browser. Under the **Create Advanced Asset** section, expand the **Miscellaneous** section and select the **Foliage Type** option.
 
 	[REGION:lightbox]
 	[![](Foliage_Scalability_00.png)(w:300 h:350)](Foliage_Scalability_00.png)
 	[/REGION]
 
-1. Input a name for the Foliage Type, for this example the name **FT_Box** was used, then double - click on the FT_Box Foliage Type to open it up.  
+1. Input a name for the Foliage Type (for this example we are using **FT_Box**) and double-click on the FT_Box Foliage Type to open it.  
 
 	[REGION:lightbox]
 	[![](Foliage_Scalability_01.png)(w:405 h:223)](Foliage_Scalability_01.png)
@@ -337,7 +336,7 @@ In Order to use the Foliage Painting mode, you will need to make sure that your
 	[![](Foliage_Scalability_02.png)(w:405 h:223)](Foliage_Scalability_02.png)
 	[/REGION]
 
-1. In the **Content Browser** select a Static Mesh, for this example we are using the **Shape_Cube** that comes with the Starter Content, then in the FT_Box Foliage Type, under the **Mesh** section, drag the **Shape_Cube** from the Content Browser to the **Mesh** option.
+1. In the **Content Browser** select a Static Mesh (for this example we are using the **Shape_Cube** that comes with the Starter Content). In the FT_Box Foliage Type, under the **Mesh** section, drag the **Shape_Cube** from the Content Browser to the **Mesh** option.
 
 	[REGION:lightbox]
 	[![](Foliage_Scalability_03.png)(w:405 h:223)](Foliage_Scalability_03.png)
@@ -353,23 +352,23 @@ In Order to use the Foliage Painting mode, you will need to make sure that your
 	The **Enable Density Scaling** option should typically only be **enabled** for detail foliage without collision such as small rocks and bushes. It should be **disabled** for Foliage Types with collision such as trees as the collision for removed instances will continue to exist in the game.
 	[/region]
 
-1. On the **Modes** panel, click on the **Foliage Tool** option then select and drag the FT_Box Foliage Type from the Content Browser to the area that says **Drop Foliage Here**.
+1. In the **Modes** panel, select the **Foliage Tool**. Select and drag the FT_Box Foliage Type from the Content Browser to the area that says **Drop Foliage Here**.
 
 	[REGION:lightbox]
 	[![](Foliage_Scalability_05.png)(w:250 h:350)](Foliage_Scalability_05.png)
 	[/REGION]
 
-1. Select the FT_Box Foliage Type that was just placed in the Foliage Type list to select it and then in the viewport hold down the **Left Mouse Button** to begin painting were the FT_Bo Foliage Type meshes should go. 
+1. Select the FT_Box Foliage Type you placed in the Foliage Type list to select it. In the viewport, hold down the **Left Mouse Button** to begin painting the FT_Bo Foliage Type meshes onto a surface. 
 
 	[REGION:lightbox]
 	[![](Foliage_Scalability_06.png)(w:592 h:398)](Foliage_Scalability_06.png)
 	[/REGION]
 
-1. On the **Toolbar** click on the **Play** button to launch the level and then navigate to a location where Foliage has been added.
+1. On the **Toolbar**, click on the **Play** button to launch the level. Navigate to a location where Foliage has been added.
 
 	![](Foliage_Scalability_07.png)
 
-1. You can control the density at which the Foliage is rendered by pressing the **Backtick** key (`) to open up the **Console** by and then enter the following command followed by a decimal number in the 0 to 1 range.
+1. You can control the density at which the Foliage is rendered by pressing the **Backtick** key (`) to open up the **Console** and entering the following command followed by a decimal number in the 0 to 1 range.
 
 	* **foliage.DensityScale**
 
@@ -377,7 +376,7 @@ In Order to use the Foliage Painting mode, you will need to make sure that your
 	Adjusting the Foliage Density Scale setting and seeing the results update in real time does not work in any of the Editor viewports. The only way to see the Foliage Density update in real time is to input the command while the project is in a game mode like PIE, Simulate, or Standalone. 
 	[/region]
 
-In the image comparison below you can see what happens to the density of the Foliage when the foliage.DensityScale setting was set to 0.1, 0.2, 0.3, 0.4, 0.5, 0.6, 0.7, 0.8, 0.9, 1.0.
+Below, the density of the Foliage changes as the foliage.DensityScale setting is set to 0.1, 0.2, 0.3, 0.4, 0.5, 0.6, 0.7, 0.8, 0.9, 1.0.
 
 [REGION:fullwidth]
 [INCLUDE:Engine\Foliage\#FS]
@@ -396,7 +395,7 @@ In the image comparison below you can see what happens to the density of the Fol
 	FS
 	[/PARAMLITERAL]
 	[PARAM:caption]
-	**Dragging the slider will show what happens when adjusting the foliage.DensityScale setting using numbers from 0.1 to 1.0.**
+	**Drag the slider to adjust the foliage.DensityScale setting from 0.1 to 1.0.**
 	[/PARAM]
 	[PARAM:image_list]
 	![](FS_1.png)
