#UE4 Docs: minor formatting and text updates

diff --git a/Engine/Documentation/Source/Gameplay/Networking/CharacterMovementComponent/CharacterMovementComponent.INT.udn b/Engine/Documentation/Source/Gameplay/Networking/CharacterMovementComponent/CharacterMovementComponent.INT.udn
--- a/Engine/Documentation/Source/Gameplay/Networking/CharacterMovementComponent/CharacterMovementComponent.INT.udn
+++ b/Engine/Documentation/Source/Gameplay/Networking/CharacterMovementComponent/CharacterMovementComponent.INT.udn
@@ -8,10 +8,11 @@ version: 4.12
 [TOC (start:2 end:3)]
 
 Characters using `CharacterMovementComponent` automatically have client-server networking built in. Here's how player movement prediction, replication and correction works in networked 
-games through `CharacterMovementComponent`:
+games through `CharacterMovementComponent` for each frame:
 
-On every tick, the `TickComponent` function is called.  It figures out the acceleration and rotation change for the frame, and then calls `PerformMovement` (for locally 
-controlled Characters), or `ReplicateMoveToServer` (for network clients).
+* `TickComponent` function is called
+* The acceleration and rotation change for the frame are calculted
+* Either `PerformMovement` (for locally controlled Characters) or `ReplicateMoveToServer` (for network clients) is called
 	
 `ReplicateMoveToServer` saves the move (into a list of pending moves), calls `PerformMovement` locally, and then replicates the move to the server by calling the replicated 
 function `ServerMove`. `ServerMove` receives the movement parameters, including the client's final position and a timestamp. It executes on the server, where it decodes the 
@@ -73,18 +74,22 @@ passes on the game client. Games using unmodified CharacterMovementComponent fun
 rates. To prevent this type of exploit, we have added the ability to detect and resolve time discrepencies to `CharacterMovementComponent`.
 
 ### Detection
+
 Time discrepency detection works by having the server compare the reported time difference between `ServerMove` RPCs (Remote Procedure Calls) coming from the client against 
 the known (and trusted) amount of time that has passed on the server. This difference is kept as a running total, and once the total exceeds a user-configurable threshold, resolution 
 measures are taken.
 
 ### Resolution
+
 The method of resolution employed is to override the timestamps of future `ServerMove` RPCs coming from the client to match the server's time. A portion of the difference in the 
 Character's movement data (location, rotation, acceleration, etc.) between the RPCs can then be subtracted to "pay back" the time discrepency until the client is back in sync with 
 the server.
 
 ### Configuration
+
 By default, time discrepency detection and resolution are disabled. The following variables in `GameNetworkManager` can be configured, and their default values are found in 
 BaseGame.ini. These values can be overridden in your project-specific game configuration file.
+
 | --- | --- |
 | **Variable Name** | **Effect** |
 | bMovementTimeDiscrepancyDetection | Enables detection. Upon detection, warning logs will be triggered. If resolution is enabled, it will be applied as well. |
@@ -94,6 +99,7 @@ BaseGame.ini. These values can be overridden in your project-specific game confi
 | MovementTimeDiscrepancyResolutionRate | Rate at which to "pay back" time during resolution. The default is 100%, meaning the client cannot move until the time debt is repaid. |
 | MovementTimeDiscrepancyDriftAllowance | Accepted drift in clocks between client and server as a percent per second allowed. This helps to prevent burst packet loss or performance hitches from triggering false positives. |
 | bMovementTimeDiscrepancyForceCorrectionsDuringResolution | Whether or not to force client updates during resolution. This is needed for projects that have lenient movement error tolerances or ClientAuthorativePosition enabled. |
+
 These settings may need to be tuned to fit your specific project. Basic testing can be done by using the "slomo" cheat on a client. The variable "p.DebugTimeDiscrepancy" will activate 
 time-discrepency logging on servers.
 
@@ -104,6 +110,7 @@ Character to move 10 meters forward upon pressing the T key, as long as the dest
 game.
 
 ###Approach 1: Execute only on Client
+
 One approach might be to check for obstructions, then move the character directly forward 10 meters if the destination is clear. This does not consider networking, and simply executes 
 on the client.
 
@@ -115,6 +122,7 @@ should end up. As far as the server is concerned, the 10 meter difference betwee
 part. The server corrects the error, and the client applies the server's correction locally, undoing the teleport.
 
 ###Approach 2: Server only RPC
+
 The simplest working approach for networked games is to have the server handle the ability and then inform the client of the result. To do this, we will set up a `UFUNCTION` tagged 
 as `Reliable` and `Server` that executes the teleport code when triggered from the client.
 
@@ -135,6 +143,7 @@ those effects will not show up correctly since the actual ability only runs on t
 server), there is no report of it succeeding or failing, just a correction that shows up shortly afterward.
 
 ###Approach 3: Server RPC and Local Trigger
+
 In this approach, the client executes the Teleport locally, and then also calls the server's Teleport RPC.
 
 **Result**: Works in networked games, with rare but potentially serious issues.
@@ -146,6 +155,7 @@ Teleport was triggered, that client won't know to retrigger the Teleport as they
 will make the game feel unresponsive in certain cases, especially in bad network conditions.
 
 ###Approach 4: CharacterMovementComponent Ability Implementation
+
 In this approach, we add knowledge of the teleport ability to a child class of `CharacterMovementComponent`. This version will be robust over the network.
 
 **Result**: Works in networked games, with some implementation care.
