#UE4 Docs: minor formatting and text updates

diff --git a/Engine/Documentation/Source/Engine/Rendering/LightingAndShadows/PlanarReflections/PlanarReflections.INT.udn b/Engine/Documentation/Source/Engine/Rendering/LightingAndShadows/PlanarReflections/PlanarReflections.INT.udn
--- a/Engine/Documentation/Source/Engine/Rendering/LightingAndShadows/PlanarReflections/PlanarReflections.INT.udn
+++ b/Engine/Documentation/Source/Engine/Rendering/LightingAndShadows/PlanarReflections/PlanarReflections.INT.udn
@@ -30,7 +30,7 @@ Screen Space Reflections (SSR) are much more efficient than Planar Reflections w
 ### Enabling Planar Reflections 
 To use enable and use Planar Reflections in your project you will need to do the following:
 
-1. From the Main Toolbar go to **Edit**> **Project Settings** and from the Project Settings menu, go to the **Rendering** > **Optimizations** section.
+1. From the Main Toolbar go to **Edit > Project Settings** and from the Project Settings menu, go to the **Rendering > Optimizations** section.
 	
 	[REGION:lightbox]
 	[![](Engine\Rendering\LightingAndShadows\ReflectionEnvironment\PR_Project_Settings_00.png)(w:502 h:322)](Engine\Rendering\LightingAndShadows\ReflectionEnvironment\PR_Project_Settings_00.png)
@@ -92,7 +92,7 @@ The Planar Reflection Actor has a number of properties that when adjusted can si
 While Planar Reflections offer some very realistic reflections, this features does have a number of limitations which are listed below.
 
 * There's no way to control what rendering features are enabled in the reflection pass.
-* Dynamic shadowing will incorrect in the reflection pass.
+* Dynamic shadowing is incorrect in the reflection pass.
 * To hit your target framerate ensure that you budget enough resources for Planar Reflections.
 * Only a constant Roughness value is supported, and it's specified on the Planar Reflection Component instead of Materials.
 * If at all possible limit the amount of Planar Reflection Actors placed in the world to only one and move, rotate, or scale it to best fit the world. While you can use multiple Planar Reflection Actors, you need to exercise caution as Planar Reflection Actors don not do any sort of distance culling, only frustum, and occlusion culling. Because of this if both Planar Reflection Actors are in view at the same time your project's frame rate will suffer severely.
@@ -147,7 +147,7 @@ While adding a Planar Reflection makes the water look more lifelike, the impact
 [INCLUDE:#ElvenRuinsReflection]
 [/REGION]
 
-The reason that Planar Reflections only cost 1.67 ms to render in the Elven Ruins map versus the 23.07 ms that it takes to render the Planar Reflections in the Kite Demo is due to how the Static Meshes and Materials that are used in the Elven Ruins were built. Since the Elven Ruins map and it's content was designed for mobile devices, close attention was paid to the amount of triangles and Material instructions each asset used. Because of this, when Planar Reflections are enabled in the Elven Ruins map, they will cost less as the assets that are used in the map are greatly reduced in complexity and size when compared to the Kite Demo assets.
+The reason that Planar Reflections only cost 1.67 ms to render in the Elven Ruins map versus the 23.07 ms that it takes to render the Planar Reflections in the Kite Demo is due to how the Static Meshes and Materials that are used in the Elven Ruins were built. Since the Elven Ruins map and its content was designed for mobile devices, close attention was paid to the amount of triangles and Material instructions each asset used. Because of this, when Planar Reflections are enabled in the Elven Ruins map, they will cost less as the assets that are used in the map are greatly reduced in complexity and size when compared to the Kite Demo assets.
 
 <!--  
 [EXCERPT:KitePlanarReflection]
