#ue4 docs

Added section about property favorites to the existing Details panel page.

diff --git a/Engine/Documentation/Source/Engine/UI/LevelEditor/Details/LevelEditorDetails.INT.udn b/Engine/Documentation/Source/Engine/UI/LevelEditor/Details/LevelEditorDetails.INT.udn
--- a/Engine/Documentation/Source/Engine/UI/LevelEditor/Details/LevelEditorDetails.INT.udn
+++ b/Engine/Documentation/Source/Engine/UI/LevelEditor/Details/LevelEditorDetails.INT.udn
@@ -2,7 +2,7 @@ Availability:Public
 Title:Details Panel
 Crumbs: %ROOT%, Engine, Engine/UI, Engine/UI/LevelEditor
 Description:Panel providing information and customized editing tools for selected Actors in the Level Editor.
-Version: 4.9
+Version: 4.12
 
 [VAR:playlistnav]
 [OBJECT:PlaylistNavigation]
@@ -61,6 +61,45 @@ Because this line hides away any details that do not match the search criteria,
 see the property you are looking for.
 [/REGION]
 
+## Favorites
+
+[REGION:warning]
+This option is currently considered experimental and some features may not work as expected. 
+[/REGION]
+
+If there are properties within an Actor that you frequently change or want quick access to, you can use the property **Favorites** to flag them so they appear at the top of the **Details** panel. 
+
+![](Favorites.png)
+
+[REGION:caption]
+Above we checked two options as favorites indicated by the star icon next to their property name for easy access.
+[/REGION]
+
+You can enable Favorites from inside the **Editor Preferences**. 
+
+![](EditorPrefs.png)
+
+Under **Experimental**, check the **Enable Details Panel Favorites** option. 
+
+![](Experimental.png)
+
+[REGION:note]
+You may need to restart the editor for the changes to be applied. 
+[/REGION]
+
+With the option enabled, next to a property you can click the star icon to mark it as a favorite. 
+
+![](ClickFavorite.png)
+
+This will list the property (along with any other marked favorites) under the **Favorites** section of the panel. 
+
+![](MarkedFavorites.png)
+
+[REGION:note]
+Some properties may not offer the ability to favorite them, due their complexity of their customization. 
+[/REGION]
+
+
 ## Default Values
 
 When a property is modified to a value other than its _default_ value, an indicator is displayed.
