Merging

//depot/UE4/Engine/Documentation/...

to //UE4/Main/Engine/Documentation/...

diff --git a/Engine/Documentation/Source/Engine/Animation/AnimBlueprints/Creation/CreatingVimBlueprints.INT.udn b/Engine/Documentation/Source/Engine/Animation/AnimBlueprints/Creation/CreatingVimBlueprints.INT.udn
--- a/Engine/Documentation/Source/Engine/Animation/AnimBlueprints/Creation/CreatingVimBlueprints.INT.udn
+++ b/Engine/Documentation/Source/Engine/Animation/AnimBlueprints/Creation/CreatingVimBlueprints.INT.udn
@@ -61,12 +61,12 @@ When creating a new Animation Blueprint, you will be asked to pick a valid [**Sk
 
 	![](CE_create_NoSkel.png)
 
-1. You will then be prompted to specify which Skeletal to target for the Animation Blueprint, select one and click **OK**. 
+1. You will then be prompted to specify which Skeleton to target for the Animation Blueprint. Select one and click **OK**. 
 
 	![](PickSkeleton.png)
 
 	You are also given the ability to select a [**Parent Class**](Engine/Blueprints/UserGuide/Types/ClassBlueprint) for the Animation Blueprint to use. 
-	Most of the time you will not need to select one and can use the default class, however based on your circumstances, you can choose one if you need to. 
+	Most of the time you will not need to select one and can use the default class. However, based on your circumstances, you can choose one if you need to. 
 
 1. The new asset will automatically be created in the **Content Browser**. 
 	
diff --git a/Engine/Documentation/Source/Engine/Animation/AnimBlueprints/VimBlueprints.INT.udn b/Engine/Documentation/Source/Engine/Animation/AnimBlueprints/VimBlueprints.INT.udn
--- a/Engine/Documentation/Source/Engine/Animation/AnimBlueprints/VimBlueprints.INT.udn
+++ b/Engine/Documentation/Source/Engine/Animation/AnimBlueprints/VimBlueprints.INT.udn
@@ -84,12 +84,12 @@ Check out the links below for more information on each topic.
 
 ## Working with Child Animation Blueprints
 
-When creating and setting up your characters and Animation Blueprints, there may be instances where you have multiple characters that are similar but have entirely different animations for different states. Rather than create Animation Blueprints for each of those characters, you may be able to use a **Child Animation Blueprint** of a pre-existing Animation Blueprint and override the animation assets that you want to replace. The Child Animation Blueprint will inherit everything from the parent, however it will also allow you to specify which animations to override via the **Asset Override Editor**. 
+When creating and setting up your characters and Animation Blueprints, there may be instances where you have multiple characters that are similar but have entirely different animations for different states. Rather than create Animation Blueprints for each of those characters, you may be able to use a **Child Animation Blueprint** of a pre-existing Animation Blueprint and override the animation assets that you want to replace. The Child Animation Blueprint will inherit everything from the parent, but it will also allow you to specify which animations to override via the **Asset Override Editor**. 
 
-For example, suppose you have two humanoid type characters where one is a human and another is a werewolf. Both characters can walk, run, jump, crawl, etc. however their motions may be entirely different. The logic of moving between those motions is the same, but the motions themselves are different. It's perfectly fine to create two different Animation Blueprints, but to save some time; you can create a Child Animation Blueprint and selectively choose which animations to override. 
+For example, suppose you have two humanoid type characters where one is a human and another is a werewolf. Both characters can walk, run, jump, crawl, etc. However, their motions may be entirely different. The logic of moving between those motions is the same, but the motions themselves are different. It's perfectly fine to create two different Animation Blueprints, but to save some time. you can create a Child Animation Blueprint and selectively choose which animations to override. 
 
 [REGION:note]
-See [**Animation Blueprint Override**](Engine/Animation/AnimHowTo/AnimBPOverride) for an example walk-through of creating and overriding an Animation Blueprint with a Child.
+See [**Animation Blueprint Override**](Engine/Animation/AnimHowTo/AnimBPOverride) for an example walk-through of creating and overriding an Animation Blueprint with a Child Animation Blueprint.
 [/REGION]
 
 
diff --git a/Engine/Documentation/Source/Engine/Sequencer/HowTo/CameraRigCrane/CameraRigCrane.INT.udn b/Engine/Documentation/Source/Engine/Sequencer/HowTo/CameraRigCrane/CameraRigCrane.INT.udn
--- a/Engine/Documentation/Source/Engine/Sequencer/HowTo/CameraRigCrane/CameraRigCrane.INT.udn
+++ b/Engine/Documentation/Source/Engine/Sequencer/HowTo/CameraRigCrane/CameraRigCrane.INT.udn
@@ -124,7 +124,7 @@ RAsfP5G0Shk
 [/PARAMLITERAL]
 [/OBJECT]
 
-You can adjust the keys by dragging them to new values in the keyframe window. Above we slowed the cranes upward movement by moving the **Crane Arm Length** and **Crane Pitch** keyframes we added at frame **50** and moved them to frame **100**. 
+You can adjust the keys by dragging them to new values in the Tracks Area. Above we slowed the cranes upward movement by moving the **Crane Arm Length** and **Crane Pitch** keyframes we added at frame **50** and moved them to frame **100**. 
 
 [REGION:tip]
 There is an advanced method you can use to work with the **Camera Rig Crane** asset in Sequencer that will automate several of the steps from this how-to. Please see the [**Workflow Shortcuts**](Engine\Sequencer\Workflow\CheatSheet) section for more information. 
diff --git a/Engine/Documentation/Source/Engine/Sequencer/HowTo/TracksPlayRate/PlayRate.INT.udn b/Engine/Documentation/Source/Engine/Sequencer/HowTo/TracksPlayRate/PlayRate.INT.udn
--- a/Engine/Documentation/Source/Engine/Sequencer/HowTo/TracksPlayRate/PlayRate.INT.udn
+++ b/Engine/Documentation/Source/Engine/Sequencer/HowTo/TracksPlayRate/PlayRate.INT.udn
@@ -48,7 +48,7 @@ For this how-to, we are using the **Blueprint Third Person Template** project wi
 	Above we are going from 1 to 0.25 then staying at 0.25 for several frames before returning to our normal 1.0 Play Rate. 
 
 	[REGION:tip]
-	You can **Right-click** drag to move around in the Keyframe window and **CTRL + Mouse Wheel** in/out to zoom in/out. 
+	You can **Right-click** drag to move around in the Tracks Area and **CTRL + Mouse Wheel** in/out to zoom in/out. 
 	[/REGION]
 
 1.	On the **Content Browser** tab under **Content/Mannequin/Character/Mesh**, drag the **SK_Mannequin** into the level. 
diff --git a/Engine/Documentation/Source/Engine/Sequencer/ReferenceEditor/EditorReference.INT.udn b/Engine/Documentation/Source/Engine/Sequencer/ReferenceEditor/EditorReference.INT.udn
--- a/Engine/Documentation/Source/Engine/Sequencer/ReferenceEditor/EditorReference.INT.udn
+++ b/Engine/Documentation/Source/Engine/Sequencer/ReferenceEditor/EditorReference.INT.udn
@@ -26,8 +26,8 @@ New **Level Sequences** will not have any content in them and the image above is
 1. [Sequence Breadcrumbs](#sequencebreadcrumbs)
 1. [Add/Filter Tracks](#add/filtertracks)
 1. [Timeline](#timeline)
-1. [Tracks Window](#trackswindow)
-1. [Keyframe Window](#keyframewindow)
+1. [Tree View](#treeview)
+1. [Tracks Area](#tracksarea)
 1. [Playback Controls](#playbackcontrols)
 
 ## Tool Bar
@@ -45,14 +45,14 @@ What follows is a description of each of the toolbar buttons, from left to right
 | ![](tb_5.png) | Opens the General Options menu with settings to define functionality for Sequencer.  |
 | ![](tb_6.png) | Enables/Disables the Key All option. With this and auto-key enabled, determines if all properties are keyed for a track when changes are made or if only the changed property is keyed. |
 | ![](tb_7.png) | Enables/Disables different Auto-Key options. By default, this option is disabled but can be set to Auto-key All or Auto-key Animated. |
-| ![](tb_8.png) | Enables/Disables grid snapping for the Keyframe window as well as grid snapping options. |
+| ![](tb_8.png) | Enables/Disables grid snapping for the Tracks Area as well as grid snapping options. |
 | ![](tb_9.png) | Allows for adjustment of the time snapping internval. |
-| ![](tb_10.png) | Opens the [**Curve Editor**](Engine/UI/CurveEditor) which allows you to fine tune the animation keys placed in the Keyframe window. |
+| ![](tb_10.png) | Opens the [**Curve Editor**](Engine/UI/CurveEditor) which allows you to fine tune the animation keys placed in the Tracks Area. |
 
 ## Sequence Breadcrumbs
 
 Depending on your workflow, you may choose to create several **Level Sequences** as scenes and place each of those scenes in a new Level Sequence which is your "master" Level Sequence which plays those scenes back-to-back as part of a **Shots** track. 
-In these instances, you can jump from your master Level Sequence directly into each scene by **Double-clicking** on them inside the Keyframe window. 
+In these instances, you can jump from your master Level Sequence directly into each scene by **Double-clicking** on them inside the Tracks Area. 
 
 In the image below, we have a Level Sequence called _shot_06_ inside another Level Sequence called _Master_.
 
@@ -86,7 +86,7 @@ Below are the available Tracks you can add to Sequencer, click each link below t
 | [**Fade Track**](Engine\Sequencer\HowTo\TracksFade) | Adds a track that will allow you to control the fading in/out of the sequence.  |
 | [**Level Visibility Track**](Engine\Sequencer\HowTo\TracksLevelVisibility) | This can be used to keyframe and control the visibilty of your level. |
 
-The **Filter** option allows you to filter the **Tracks Window** for only Actors that have been added to the sequence from the level.
+The **Filter** option allows you to filter the **Tree View** for only Actors that have been added to the sequence from the level.
 
 ![](FilterOption.png)
 
@@ -107,9 +107,9 @@ Your sequence will continue to play from the start marker all the way to the end
 [/REGION]
 
 
-## Tracks Window
+## Tree View
 
-The **Tracks** window is where you determine what types of actions you want to perform during your sequence. 
+The **Tree View** is where you determine what types of actions you want to perform during your sequence. 
 
 ![](tracks1.png)
 
@@ -120,18 +120,18 @@ For example, in the image below we have added a **CameraComponent** track to our
 
 ![](tracks2.png)
 
-You can also add Keyframes from the Tracks Window for your tracks to define properties at a specific point on the Timeline.
+You can also add Keyframes from the Tree View for your tracks to define properties at a specific point on the Timeline.
 
 ![](AddKeyFrameButton.png)
 
 You can click the **+** keyframe button or press the **Enter** key with a Track selected. 
 
 
-## Keyframe Window
+## Tracks Area
 
 ![](KeyframeWindow.png)
 
-They **Keyframe** Window is where you adjust the content in each of your Tracks. You can adjust any keyframes you have added, trim/split or taper in/out of sections as well as several other options that are available from the **Right-click** context menu (pictured below). 
+They **Tracks Area** is where you adjust the content in each of your Tracks. You can adjust any keyframes you have added, trim/split or taper in/out of sections as well as several other options that are available from the **Right-click** context menu (pictured below). 
 
 ![](RightClickContext1.png)
 
@@ -141,9 +141,9 @@ Below, **Right-clicking** on an Animation assigned to an Animation Track and sel
 
 ![](RightClickProperties.png)
 
-### Navigating the Keyframe Window
+### Navigating the Tracks Area
 
-Below are methods and shortcuts you can use to navigate the Keyframe window. 
+Below are methods and shortcuts you can use to navigate the Tracks Area. 
 
 [INCLUDE:Engine\Sequencer\Workflow\CheatSheet#keyframe]
 
diff --git a/Engine/Documentation/Source/Engine/Sequencer/Workflow/CheatSheet/CheatSheet.INT.udn b/Engine/Documentation/Source/Engine/Sequencer/Workflow/CheatSheet/CheatSheet.INT.udn
--- a/Engine/Documentation/Source/Engine/Sequencer/Workflow/CheatSheet/CheatSheet.INT.udn
+++ b/Engine/Documentation/Source/Engine/Sequencer/Workflow/CheatSheet/CheatSheet.INT.udn
@@ -29,7 +29,7 @@ To improve your productivity, several shortcuts have been built into the Sequenc
 | **]** | Sets the end of Playback Range to the current time. |
 
 
-## Keyframe Window Navigation
+## Tracks Area Navigation
 
 ![](KeyframeWindow.png)
 
@@ -38,9 +38,9 @@ To improve your productivity, several shortcuts have been built into the Sequenc
 [EXCERPT:keyframe]
 | Command | Description |
 | ---- | ----------- |
-| **Right-click Drag** | Allows you to move around the keyframe window. |
-| **Ctrl + Mouse Wheel In/Out** | Zoom in/out in the Keyframe window.  |
-| **Home** | Resets the zoom/view range of the Keyframe window. |
+| **Right-click Drag** | Allows you to move around the Tracks Area. |
+| **Ctrl + Mouse Wheel In/Out** | Zoom in/out in the Tracks Area.  |
+| **Home** | Resets the zoom/view range of the Tracks Area. |
 | **-** and **=** | Zoom into and out of the view range by 10%. |
 | **,** and **.** | Moves backwards and forwards between placed keyframes. |
 | **Left-click Drag** | Creates a selection box around keys in the window. |
@@ -75,7 +75,7 @@ The commands below do not respond to Numpad entries.
 
 ![](Editing.png)
 
-Select a section in the Keyframe window (audio or animation asset, camera, etc.) and edit it with the commands below. 
+Select a section in the Tracks Area (audio or animation asset, camera, etc.) and edit it with the commands below. 
 
 | Command | Description |
 | ---- | ----------- |
