#UE4 Docs: Applying Feedback peer feedback.

diff --git a/Engine/Documentation/Source/Engine/Landscape/TechnicalGuide/LansdscapeTechnicalGuide.INT.udn b/Engine/Documentation/Source/Engine/Landscape/TechnicalGuide/LansdscapeTechnicalGuide.INT.udn
--- a/Engine/Documentation/Source/Engine/Landscape/TechnicalGuide/LansdscapeTechnicalGuide.INT.udn
+++ b/Engine/Documentation/Source/Engine/Landscape/TechnicalGuide/LansdscapeTechnicalGuide.INT.udn
@@ -132,7 +132,8 @@ Selecting a channel will apply a shader to the Landscape which shows the area th
 [/EXCERPT:RecommendedLSS]
 
 ## Landscape Collision 
-The Unreal Engine 4 (UE4) Landscape system allows you to specify how complex both the simple and complex Landscape collision meshes should be. This system also allows you specify theses settings for the entire Landscape as a whole or for each individual Landscape Component and will work with any platform that supports Landscape. In the following section, we will go over how to use this system as well as any relevant information you must know before using this in your UE4 projects.
+
+Landscapes enable you to specify the detail level of the geometry used for both simple and complex collisions for the entire Landscape or on a per-Component basis. In the following sections, we will go over how to use the system as well as any relevant information you must know before using this in your UE4 projects.
 
 [region:note]
 For this example, we are using the Landscape Mountains project that can be found in the Learn tab of the UE4 Launcher.
@@ -149,7 +150,6 @@ If you select any Landscape Actor that has been placed in a level and go to the
 	[INCLUDE:Engine\Landscape\TechnicalGuide\#VC]
 	[/REGION]
 
-
 * **Simple Collision Mip Level:** The Simple Collision Mip Level sets the complexity of the **Simple** collision that is used for the Landscape. Setting Simple Collision Mip Level to **0**, the default,  will give you a very accurate Landscape collision at the expense of memory. Setting this value to **5**, the maximum setting, will make the Landscape collision cheaper, but the accuracy of the collision will suffer.
 
 	[REGION:fullwidth]
@@ -158,7 +158,7 @@ If you select any Landscape Actor that has been placed in a level and go to the
 
 ### Viewing Collision Mip Level 
 
-You can see what happens to the Landscape collision mesh when you change the values of the Simple Collision Mip Level option or the Collision Mip Level option by using the Collision Viewmodes. To make your viewport display the collision view mode, go to the **View Mode** menu in the Editor viewport and from the drop down list, under the **Collision** section, select either the **Player Collision** option or the **Visibility Option**.
+You can visualize the Landscape collision geometry using the Collision Viewmodes. To enable the collision view mode, go to the **View Mode** menu in the Editor viewport toolbar, under the **Collision** section, select either **Player Collision**  or **Visibility Collision**.
 
 [REGION:lightbox]
 [![](LSC_00.png)(w:524 h:303)](LSC_00.png)
@@ -169,7 +169,7 @@ Click for full image.
 [/REGION]
 
 
-* **Player Collision:** Using the **Player Collision** view mode will show you what the Simple Collision Mip Level will look like. 
+* **Player Collision:** The **Player Collision** view mode displays the Simple Collision Mip Level. 
 
 	[REGION:lightbox]
 	[![](CML_Player_Collision.png)(w:524 h:303)](CML_Player_Collision.png)
@@ -179,7 +179,7 @@ Click for full image.
 	Click for full image.
 	[/REGION]
 
-* **Visibility Collision:** Using the Visibility Collision view mode will show you what the Complex Collision Mip level will look like.
+* **Visibility Collision:** The Visibility Collision view mode displays the Complex Collision Mip Level.
 
 	[REGION:lightbox]
 	[![](CML_Vis_Collision.png)(w:524 h:303)](CML_Vis_Collision.png)
@@ -191,9 +191,9 @@ Click for full image.
 
 ### Adjusting the Landscape Collision Mip Level
 
-To increase or decrease the complexity for both simple and complex Landscape collision you will need to do the following.
+To set the complexity for both simple and complex Landscape collision you will need to do the following:
 
-1. Select your Landscape Terrain in the Editor viewport and then in the **Details** panel,  expand the **Collision** section.
+1. Select your Landscape Terrain in the Editor viewport. In the **Details** panel, expand the **Collision** section.
 
 	[REGION:lightbox]
 	[![](LSC_01.png)(w:524 h:303)](LSC_01.png)
@@ -203,7 +203,11 @@ To increase or decrease the complexity for both simple and complex Landscape col
 	Click for full image.
 	[/REGION]
 
-1. Under the **Collision** section, look for the **Collision Mip Level** option and change the value from **0** to **5** and then press the **Enter** key to apply the change. 
+1. Under the **Collision** section, find the **Collision Mip Level** option. Set the value  from **0** to **5** and then press the **Enter** key to apply the change. 
+
+	[REGION:caption]
+	The gray collision mesh in the level updates automatically to reflect the changes.
+	[/REGION]
 
 	[OBJECT:ComparisonSlider]
 		[PARAM:before]
@@ -214,15 +218,12 @@ To increase or decrease the complexity for both simple and complex Landscape col
 		[/PARAM]
 	[/OBJECT]
 
-	[REGION:caption]
-	When you do you should see the gray collision mesh in the level update automatically to reflect the changes.
-	[/REGION]
 
 ### Mixing Collision Mip Level Options 
 
-It is possible to independently increase or decrease the complexity of both the simple and complex Landscape collision meshes to help provide a better balance between performance and accuracy. To achieve this in your project, you will need to do the following.
+You can set the complexity of both the simple and complex Landscape collision meshes to provide a better balance between performance and accuracy. To set the simple and complex collision levels independently in your project you will need to do the following:
 
-1. Select your Landscape and in the **Details** panel, expand the **Collision** section.
+1. Select your Landscape and in the **Details** panel and expand the **Collision** section.
 
 	[REGION:lightbox]
 	[![](LSC_01.png)(w:524 h:303)](LSC_01.png)
@@ -231,7 +232,7 @@ It is possible to independently increase or decrease the complexity of both the
 	Click for full image.
 	[/REGION]
 
-1. Leave the **Collision Mip Level** at a value of **0** and then set the **Simple Collision Mip Level** to a value of **2**.
+1. Set the **Collision Mip Level** to a value of **0** and **Simple Collision Mip Level** to a value of **2**.
 
 	[REGION:lightbox]
 	[![](LSC_04.png)(w:524 h:303)](LSC_04.png)
@@ -241,27 +242,30 @@ It is possible to independently increase or decrease the complexity of both the
 	Click for full image.
 	[/REGION]
 
-1. In the following image comparison, you can see what happens to the Landscape collision when different values for the **Collision Mip Level** and **Simple Collision Mip Level** are set to different values.
+In the following image comparison, you can see what happens to the Landscape collision when the Collision Mip Level and Simple Collision Mip Level are set to different values.
 
-	[OBJECT:ComparisonSlider]
-		[PARAM:before]
-		![Player Collision 2](Player_Collision.png)(w:719 h:407)
-		[/PARAM]
-		[PARAM:after]
-		![Visibility Collision 0](Vis_Collision.png)(w:719 h:407)
-		[/PARAM]
-	[/OBJECT]
+[OBJECT:ComparisonSlider]
+	[PARAM:before]
+	![Player Collision 2](Player_Collision.png)(w:719 h:407)
+	[/PARAM]
+	[PARAM:after]
+	![Visibility Collision 0](Vis_Collision.png)(w:719 h:407)
+	[/PARAM]
+[/OBJECT]
 
-	[region:note]
-	For most cases you will want to leave the **Collision Mip Level** at 0 and then use 1 or 2 for the **Simple Collision Mip Level**. Using any higher numbers will start to show inaccuracies between the player and collision.
-	[/region]
+[region:note]
+For most cases, you will want to leave the **Collision Mip Level** at 0 and then use 1 or 2 for the **Simple Collision Mip Level**. Using any higher numbers will start to show inaccuracies between the player and collision.
+[/region]
 
 
 ### Setting Collision Mip Level Per Landscape Components
 
-You can also set the Collision Mip Level collision settings for individual Landscape Components which allow you to reduce the Landscape collision complexity even further in non-playable areas of the level. To use this functionality in your project, you will need to do the following.
+You can set the Collision Mip Level collision for individual Landscape Components which enables you to reduce the Landscape collision complexity even further in non-playable areas of the level.
 
-1. From the **Modes** panel, click on the Landscape option and make sure that **Manage** mode is selected.
+
+ To set the Collision Mip Level for an individual component in your project, you will need to do the following:
+
+1. From the **Modes** panel, click on the Landscape option and select **Manage** mode.
 
 	[REGION:lightbox]
 	[![](CC_01.png)(w:524 h:303)](CC_01.png)
@@ -281,7 +285,7 @@ You can also set the Collision Mip Level collision settings for individual Lands
 	Click for full image.
 	[/REGION]
 
-1. In the **Details** panel expands the **Landscape Component** section and change both the **Collision Mip Level** and **Simple Collision Mip Level** to **5**.
+1. In the **Details** panel, expand the **Landscape Component** section and change both the **Collision Mip Level** and **Simple Collision Mip Level** to **5**.
 
 	[REGION:lightbox]
 	[![](CC_03.png)(w:524 h:303)](CC_03.png)
@@ -291,6 +295,31 @@ You can also set the Collision Mip Level collision settings for individual Lands
 	Click for full image.
 	[/REGION]
 
+1. In the Landscape **Manage** section under **Tool Settings**, press the **Clear Component Selection** button to deselect any currently selected Landscape components.  
+
+	![](Clear_Selected_Comps.png)
+
+1. Select a few more Landscape components and this time set both Collision Mip levels to a value of 2.
+
+	[REGION:lightbox]
+	[![](CC_04.png)(w:524 h:303)](CC_04.png)
+	[/REGION]
+	
+	[REGION:caption]
+	Click for full image.
+	[/REGION]
+
+In the following image, the Collision Mip level for each of the four outlined Landscape Component has been set to a different level. 
+
+![](CC_05.png)
+
+|Number| Collision Mip Level|
+|----------|-------------------------|
+|1|3|
+|2|4|
+|3|5|
+|4|2|
+
 
 
 <!-----
@@ -306,7 +335,7 @@ You can also set the Collision Mip Level collision settings for individual Lands
 	VC
 	[/PARAMLITERAL]
 	[PARAM:caption]
-	**Dragging the slider will show what happens when setting the Collision Mip Level from 0 to 5**
+	**Drag the slider to adjust the Collision Mip Level from 0 to 5**
 	[/PARAM]
 	[PARAM:image_list]
 	![](VC_1.png)
@@ -333,7 +362,7 @@ You can also set the Collision Mip Level collision settings for individual Lands
 	PC
 	[/PARAMLITERAL]
 	[PARAM:caption]
-	**Dragging the slider will show what happens when setting the Simple Collision Mip Level from 0 to 5**
+	**Drag the slider to adjust the Simple Collision Mip Level from 0 to 5**
 	[/PARAM]
 	[PARAM:image_list]
 	![](PC_1.png)
@@ -348,41 +377,3 @@ You can also set the Collision Mip Level collision settings for individual Lands
 -->
 
 
-
-
-
-
-
-
-
-
-
-
-
-
-
-
-
-
-
-
-
-
-
-
-
-
-
-
-
-
-
-
-
-
-
-
-
-
-
-
