#UE4 Docs: Applying peer feedback.

diff --git a/Engine/Documentation/Source/Platforms/Mobile/PostProcessEffects/MobilePostProcessEffects.INT.udn b/Engine/Documentation/Source/Platforms/Mobile/PostProcessEffects/MobilePostProcessEffects.INT.udn
--- a/Engine/Documentation/Source/Platforms/Mobile/PostProcessEffects/MobilePostProcessEffects.INT.udn
+++ b/Engine/Documentation/Source/Platforms/Mobile/PostProcessEffects/MobilePostProcessEffects.INT.udn
@@ -23,10 +23,10 @@ To use the Tonemapper Post Processing options on Mobile Devices, you will need t
 
 [OBJECT:ComparisonSlider]
 	[PARAM:before]
-	![Tonemapper On](Tonemapper_On.png)(w:660 h:417)
+	![Tonemapper On](Tonemapper_On.png)
 	[/PARAM]
 	[PARAM:after]
-	![Tonemapper Off](Tonemapper_Off.png)(w:660 h:417)
+	![Tonemapper Off](Tonemapper_Off.png)
 	[/PARAM]
 [/OBJECT]
 
@@ -77,7 +77,7 @@ Click for full image.
 [/REGION]
 
 ## Post Process Settings
-In the following section we will go over what Post Processing settings you can expect to work with **r.tonemapper** enabled as well as which Post Processing settings will work with **r.tonemapper** disabled.
+In the following section we will go over what Post Processing settings you can expect to work on Mobile devices with the tonemapper enabled or disabled.
 
 [REGION:lightbox]
 [![](TM_PostProcess_Settings.png)(w:275 h:355)](TM_PostProcess_Settings.png)
@@ -155,12 +155,6 @@ Click for full image.
 |Dirt Mask Tint |Yes|1|
 |Dirt Mask |Yes|1|
 
-### Light Propagation Volume
-Not currently supported or usable on Mobile Devices.
-
-### Ambient Cubemap
-Not currently supported or usable on Mobile Devices.
-
 ### Auto Exposure
 |Post Process Property| Supported|R.tonemapperfilm Value| 
 |--------------------|-----------|--------------|
@@ -176,16 +170,7 @@ Not currently supported or usable on Mobile Devices.
 |Histogram Log Min|No|n/a|
 |Histogram Log Max|No|n/a|
 
-### Lense Flares
-Not currently supported or usable on Mobile Devices.
-
-### Ambient Occlusion
-Not currently supported or usable on Mobile Devices.
-
-### Global Illumination
-Not currently supported or usable on Mobile Devices.
-
-## Depth Of Field
+### Depth Of Field
 |Post Process Property| Supported|R.tonemapperfilm Value| 
 |--------------------|-----------|--------------|
 |BokehDOF|Yes|0 |
@@ -209,10 +194,7 @@ Not currently supported or usable on Mobile Devices.
 |Sky Distance|No|n/a|
 |Vignette Size|No|n/a|
 
-### Motion Blur
-Not currently supported or usable on Mobile Devices.
-
-## Misc
+### Misc
 |Post Process Property| Supported|R.tonemapperfilm Value| 
 |--------------------|-----------|--------------|
 |Screen Percentage|No|n/a |
@@ -228,8 +210,6 @@ The AA Method supports "None" (the fastest) and "TemporalAA" which provides a sp
 |Quality|Yes|0 & 1|
 |Max Roughness|Yes|0 & 1|
 
-### Blendables
-Not currently supported or usable on Mobile Devices.
 
 ## General Performance Tips
 
