#UE4 Docs: Applying peer review feedback.

diff --git a/Engine/Documentation/Source/Platforms/OSVR/QuickStart/1/OSVRQuickStart_1.INT.udn b/Engine/Documentation/Source/Platforms/OSVR/QuickStart/1/OSVRQuickStart_1.INT.udn
--- a/Engine/Documentation/Source/Platforms/OSVR/QuickStart/1/OSVRQuickStart_1.INT.udn
+++ b/Engine/Documentation/Source/Platforms/OSVR/QuickStart/1/OSVRQuickStart_1.INT.udn
@@ -41,7 +41,7 @@ In this part of the OSVR Quick Start, we will go over what software you need to
 	Click for full image.
 	[/REGION]
 
-1. When the OSVR server launches you will the following command line window being displayed.
+1. When the OSVR server launches, you will the following command line window being displayed.
 
 	![](OSVR_Server_Windows.png)
 
diff --git a/Engine/Documentation/Source/Platforms/OSVR/QuickStart/2/OSVRQuickStart_2.INT.udn b/Engine/Documentation/Source/Platforms/OSVR/QuickStart/2/OSVRQuickStart_2.INT.udn
--- a/Engine/Documentation/Source/Platforms/OSVR/QuickStart/2/OSVRQuickStart_2.INT.udn
+++ b/Engine/Documentation/Source/Platforms/OSVR/QuickStart/2/OSVRQuickStart_2.INT.udn
@@ -15,7 +15,7 @@ tags:OSVR
 
 In this part of the OSVR Quick Start, we will go over how to create a new UE4 project that has the best settings for OSVR development.
 
-1. From the UE4 Launcher make sure that you have UE4 version **4.12** or later downloaded and installed then press the **Launch** button to open the **Unreal Project Browser**.
+1. From the UE4 Launcher, make sure that you have UE4 version **4.12** or later downloaded and installed then press the **Launch** button to open the **Unreal Project Browser**.
 
 	![](Platforms\GoogleVR\QuickStart\1\GVRQS_Launcher_00.png)
 
@@ -23,7 +23,7 @@ In this part of the OSVR Quick Start, we will go over how to create a new UE4 pr
 
 	![](Platforms\GoogleVR\QuickStart\1\GVRQS_New_Project_Options_00.png)
 
-1. Due to the demanding rendering requirements of VR in general, it is best to start out any new UE4 VR project with the following settings to ensure that your project will run at the correct frame rate right from the start. 
+1. Due to the demanding rendering requirements of VR in general, it is best to start out any new UE4 VR project with the following settings to ensure that your project will run at the correct frame rate right from the start: 
 
 	![](Platforms\GoogleVR\QuickStart\1\GVRQS_New_Project_Options_01.png)
 	* **Blank Project**
