#UE4 Docs: minor formatting updates,improved images to make focus more clear

diff --git a/Engine/Documentation/Source/Engine/Landscape/TechnicalGuide/LansdscapeTechnicalGuide.INT.udn b/Engine/Documentation/Source/Engine/Landscape/TechnicalGuide/LansdscapeTechnicalGuide.INT.udn
--- a/Engine/Documentation/Source/Engine/Landscape/TechnicalGuide/LansdscapeTechnicalGuide.INT.udn
+++ b/Engine/Documentation/Source/Engine/Landscape/TechnicalGuide/LansdscapeTechnicalGuide.INT.udn
@@ -143,51 +143,19 @@ For this example, we are using the Landscape Mountains project that can be found
 
 If you select any Landscape Actor that has been placed in a level and go to the **Details** panel, under the **Collision** section you will find two settings called **Collision Mip Level** and **Simple Collision Mip Level**.
 
-
-* **Collision Mip Level:** The Collision Mip Level sets the complexity of the **Complex** collision that is used for the Landscape. Setting Collision Mip Level to **0**, the default,  will give you a very accurate Landscape collision at the expense of memory. Setting this value to **5**, the maximum setting, will make the Landscape collision cheaper, but the accuracy of the collision will suffer.
-	
-	[REGION:fullwidth]
-	[INCLUDE:Engine\Landscape\TechnicalGuide\#VC]
-	[/REGION]
-
-* **Simple Collision Mip Level:** The Simple Collision Mip Level sets the complexity of the **Simple** collision that is used for the Landscape. Setting Simple Collision Mip Level to **0**, the default,  will give you a very accurate Landscape collision at the expense of memory. Setting this value to **5**, the maximum setting, will make the Landscape collision cheaper, but the accuracy of the collision will suffer.
-
-	[REGION:fullwidth]
-	[INCLUDE:Engine\Landscape\TechnicalGuide\#PC]
-	[/REGION]
+| --- | --- |
+| **Collision Mip Level** | The Collision Mip Level sets the complexity of the **Complex** collision that is used for the Landscape. Setting Collision Mip Level to **0**, the default,  will give you a very accurate Landscape collision at the expense of memory. Setting this value to **5**, the maximum setting, will make the Landscape collision cheaper, but the accuracy of the collision will suffer.  [INCLUDE:Engine\Landscape\TechnicalGuide\#VC] |
+| **Simple Collision Mip Level** | The Simple Collision Mip Level sets the complexity of the **Simple** collision that is used for the Landscape. Setting Simple Collision Mip Level to **0**, the default,  will give you a very accurate Landscape collision at the expense of memory. Setting this value to **5**, the maximum setting, will make the Landscape collision cheaper, but the accuracy of the collision will suffer. [INCLUDE:Engine\Landscape\TechnicalGuide\#PC] |
 
 ### Viewing Collision Mip Level 
 
 You can visualize the Landscape collision geometry using the Collision Viewmodes. To enable the collision view mode, go to the **View Mode** menu in the Editor viewport toolbar, under the **Collision** section, select either **Player Collision**  or **Visibility Collision**.
 
-[REGION:lightbox]
-[![](LSC_00.png)(w:524 h:303)](LSC_00.png)
-[/REGION]
-
-[REGION:caption]
-Click for full image.
-[/REGION]
-
-
-* **Player Collision:** The **Player Collision** view mode displays the Simple Collision Mip Level. 
-
-	[REGION:lightbox]
-	[![](CML_Player_Collision.png)(w:524 h:303)](CML_Player_Collision.png)
-	[/REGION]
-
-	[REGION:caption]
-	Click for full image.
-	[/REGION]
-
-* **Visibility Collision:** The Visibility Collision view mode displays the Complex Collision Mip Level.
+![](LSC_00.png)(convert:false)
 
-	[REGION:lightbox]
-	[![](CML_Vis_Collision.png)(w:524 h:303)](CML_Vis_Collision.png)
-	[/REGION]
-
-	[REGION:caption]
-	Click for full image.
-	[/REGION]
+| --- | --- |
+| **Player Collision** | The **Player Collision** view mode displays the Simple Collision Mip Level. ![](CML_Player_Collision.png)(w:640) |
+| **Visibility Collision** | The **Visibility Collision** view mode displays the Collision Mip Level. ![](CML_Vis_Collision.png)(w:640) |
 
 ### Adjusting the Landscape Collision Mip Level
 
@@ -196,25 +164,20 @@ To set the complexity for both simple and complex Landscape collision you will n
 1. Select your Landscape Terrain in the Editor viewport. In the **Details** panel, expand the **Collision** section.
 
 	[REGION:lightbox]
-	[![](LSC_01.png)(w:524 h:303)](LSC_01.png)
+	[![](LSC_01.png)(convert:false)](LSC_01.png)
 	[/REGION]
-
 	[REGION:caption]
 	Click for full image.
 	[/REGION]
 
-1. Under the **Collision** section, find the **Collision Mip Level** option. Set the value  from **0** to **5** and then press the **Enter** key to apply the change. 
-
-	[REGION:caption]
-	The gray collision mesh in the level updates automatically to reflect the changes.
-	[/REGION]
+1. Under the **Collision** section, find the **Collision Mip Level** option. Set the value  from **0** to **5** and then press the **Enter** key to apply the change. The gray collision mesh in the level updates automatically to reflect the changes.
 
 	[OBJECT:ComparisonSlider]
 		[PARAM:before]
-		![Collision Mip Level 0](LSC_03.png)(w:787 h:455)
+		![Collision Mip Level 0](LSC_03.png)
 		[/PARAM]
 		[PARAM:after]
-		![Collision Mip Level 5](LSC_02.png)(w:787 h:455)
+		![Collision Mip Level 5](LSC_02.png)
 		[/PARAM]
 	[/OBJECT]
 
@@ -226,7 +189,7 @@ You can set the complexity of both the simple and complex Landscape collision me
 1. Select your Landscape and in the **Details** panel and expand the **Collision** section.
 
 	[REGION:lightbox]
-	[![](LSC_01.png)(w:524 h:303)](LSC_01.png)
+	[![](LSC_01.png)(convert:false)](LSC_01.png)
 	[/REGION]
 	[REGION:caption]
 	Click for full image.
@@ -234,22 +197,16 @@ You can set the complexity of both the simple and complex Landscape collision me
 
 1. Set the **Collision Mip Level** to a value of **0** and **Simple Collision Mip Level** to a value of **2**.
 
-	[REGION:lightbox]
-	[![](LSC_04.png)(w:524 h:303)](LSC_04.png)
-	[/REGION]
-
-	[REGION:caption]
-	Click for full image.
-	[/REGION]
+	![](LSC_04.png)
 
 In the following image comparison, you can see what happens to the Landscape collision when the Collision Mip Level and Simple Collision Mip Level are set to different values.
 
 [OBJECT:ComparisonSlider]
 	[PARAM:before]
-	![Player Collision 2](Player_Collision.png)(w:719 h:407)
+	![Player Collision|Simple Collision Mip Level = 2](Player_Collision.png)(w:940)
 	[/PARAM]
 	[PARAM:after]
-	![Visibility Collision 0](Vis_Collision.png)(w:719 h:407)
+	![Visibility Collision|Collision Mip Level = 0](Vis_Collision.png)(w:940)
 	[/PARAM]
 [/OBJECT]
 
@@ -267,33 +224,15 @@ You can set the Collision Mip Level collision for individual Landscape Component
 
 1. From the **Modes** panel, click on the Landscape option and select **Manage** mode.
 
-	[REGION:lightbox]
-	[![](CC_01.png)(w:524 h:303)](CC_01.png)
-	[/REGION]
-	
-	[REGION:caption]
-	Click for full image.
-	[/REGION]
+	![](CC_01.png)
 
-1. Select a few Landscape components by clicking on them with the **Left Mouse Button**. You can tell if a Landscape component has been successfully selected as it will turn slightly red in color.
+1. Select a few Landscape components by clicking on them with the **Left Mouse Button**. The selected Landscape Components are highlighted in red.
 
-	[REGION:lightbox]
-	[![](CC_02.png)(w:524 h:303)](CC_02.png)
-	[/REGION]
-	
-	[REGION:caption]
-	Click for full image.
-	[/REGION]
+	![](CC_02.png)
 
 1. In the **Details** panel, expand the **Landscape Component** section and change both the **Collision Mip Level** and **Simple Collision Mip Level** to **5**.
 
-	[REGION:lightbox]
-	[![](CC_03.png)(w:524 h:303)](CC_03.png)
-	[/REGION]
-	
-	[REGION:caption]
-	Click for full image.
-	[/REGION]
+	![](CC_03.png)
 
 1. In the Landscape **Manage** section under **Tool Settings**, press the **Clear Component Selection** button to deselect any currently selected Landscape components.  
 
@@ -301,13 +240,7 @@ You can set the Collision Mip Level collision for individual Landscape Component
 
 1. Select a few more Landscape components and this time set both Collision Mip levels to a value of 2.
 
-	[REGION:lightbox]
-	[![](CC_04.png)(w:524 h:303)](CC_04.png)
-	[/REGION]
-	
-	[REGION:caption]
-	Click for full image.
-	[/REGION]
+	![](CC_04.png)
 
 In the following image, the Collision Mip level for each of the four outlined Landscape Component has been set to a different level. 
 
@@ -338,12 +271,12 @@ In the following image, the Collision Mip level for each of the four outlined La
 	**Drag the slider to adjust the Collision Mip Level from 0 to 5**
 	[/PARAM]
 	[PARAM:image_list]
-	![](VC_1.png)
-	![](VC_2.png)
-	![](VC_3.png)
-	![](VC_4.png)
-	![](VC_5.png)
-	![](VC_6.png)
+	![](VC_1.png)(w:780)
+	![](VC_2.png)(w:780)
+	![](VC_3.png)(w:780)
+	![](VC_4.png)(w:780)
+	![](VC_5.png)(w:780)
+	![](VC_6.png)(w:780)
 	[/PARAM]
 [/OBJECT]
 [/EXCERPT:VC]
@@ -365,12 +298,12 @@ In the following image, the Collision Mip level for each of the four outlined La
 	**Drag the slider to adjust the Simple Collision Mip Level from 0 to 5**
 	[/PARAM]
 	[PARAM:image_list]
-	![](PC_1.png)
-	![](PC_2.png)
-	![](PC_3.png)
-	![](PC_4.png)
-	![](PC_5.png)
-	![](PC_6.png)
+	![](PC_1.png)(w:780)
+	![](PC_2.png)(w:780)
+	![](PC_3.png)(w:780)
+	![](PC_4.png)(w:780)
+	![](PC_5.png)(w:780)
+	![](PC_6.png)(w:780)
 	[/PARAM]
 [/OBJECT]
 [/EXCERPT:PC]
