#UE4 Docs: minor text update

diff --git a/Engine/Documentation/Source/Platforms/Mobile/PostProcessEffects/MobilePostProcessEffects.INT.udn b/Engine/Documentation/Source/Platforms/Mobile/PostProcessEffects/MobilePostProcessEffects.INT.udn
--- a/Engine/Documentation/Source/Platforms/Mobile/PostProcessEffects/MobilePostProcessEffects.INT.udn
+++ b/Engine/Documentation/Source/Platforms/Mobile/PostProcessEffects/MobilePostProcessEffects.INT.udn
@@ -76,8 +76,8 @@ To enable or disable the Tonemapper on a Mobile Device you will need to tap the
 Click for full image.
 [/REGION]
 
-## Post Process Settings
-In the following section we will go over what Post Processing settings you can expect to work on Mobile devices with the tonemapper enabled or disabled.
+## Supported Post Process Effects and Settings on Mobile
+In the following section we will go over what Post Process Effects and settings are supported on Mobile devices, as well as whether they require the tonemapper to be enabled.
 
 [REGION:lightbox]
 [![](TM_PostProcess_Settings.png)(w:275 h:355)](TM_PostProcess_Settings.png)
