#UE4 Docs: minor formatting updates

diff --git a/Engine/Documentation/Source/Engine/Sequencer/HowTo/SequenceRecorder/SequenceRecorder.INT.udn b/Engine/Documentation/Source/Engine/Sequencer/HowTo/SequenceRecorder/SequenceRecorder.INT.udn
--- a/Engine/Documentation/Source/Engine/Sequencer/HowTo/SequenceRecorder/SequenceRecorder.INT.udn
+++ b/Engine/Documentation/Source/Engine/Sequencer/HowTo/SequenceRecorder/SequenceRecorder.INT.udn
@@ -22,14 +22,10 @@ In this example, we will use the Sequence Recorder to record our playable charac
 For this how-to, we are using the **Blueprint Third Person Template** project with **Starter Content Enabled**. 
 [/REGION]
 
-1.	From the **Menu Bar** under **Window**, select **Sequence Recorder**. 
+1.	From the **Menu Bar** under **Window**, select **Sequence Recorder**. The Sequence Recorder window will automatically open. There are some options under **Sequence Recording** which will determine how (and where) the new Level Sequence asset will be saved. You can choose to record Actors that are spawned (such as particle effects, other characters, etc.), determine the naming convention and save location, and change other settings like the duration to record or delay before starting to record. 
 
 	![](RecorderStep1.png)
 	
-	The Sequence Recorder window will automatically open. 
-	There are some options under **Sequence Recording** which will determine how (and where) the new Level Sequence asset will be saved. 
-	You can choose to record Actors that are spawned (such as particle effects, other characters, etc.), determine the naming convention and save location, and change other settings like the duration to record or delay before starting to record. 
-	
 1.	From the Main Editor Toolbar, click the **Play** button to start a Play in Editor session. 
 
 1.	Press **Shift+F1** to gain mouse control once inside the game. 
@@ -38,23 +34,19 @@ For this how-to, we are using the **Blueprint Third Person Template** project wi
 	
 	![](SequenceRecorderAdd.png)
 
-1.	Click on the new recording, then for **Actor to Record** click the drop-down and select **ThirdPersonCharacter**. 
+1.	Click on the new recording, then for **Actor to Record** click the drop-down and select **ThirdPersonCharacter**. This is where we specify which Actor to target recording around before starting the recording process. 
 
 	![](ActorToRecord.png)
 
-	This is where we specify which Actor to target recording around before starting the recording process. 
-
-1.	Click the **Record** button. 
+1.	Click the **Record** button. After 4 seconds (which is the **Record Delay** option under the Sequence Recording section), the record process will start. 
 
 	![](RecordStart.png)
 
-	After 4 seconds (which is the **Record Delay** option under the Sequence Recording section), the record process will start. 
-
 	[REGION:note]
 	When clicking the **Record** button, all Actors in the list that are set to be tracked will also be recorded. 
 	[/REGION]
 
-1.	Use the **WASD** and **Spacebar** keys to move the character around and jump, then press **ESC** when done.
+1.	Use the **WASD** and **Spacebar** keys to move the character around and jump, then press **ESC** when done. Below is the sequence we recorded. However, if you have a character that can attack or perform other actions, feel free to use that.
 
 	[OBJECT:EmbeddedVideo]
 	[PARAMLITERAL:width]
@@ -66,20 +58,16 @@ For this how-to, we are using the **Blueprint Third Person Template** project wi
 	[PARAMLITERAL:videoid]
 	xo5JRkV3_nI
 	[/PARAMLITERAL]
-	[/OBJECT]
+	[/OBJECT] 
 
-	Above is the sequence we recorded. However, if you have a character that can attack or perform other actions, feel free to use that. 
+1.	Inside the **Content Browser** a new folder will be created containing assets related to the recorded sequence. When recording is complete, a **Level Sequence** asset containing your cniematic data as well as an **Animation Sequence** containing your animation data will be created.
 
-1.	Inside the **Content Browser** a new folder will be created containing assets related to the recorded sequence. 
-
-	![](RecordedAssets.png)
-
-	When recording is complete, a **Level Sequence** asset containing your cniematic data as well as an **Animation Sequence** containing your animation data will be created. 
+	![](RecordedAssets.png) 
 	 
 
 ## End Result
 
-You can now open the **RecordedSequence** asset and begin editing it as you would a normal Level Sequence. 
+You can open the **RecordedSequence** asset and begin editing it as you would a normal Level Sequence. Below is our recorded sequence to which we could add cameras and a [**Camera Cuts**](Engine/Sequencer/HowTo/TracksCameraCut) track to provide multiple angles, music, effects or anything else we wanted. We could even take this sequence and embed it in another sequences as part of a [**Shots Track**](Engine/Sequencer/HowTo/TracksShot). 
 
 [OBJECT:EmbeddedVideo]
 [PARAMLITERAL:width]
@@ -93,9 +81,6 @@ hR5zBjYO5Oc
 [/PARAMLITERAL]
 [/OBJECT]
 
-Above is our recorded sequence to which we could add cameras and a [**Camera Cuts**](Engine/Sequencer/HowTo/TracksCameraCut) track to provide multiple angles, music, effects or anything else we wanted. 
-We could even take this sequence and embed it in another sequences as part of a [**Shots Track**](Engine/Sequencer/HowTo/TracksShot). 
-
 
 
 
