#UE4 Docs: minor formatting updates

diff --git a/Engine/Documentation/Source/Engine/Sequencer/HowTo/TracksShot/Shot.INT.udn b/Engine/Documentation/Source/Engine/Sequencer/HowTo/TracksShot/Shot.INT.udn
--- a/Engine/Documentation/Source/Engine/Sequencer/HowTo/TracksShot/Shot.INT.udn
+++ b/Engine/Documentation/Source/Engine/Sequencer/HowTo/TracksShot/Shot.INT.udn
@@ -22,12 +22,10 @@ In this How-to, we will create a multi-camera cinematic using the Shot Track and
 For this how-to, we are using the **Blueprint Third Person Template** project with **Starter Content** enabled. 
 [/REGION]
 
-1.	In the **Content Browser** create a folder and four **Level Sequences** renamed as shown below. 
+1.	In the **Content Browser** create a folder and four **Level Sequences** renamed as shown below. You will need a **Master**, **Shot01**, **Shot02** and **Shot03** respectively. 
 
 	![](Shots_01.png)
 
-	You will need a **Master**, **Shot01**, **Shot02** and **Shot03** respectively. 
-
 1.	Open **Shot01** and click the **Add Camera** button. 
 
 	![](Shots_02.png)(w:540)
@@ -40,33 +38,25 @@ For this how-to, we are using the **Blueprint Third Person Template** project wi
 	New to Unreal Engine? Check out the [**Viewport Controls**](Engine/UI/LevelEditor/Viewports/ViewportControls/) page for navigation methods. 
 	[/REGION]
 
-1.	In Sequencer, add a key for **Transform** at frame 0 and adjust the **Manual Focus Distance** then add a key. 
+1.	In Sequencer, add a key for **Transform** at frame 0 and adjust the **Manual Focus Distance** then add a key. We used a value of **123** for Manual Focus Distance, however feel free to use your desired setting.
 
 	![](Shots_04.png)
-	
-	We used a value of **123** for Manual Focus Distance, however feel free to use your desired setting. 
-
-1.	Scrub to the end of the sequence, move the CineCamera up towards the face of the character, then add a key for its **Transform**.
+ 
+1.	Scrub to the end of the sequence, move the CineCamera up towards the face of the character, then add a key for its **Transform**. This will allow the camera to move up from the feet towards the head as the sequence plays out.
 
 	![](Shots_05.png)
 
-	This will allow the camera to move up from the feet towards the head as the sequence plays out.
-
-1.	Open **Shot02** and repeat the process for adding camera movement. 
-
-	![](Shots_07.png)
+1.	Open **Shot02** and repeat the process for adding camera movement. Here, we placed the camera on the left side of the character at the feet then moved it up over time. You can also adjust the **Manual Focus Distance** and add a key for your change to bring the character into focus.
 
-	Here we placed the camera on the left side of the character at the feet then moved it up over time. You can also adjust the **Manual Focus Distance** and add a key for your change to bring the character into focus. 
+	![](Shots_07.png) 
 
 	[REGION:note]
 	See [**Working with Camera Cuts**](Engine/Sequencer/HowTo/TracksCameraCut) for a refresher on how to add camera movement. 
 	[/REGION]
 
-1.	Open **Shot03** and add camera movement for the shot.
+1.	Open **Shot03** and add camera movement for the shot. Here, we moved the camera to the right side of the character at the feet, adjusted the focus, then moved it up to the head over time.
 
-	![](Shots_08.png)
-
-	Here we moved the camera to the right side of the character at the feet, adjusted the focus, then moved it up to the head over time. 
+	![](Shots_08.png) 
 
 1.	In the **Content Browser** open the **Master** Level Sequence, then add a **Shot Track**. 
 
@@ -80,7 +70,7 @@ For this how-to, we are using the **Blueprint Third Person Template** project wi
 
 	![](Shots_11.png)
 
-1.	**Left-click** and drag a shot into a new position to re-order the sequence. 
+1.	**Left-click** and drag a shot into a new position to re-order the sequence. Below, we play back the original setup then quickly change that by re-arranging the shots and moving them around in the sequence.
 
 	[OBJECT:EmbeddedVideo]
 	[PARAMLITERAL:width]
@@ -92,15 +82,11 @@ For this how-to, we are using the **Blueprint Third Person Template** project wi
 	[PARAMLITERAL:videoid]
 	yUuBSW4CVcw
 	[/PARAMLITERAL]
-	[/OBJECT]
-
-	Above we play back the original setup then quickly change that by re-arranging the shots and moving them around in the sequence. 
-
-1.	**Right-click** on a shot and select the **New Take** option. 
+	[/OBJECT] 
 
-	![](Shots_12.png)
+1.	**Right-click** on a shot and select the **New Take** option. A save confirmation will appear and the shot will automatically be given a save name/location.
 
-	A save confirmation will appear and the shot will automatically be given a save name/location. 
+	![](Shots_12.png) 
 
 	[REGION:note]
 	You can adjust the naming convention of Shots/Takes in the **Project Settings** under **Level Sequences**. 
@@ -112,13 +98,10 @@ For this how-to, we are using the **Blueprint Third Person Template** project wi
 
 	![](Shots_16.png)
 
-1.	**Right-click** on the new take and highlight **Takes**. 
+1.	**Right-click** on the new take and highlight **Takes**. Each of your Takes will appear under the **Takes** menu where you can change out the shot quickly by selecting a different Take. Takes are duplicated Level Sequences, which means that when you alter them you are not altering the original sequence but instead altering the duplicated version. 
 
 	![](Shots_17.png)
 
-	Each of your Takes will appear under the **Takes** menu where you can change out the shot quickly by selecting a different Take.
-	Takes are duplicated Level Sequences, which means that when you alter them you are not altering the original sequence but instead altering the duplicated version. 
-
 	[REGION:note]
 	Currently when creating a Take for the first time, it is listed in the menu as Take 2 and Take 1 (the original) is not listed but still exists in the **Content Browser**. As a workaround for this, you can create a Take (to use as the original) then create another Take which you can then modify (this will allow you to quickly switch between them using the Takes menu). 
 	[/REGION]
@@ -127,7 +110,7 @@ For this how-to, we are using the **Blueprint Third Person Template** project wi
 
 ## End Result
 
-You can now arrange your shots and quickly swap out those shots with a different Take using the **Shots Track**. 
+You can now arrange your shots and quickly swap out those shots with a different Take using the **Shots Track**. The example below shows how you can use the Shot Track to quickly modify the order in which shots are played back as well as swap in and out shots with the Take system. With Takes you have the freedom to experiment in the way your shots play out as those duplicated shots can quickly be reverted by switching to the original take.
 
 [OBJECT:EmbeddedVideo]
 [PARAMLITERAL:width]
@@ -139,9 +122,7 @@ You can now arrange your shots and quickly swap out those shots with a different
 [PARAMLITERAL:videoid]
 wqjoKdhv-OA
 [/PARAMLITERAL]
-[/OBJECT]
-
-The example above shows how you can use the Shot Track to quickly modify the order in which shots are played back as well as swap in and out shots with the Take system. With Takes you have the freedom to experiment in the way your shots play out as those duplicated shots can quickly be reverted by switching to the original take. 
+[/OBJECT] 
 
 
 
