#UE4 Docs

Overlays/Burn Ins (from UMG) in Sequencer
- Applying the Default Burn Ins How-to.
- Applying Custom Burn Ins How-to.

diff --git a/Engine/Documentation/Source/Engine/Sequencer/HowTo/BurnIns/BurnIns.INT.udn b/Engine/Documentation/Source/Engine/Sequencer/HowTo/BurnIns/BurnIns.INT.udn
new file mode 100644
--- /dev/null
+++ b/Engine/Documentation/Source/Engine/Sequencer/HowTo/BurnIns/BurnIns.INT.udn
@@ -0,0 +1,74 @@
+Availability: Docs
+Title:Applying Burn Ins
+Crumbs: %ROOT%, Engine, Engine/Sequencer, Engine/Sequencer/HowTo
+Description:Shows how you can apply overlays from UMG to rendered Sequencer movies.  
+Version: 4.13
+SkillLevel: Intermediate
+parent:Engine/Sequencer/HowTo
+order:1
+checkpoint: editorqs
+Related: Engine/Sequencer/Overview
+Related: Engine\Sequencer\ReferenceEditor\
+
+When creating and rendering out cinematic movies, you may want to include overlays with information about the scene you are viewing such as the name of the shot, the date, time or frame information. 
+These overlays are referred to as **Burn Ins** as they are burned into the movie when it is rendered out. 
+Often times you will see this process as part of real world movie making to provide directors, editors or anyone viewing the scene with contextual information about the scene (some will even be watermarked with copyright information). 
+
+**Sequencer** allows you to apply burn ins to your rendered out movies and provides a default burn in with general information about the shot. 
+In this guide we take a look at how to enable the default burn in options when rendering out your cinematics. 
+
+
+[REGION:note]
+For this how-to we are using the **Blueprint Third Person Template** project and have created a small sample cinematic.
+[/REGION]
+
+
+## Steps
+
+1.	Once you are ready to render out your **Level Sequence**, inside a **Level Sequence** click the **Render Movie** button. 
+
+	![](Overlay01.png)
+
+1.	In the **Render Movie Settings** window, click the **Show Advanced** expand button. 
+
+	![](Overlay02.png)
+
+1.	Under **Burn in Options**, click the **Burn in Class** dropdown and select **DefaultBurnIn**. 
+
+	![](Overlay03.png)
+
+	This will enable the ability to define the burn in settings. 
+
+	![](Overlay04.png)
+
+1.	Leave the options as default, then click **Capture Movie**. 
+
+
+## End Result
+
+When the capturing of your movie is complete, upon playing back your movie you will notice similar overlays applied to the cinematic. 
+
+[OBJECT:EmbeddedVideo]
+[PARAMLITERAL:width]
+640
+[/PARAMLITERAL]
+[PARAMLITERAL:height]
+360
+[/PARAMLITERAL]
+[PARAMLITERAL:videoid]
+DRIykPNZBAs
+[/PARAMLITERAL]
+[/OBJECT]
+
+[REGION:caption]
+You may want to view the video above in full screen or on YouTube to see the overlays more clearly.
+[/REGION]
+
+The Default Burn in applies the name of your Level Sequence across the Top Center followed by the current date. In the Bottom Left, the current Shot Name is applied (or Level Sequence name if no shots are used). 
+The Bottom Center displays the time and frame count for the Master Sequence while the Bottom Right shows the current Shot Frame. 
+
+Each of these options can be changed to display your own text inside the **Settings** section of the **Burn in Options**. Additionally, you can apply Watermarks or make Font changes within the Default Burn In options. 
+If the default information is not enough for you, you can also take a look at [](Engine\Sequencer\HowTo\CustomBurnIns) to include your desired content. 
+
+
+
diff --git a/Engine/Documentation/Source/Engine/Sequencer/HowTo/CustomBurnIns/CustomBurnIns.INT.udn b/Engine/Documentation/Source/Engine/Sequencer/HowTo/CustomBurnIns/CustomBurnIns.INT.udn
new file mode 100644
--- /dev/null
+++ b/Engine/Documentation/Source/Engine/Sequencer/HowTo/CustomBurnIns/CustomBurnIns.INT.udn
@@ -0,0 +1,136 @@
+Availability: Docs
+Title:Applying Custom Burn Ins
+Crumbs: %ROOT%, Engine, Engine/Sequencer, Engine/Sequencer/HowTo
+Description:Shows you to create and apply custom burn ins to rendered movies. 
+Version: 4.13
+SkillLevel: Advanced
+parent:Engine/Sequencer/HowTo
+order:1
+checkpoint: editorqs
+Related: Engine/Sequencer/Overview
+Related: Engine\Sequencer\ReferenceEditor\
+
+When rendering out your cinematics, you may want to provide additional information or functionality that extends past the [Default Burn In](Engine\Sequencer\HowTo\BurnIns) settings that include time, shot and frame information. 
+With a little bit of [Unreal Motion Graphics](Engine/UMG) and [](Engine/Blueprints) you can create your own custom burn ins and overlay functionality. 
+
+In this example we show what type of assets you will need and the settings to use to get started with custom burn ins. 
+
+[REGION:note]
+For this how-to we are using the **Blueprint Third Person Template** project and have created a small sample cinematic.
+[/REGION]
+
+## Steps
+
+1.	In the **Content Browser**, click the **View Options** and **Show Engine Content**. 
+
+	![](Custom01.png)
+
+	This will grant us the ability to see and use Font options that are provided with the Engine Content. 
+
+1.	**Right-click** in the **Content Browser** and create a **Blueprint Class** of the **LevelSequenceBurnInInitSettings** type. 
+
+	![](Custom02.png)
+
+	Name the new Blueprint **CustomSettings** as this will hold all the settings you wish to configure inside the Burn In Options.  
+
+1.	**Right-click** in the **Content Browser** and create a **Blueprint Class** of the **LevelSequenceBurnIn** type. 
+
+	![](Custom03.png)
+
+	Name the new Blueprint **CustomBurnIn** as this will be the [UMG](Engine/UMG/) widget that we will create our custom burn in content in. 
+
+1.	Open the **CustomSettings** Blueprint. 
+
+	![](Custom04.png)
+
+1.	Click the **+** sign next to **Variables** and add the options you want to expose to the Burn In Options. 
+
+	![](Custom05.png)
+
+	Above we have created two variables, one called **CustomText** of the **Text** variable type and another called **CustomFont** of the **Slate Font Info** type (this will allow us to define the text and font style of burn in we create). 
+	We've set both variables to **Editable**, **Compiled** then set their **Default Values** (use any default values you wish). 
+
+1.	Open the **CustomSettings** Blueprint, then arrange your burn in content on the **Designer** tab. 
+
+	[REGION:lightbox]
+	[![](Custom06.png)](Custom06.png)
+	[/REGION]
+
+	[REGION:caption]
+	Click image for full view. 
+	[/REGION]
+
+	Above we added a Text widget to the graph and anchored it to the bottom center of the screen. 
+	We also changed the name of the widget to **DisplayedText** and made it a **Variable** so that we can access it through script. 
+	You can arrange your content here just as you would any UMG Widget Blueprint (apply images, animation, etc.), for this example we are only going to display some custom text. 
+
+1.	On the **Graph** tab, in the **My Blueprint** window click the **Function Override** button and select **Get Settings Class**. 
+
+	![](Custom07.png)
+
+	We are going to point this Blueprint toward our **CustomSettings** Blueprint Class we created that way it knows which settings to use. 
+
+1.	In the graph that opens, click the **Return Value** dropdown and select **CustomSettings**. 
+
+	![](Custom08.png)
+
+1.	Return to the **Event Graph** and **Right-click** in the graph window and add the **Event Set Settings** node. 
+
+1.	Drag off the **In Settings** and **Cast To CustomSettings**. 
+
+	![](Custom09.png)
+
+	This will allow us to access the options we've defined within it. 
+
+1.	**Right-click** on the **As Custom Settings** pin and **Promote to Variable** called **Options**. 
+
+	![](Custom10.png)
+
+1.	Hold **Ctrl** and drag in the **Displayed Text** and **Options** variables from the **My Blueprint** panel. 
+
+1.	Drag off **Displayed Text** and **Set Font**, then drag off **Options** and get **Custom Font** and connect as shown below. 
+
+	![](Custom11.png)
+
+1.	Return to the **Designer** tab, then in the **Details** panel for the **Text** widget, click **Create Binding** next to the Text Content. 
+
+	![](Custom12.png)
+
+	This will allow us to bind a function to update the content (text displayed) for the Text Widget. 
+
+1.	In the Graph that opens, **Ctrl** drag in the **Options** then drag off it to get **Custom Text** and connect as shown. 
+
+	![](Custom13.png)
+
+	This will display whatever we enter as the **CustomText** in our burn in. 
+
+1.	Inside your **Level Sequence** that you want to render with a burn in, click the **Render Movie** button. 
+
+1.	In the **Render Movie Settings** under **Burn in Options**, change the **Burn in Class** to **CustomBurnIn**. 
+
+	![](Custom14.png)
+
+	When doing so, the settings section updates to list our customizable options. 
+
+	![](Custom15.png)
+
+1.	Click the **Capture Movie** button to render out the movie. 
+
+
+## End Result
+
+When the capturing of your movie is complete, upon playing back your movie you will notice the overlays applied to the cinematic. 
+
+[OBJECT:EmbeddedVideo]
+[PARAMLITERAL:width]
+640
+[/PARAMLITERAL]
+[PARAMLITERAL:height]
+360
+[/PARAMLITERAL]
+[PARAMLITERAL:videoid]
+mZzVTyb0r34
+[/PARAMLITERAL]
+[/OBJECT]
+
+While our example allows us to enter custom text and adjust the font style, you can expose various properties inside the **Level Sequence Burn In Init Settings** class to generate your desired effects. Inside your UMG Widget you can then script the functionality you want your burn ins to take, perhaps even animating them upon **Event Construct**.
diff --git a/Engine/Documentation/Source/Engine/Sequencer/HowTo/SequencerHowTo.INT.udn b/Engine/Documentation/Source/Engine/Sequencer/HowTo/SequencerHowTo.INT.udn
--- a/Engine/Documentation/Source/Engine/Sequencer/HowTo/SequencerHowTo.INT.udn
+++ b/Engine/Documentation/Source/Engine/Sequencer/HowTo/SequencerHowTo.INT.udn
@@ -20,7 +20,6 @@ Refer to the table below for more information and click each link to view its co
 [](FANCY:Engine\Sequencer\HowTo\TracksPlayRate "%Engine\Sequencer\HowTo\TracksPlayRate:description%")
 [/REGION]
 
-
 ## Intermediate
 [REGION:wrapper transparent]
 [](FANCY:Engine\Sequencer\HowTo\CharacterAnimation "%Engine\Sequencer\HowTo\CharacterAnimation:description%")
@@ -30,15 +29,22 @@ Refer to the table below for more information and click each link to view its co
 [](FANCY:Engine\Sequencer\HowTo\CameraRigCrane "%Engine\Sequencer\HowTo\CameraRigCrane:description%")
 [](FANCY:Engine\Sequencer\HowTo\CameraRigRail "%Engine\Sequencer\HowTo\CameraRigRail:description%")
 [](FANCY:Engine\Sequencer\HowTo\CineCameraActors "%Engine\Sequencer\HowTo\CineCameraActors:description%")
+[PUBLISH:docs]
+[](FANCY:Engine\Sequencer\HowTo\BurnIns "%Engine\Sequencer\HowTo\BurnIns:description%")
+[/PUBLISH:docs]
 [/REGION]
 
 
+
 ## Advanced
 [REGION:wrapper transparent]
 [](FANCY:Engine\Sequencer\HowTo\TracksSub "%Engine\Sequencer\HowTo\TracksSub:description%")
 [](FANCY:Engine\Sequencer\HowTo\Spawnables "%Engine\Sequencer\HowTo\Spawnables:description%")
 [](FANCY:Engine\Sequencer\HowTo\SequenceRecorder "%Engine\Sequencer\HowTo\SequenceRecorder:description%")
 [](FANCY:Engine\Sequencer\HowTo\TracksLevelVisibility "%Engine\Sequencer\HowTo\TracksLevelVisibility:description%")
+[PUBLISH:docs]
+[](FANCY:Engine\Sequencer\HowTo\CustomBurnIns "%Engine\Sequencer\HowTo\CustomBurnIns:description%")
+[/PUBLISH:docs]
 [/REGION]
 
 
