#UE4 Docs: minor formatting updates

diff --git a/Engine/Documentation/Source/Engine/Sequencer/HowTo/CineCameraActors/CineCameraActor.INT.udn b/Engine/Documentation/Source/Engine/Sequencer/HowTo/CineCameraActors/CineCameraActor.INT.udn
--- a/Engine/Documentation/Source/Engine/Sequencer/HowTo/CineCameraActors/CineCameraActor.INT.udn
+++ b/Engine/Documentation/Source/Engine/Sequencer/HowTo/CineCameraActors/CineCameraActor.INT.udn
@@ -38,94 +38,73 @@ For this how-to, we are using the **Blueprint Third Person Template** project wi
 
 	![](CineCamera3.png)
 
-1.	With the **CineCameraActor** selected, press **S** to add a keyframe for the current **Transform**. 
+1.	With the **CineCameraActor** selected, press **S** to add a keyframe for the current **Transform**. This lets Sequencer know the starting position of the camera and where to start moving from.
 
-	![](CineCamera4.png)
-
-	This lets Sequencer know the starting position of the camera and where to start moving from. 
+	![](CineCamera4.png) 
 
 1.	In Sequencer, move the Timeline to the end of the Level Sequence (frame **150**). 
 
-1.	In the level, move the camera up and forward as shown below then press **S** to add a keyframe at the desired position. 
-
-	![](CineCamera5.png)
-
-	This will give us some camera movement during our scene. 
-
-1.	In the **Details** panel for the **CineCameraActor**, set the **Look at Tracking Settings** below. 
+1.	In the level, move the camera up and forward as shown below then press **S** to add a keyframe at the desired position. This will give us some camera movement during our scene.
 
-	![](CineCamera6.png)
+	![](CineCamera5.png) 
 
-	Above we have used the following settings : 
+1.	In the **Details** panel for the **CineCameraActor**, set the **Look at Tracking Settings** below. We used the following settings : 
 
 	* **Enable Look at Tracking** - check this option as it will allow the camera to track an Actor. 
 	* **Draw Debug Look at Tracking Position** - check this option as it will allow us to see where the camera is looking at. 
 	* **Actor to Track** - set this to **ThirdPersonCharacter** as we want to track this character for our shot. 
-	* **Relative Offset** - set the **Z** to **60** to raise the tracking position slightly from its default tracking position. 
+	* **Relative Offset** - set the **Z** to **60** to raise the tracking position slightly from its default tracking position.
+
+	![](CineCamera6.png) 
 
 	You should see a yellow box near the head of the character indicating where the tracking position is. 
 
 	![](CineCamera6b.png)
 
-1.	Also in the **Details** panel for the **CineCameraActor**, set the **Filmback Settings** to use **Super 16mm**. 
-
-	![](CineCamera7.png)
+1.	Also in the **Details** panel for the **CineCameraActor**, set the **Filmback Settings** to use **Super 16mm**. The **Filmback Settings** provide different formats that you can experiment with to get different aspect ratios.
 
-	The **Filmback Settings** provide different formats that you can experiment with to get different aspect ratios. 
+	![](CineCamera7.png) 
 
-1.	Expand the **Lens Settings** section and set the **Min Focal Length** to **50.0mm** and **Max Focal Length** to **1000.0mm**. 
+1.	Expand the **Lens Settings** section and set the **Min Focal Length** to **50.0mm** and **Max Focal Length** to **1000.0mm**. The Min/Max Focal Length settings will affect the range in which we can apply Focal Length in mm which we will set in a moment. You can select different Lens Settings from the drop-down menu which are some of the standard settings used with real-world cameras today or you can use your own custom settings. 
 
 	![](CineCamera8.png)
 
-	The Min/Max Focal Length settings will affect the range in which we can apply Focal Length in mm which we will set in a moment. 
-	You can select different Lens Settings from the drop-down menu which are some of the standard settings used with real-world cameras today or you can use your own custom settings. 
-
 1.	Expand the **Focus Settings** section and change the **Focus Method** to **Tracking**. 
 
 1.	Next to **Actor to Track**, click the **Eye Dropper** icon. 
 
 	![](CineCamera9.png)
 
-1.	Move the **Eye Dropper** over the character in the level, then **left-click** to select it. 
+1.	Move the **Eye Dropper** over the character in the level, then **left-click** to select it. This will ensure that the focus will always track the Actor we've selected.
 
-	![](CineCamera10.png)
+	![](CineCamera10.png) 
 
-	This will ensure that the focus will always track the Actor we've selected. 
+1.	Move the Timeline back to frame 0, then for **Current Focal Length** add another keyframe. Our Focal Length of **50.0mm** will be used to start the scene, and we will change this value over time.
 
-1.	Move the Timeline back to frame 0, then for **Current Focal Length** add another keyframe. 
+	![](CineCamera11.png) 
 
-	![](CineCamera11.png)
-
-	Our Focal Length of **50.0mm** will be used to start the scene, and we will change this value over time. 
-
-1.	Move the Timeline to the end of the Sequence (frame **150**) and set the **Current Focal Length** to **75.0** and keyframe it. 
+1.	Move the Timeline to the end of the Sequence (frame **150**) and set the **Current Focal Length** to **75.0** and keyframe it. This will cause the camera to zoom in slightly during the course of the shot.
 
 	![](CineCamera12.png)
 
-	This will cause the camera to zoom in slightly during the course of the shot.
-
-1.	In the level select the **ThirdPersonCharacter** then hold **Alt** and drag out a copy and position slightly behind the first. 
+1.	In the level select the **ThirdPersonCharacter** then hold **Alt** and drag out a copy and position slightly behind the first. Our camera will pan around and up to reveal the character standing behind our tracked character.
 
-	![](CineCamera13.png)
-
-	Our camera will pan around and up to reveal the character standing behind our tracked character. 
+	![](CineCamera13.png) 
 
 1.	In Sequencer, click the **Add** button and add a **Camera Cut Track**. 
 
 	![](CineCamera14.png)
 
-1.	Move the Timeline back to **0**, then click the **+ Camera** button and select the **CineCameraActor**. 
+1.	Move the Timeline back to **0**, then click the **+ Camera** button and select the **CineCameraActor**. This will ensure that this sequence will use the camera we have added to shoot the scene from. 
 
 	![](CineCamera15.png)
 
-	This will ensure that this sequence will use the camera we have added to shoot the scene from. 
-
 1.	Click the **Preview Play** button. 
 
 	
 ## End Result
 
-You should have something similar to below where the camera pans up and around to shoot our characters. 
+You should have something similar to below where the camera pans up and around to shoot our characters. Below, we turned off the debug tracking, enabled **Game View** (**G** key in the viewport) and locked the camera to our **Camera Cuts Track**. 
 
 [OBJECT:EmbeddedVideo]
 [PARAMLITERAL:width]
@@ -139,8 +118,6 @@ VlhV5pH_62k
 [/PARAMLITERAL]
 [/OBJECT]
 
-Above we turned off the debug tracking, enabled **Game View** (**G** key in the viewport) and locked the camera to our **Camera Cuts Track**. 
-
 Experiment with each of the settings under **Current Camera Settings** of a **Cine Camera Actor** to produce different looks for your shots! 
 
 
