#UE4 Docs: Fixed include

diff --git a/Engine/Documentation/Source/Engine/Blueprints/Overview/BlueprintOverview.INT.udn b/Engine/Documentation/Source/Engine/Blueprints/Overview/BlueprintOverview.INT.udn
--- a/Engine/Documentation/Source/Engine/Blueprints/Overview/BlueprintOverview.INT.udn
+++ b/Engine/Documentation/Source/Engine/Blueprints/Overview/BlueprintOverview.INT.udn
@@ -49,7 +49,7 @@ order:2
 
 [TOC (start:2)]
 
-[INCLUDE:Engine/Blueprints/GettingStarted/#BlueprintIntro]
+[INCLUDE:Engine/Blueprints/GettingStarted/#Intro]
 
 Does that mean Blueprints are a replacement for UnrealScript? Yes and no. Gameplay programming and everything that UnrealScript 
 was used for in the past can still be handled through code using C++. At the same time, while Blueprints are not meant as a replacement 
