Updating topic image for BP Workflow page

diff --git a/Engine/Documentation/Source/Engine/Blueprints/Workflow/Workflow.INT.udn b/Engine/Documentation/Source/Engine/Blueprints/Workflow/Workflow.INT.udn
--- a/Engine/Documentation/Source/Engine/Blueprints/Workflow/Workflow.INT.udn
+++ b/Engine/Documentation/Source/Engine/Blueprints/Workflow/Workflow.INT.udn
@@ -9,7 +9,7 @@ order:3
 [VAR:Topic]
 [OBJECT:Topic]
 	[PARAM:image]
-		![%Engine/Blueprints/Workflow:title%](Engine/Blueprints/blueprintcheatsheet.png)
+		![%Engine/Blueprints/Workflow:title%](bp_workflow_topic.png)
 	[/PARAM]
 	[PARAM:icon]
 		![](%ROOT%/reference_icon.png)(convert:false)
@@ -28,7 +28,7 @@ order:3
 [VAR:TopicCompact]
 [OBJECT:TopicCompact]
 	[PARAM:image]
-		![%Engine/Blueprints/Workflow:title%](Engine/Blueprints/blueprintcheatsheet.png)
+		![%Engine/Blueprints/Workflow:title%](bp_workflow_topic.png)
 	[/PARAM]
 	[PARAM:icon]
 		![](%ROOT%/reference_icon.png)(convert:false)
