#UE4 Docs: updated metadata (added checkpoint ID)

diff --git a/Engine/Documentation/Source/Engine/Sequencer/QuickStart/SequencerQuickStart.INT.udn b/Engine/Documentation/Source/Engine/Sequencer/QuickStart/SequencerQuickStart.INT.udn
--- a/Engine/Documentation/Source/Engine/Sequencer/QuickStart/SequencerQuickStart.INT.udn
+++ b/Engine/Documentation/Source/Engine/Sequencer/QuickStart/SequencerQuickStart.INT.udn
@@ -7,7 +7,7 @@ parent:Engine/Sequencer
 order:1
 Version: 4.12
 type:quickstart
-
+checkpoint:sequencerqs
 
 [VAR:Topic]
 [OBJECT:Topic]
