#UE4 Docs: Moving the Landscape Collision to it's own page in the docs.

diff --git a/Engine/Documentation/Source/Engine/Landscape/Collision/LandscapeCollision.INT.udn b/Engine/Documentation/Source/Engine/Landscape/Collision/LandscapeCollision.INT.udn
new file mode 100644
--- /dev/null
+++ b/Engine/Documentation/Source/Engine/Landscape/Collision/LandscapeCollision.INT.udn
@@ -0,0 +1,189 @@
+Availability: Public
+Title:Landscape Collision Guide
+Crumbs: %ROOT%, Engine, Engine/Landscape
+Description:Collision settings for Landscape.
+version: 4.12
+
+
+
+
+[TOC(start:2 end:3)]
+
+
+## Landscape Collision 
+
+The Unreal Engine 4 (UE4) Landscape system enables you to specify the detail level of the geometry used for both simple and complex collisions for the entire Landscape as a whole or on a per-Component basis. In the following sections, we will go over how to use the system as well as any relevant information you must know before using this in your UE4 projects.
+
+[region:note]
+For this example, we are using the Landscape Mountains project that can be found in the Learn tab of the UE4 Launcher.
+[/region]
+
+### Collision Mip Level
+
+If you select any Landscape Actor that has been placed in a level and go to the **Details** panel, under the **Collision** section you will find two settings called **Collision Mip Level** and **Simple Collision Mip Level**.
+
+| --- | --- |
+| **Collision Mip Level** | The Collision Mip Level sets the complexity of the **Complex** collision that is used for the Landscape. Setting Collision Mip Level to **0**, the default,  will give you a very accurate Landscape collision at the expense of memory. Setting this value to **5**, the maximum setting, will make the Landscape collision cheaper, but the accuracy of the collision will suffer.  [INCLUDE:Engine\Landscape\\Collision\#VC] |
+| **Simple Collision Mip Level** | The Simple Collision Mip Level sets the complexity of the **Simple** collision that is used for the Landscape. Setting Simple Collision Mip Level to **0**, the default,  will give you a very accurate Landscape collision at the expense of memory. Setting this value to **5**, the maximum setting, will make the Landscape collision cheaper, but the accuracy of the collision will suffer. [INCLUDE:Engine\Landscape\\Collision\#PC] |
+
+### Viewing Collision Mip Level 
+
+You can visualize the Landscape collision geometry using the Collision Viewmodes. To enable the collision view mode, go to the **View Mode** menu in the Editor viewport toolbar, under the **Collision** section, select either **Player Collision**  or **Visibility Collision**.
+
+![](LSC_00.png)(convert:false)
+
+| --- | --- |
+| **Player Collision** | The **Player Collision** view mode displays the Simple Collision Mip Level. ![](CML_Player_Collision.png)(w:640) |
+| **Visibility Collision** | The **Visibility Collision** view mode displays the Collision Mip Level. ![](CML_Vis_Collision.png)(w:640) |
+
+### Adjusting the Landscape Collision Mip Level
+
+To set the complexity for both simple and complex Landscape collision you will need to do the following:
+
+1. Select your Landscape Terrain in the Editor viewport. In the **Details** panel, expand the **Collision** section.
+
+	[REGION:lightbox]
+	[![](LSC_01.png)(convert:false)](LSC_01.png)
+	[/REGION]
+	[REGION:caption]
+	Click for full image.
+	[/REGION]
+
+1. Under the **Collision** section, find the **Collision Mip Level** option. Set the value  from **0** to **5** and then press the **Enter** key to apply the change. The gray collision mesh in the level updates automatically to reflect the changes.
+
+	[OBJECT:ComparisonSlider]
+		[PARAM:before]
+		![Collision Mip Level 0](LSC_03.png)
+		[/PARAM]
+		[PARAM:after]
+		![Collision Mip Level 5](LSC_02.png)
+		[/PARAM]
+	[/OBJECT]
+
+
+### Mixing Collision Mip Level Options 
+
+You can set the complexity of both the simple and complex Landscape collision meshes to provide a better balance between performance and accuracy. To set the simple and complex collision levels independently in your project you will need to do the following:
+
+1. Select your Landscape and in the **Details** panel and expand the **Collision** section.
+
+	[REGION:lightbox]
+	[![](LSC_01.png)(convert:false)](LSC_01.png)
+	[/REGION]
+	[REGION:caption]
+	Click for full image.
+	[/REGION]
+
+1. Set the **Collision Mip Level** to a value of **0** and **Simple Collision Mip Level** to a value of **2**.
+
+	![](LSC_04.png)
+
+In the following image comparison, you can see what happens to the Landscape collision when the Collision Mip Level and Simple Collision Mip Level are set to different values.
+
+[OBJECT:ComparisonSlider]
+	[PARAM:before]
+	![Player Collision|Simple Collision Mip Level = 2](Player_Collision.png)(w:940)
+	[/PARAM]
+	[PARAM:after]
+	![Visibility Collision|Collision Mip Level = 0](Vis_Collision.png)(w:940)
+	[/PARAM]
+[/OBJECT]
+
+[region:note]
+For most cases, you will want to leave the **Collision Mip Level** at 0 and then use 1 or 2 for the **Simple Collision Mip Level**. Using any higher numbers will start to show inaccuracies between the player and collision.
+[/region]
+
+
+### Setting Collision Mip Level Per Landscape Components
+
+You can set the Collision Mip Level collision for individual Landscape Components which enables you to reduce the Landscape collision complexity even further in non-playable areas of the level.
+
+
+ To set the Collision Mip Level for an individual component in your project, you will need to do the following:
+
+1. From the **Modes** panel, click on the Landscape option and select **Manage** mode.
+
+	![](CC_01.png)
+
+1. Select a few Landscape components by clicking on them with the **Left Mouse Button**. The selected Landscape Components are highlighted in red.
+
+	![](CC_02.png)
+
+1. In the **Details** panel, expand the **Landscape Component** section and change both the **Collision Mip Level** and **Simple Collision Mip Level** to **5**.
+
+	![](CC_03.png)
+
+1. In the Landscape **Manage** section under **Tool Settings**, press the **Clear Component Selection** button to deselect any currently selected Landscape components.  
+
+	![](Clear_Selected_Comps.png)
+
+1. Select a few more Landscape components and this time set both Collision Mip levels to a value of 2.
+
+	![](CC_04.png)
+
+In the following image, the Collision Mip level for each of the four outlined Landscape Component has been set to a different level. 
+
+![](CC_05.png)
+
+|Number| Collision Mip Level|
+|----------|-------------------------|
+|1|3|
+|2|4|
+|3|5|
+|4|2|
+
+
+
+<!-----
+[EXCERPT:VC]
+[OBJECT:SliderSequence]
+	[PARAMLITERAL:max]
+	6
+	[/PARAMLITERAL]
+	[PARAMLITERAL:ext]
+	jpg
+	[/PARAMLITERAL]
+	[PARAMLITERAL:id]
+	VC
+	[/PARAMLITERAL]
+	[PARAM:caption]
+	**Drag the slider to adjust the Collision Mip Level from 0 to 5**
+	[/PARAM]
+	[PARAM:image_list]
+	![](VC_1.png)(w:780)
+	![](VC_2.png)(w:780)
+	![](VC_3.png)(w:780)
+	![](VC_4.png)(w:780)
+	![](VC_5.png)(w:780)
+	![](VC_6.png)(w:780)
+	[/PARAM]
+[/OBJECT]
+[/EXCERPT:VC]
+-->
+
+<!-----
+[EXCERPT:PC]
+[OBJECT:SliderSequence]
+	[PARAMLITERAL:max]
+	6
+	[/PARAMLITERAL]
+	[PARAMLITERAL:ext]
+	jpg
+	[/PARAMLITERAL]
+	[PARAMLITERAL:id]
+	PC
+	[/PARAMLITERAL]
+	[PARAM:caption]
+	**Drag the slider to adjust the Simple Collision Mip Level from 0 to 5**
+	[/PARAM]
+	[PARAM:image_list]
+	![](PC_1.png)(w:780)
+	![](PC_2.png)(w:780)
+	![](PC_3.png)(w:780)
+	![](PC_4.png)(w:780)
+	![](PC_5.png)(w:780)
+	![](PC_6.png)(w:780)
+	[/PARAM]
+[/OBJECT]
+[/EXCERPT:PC]
+-->
diff --git a/Engine/Documentation/Source/Engine/Landscape/Landscape.INT.udn b/Engine/Documentation/Source/Engine/Landscape/Landscape.INT.udn
--- a/Engine/Documentation/Source/Engine/Landscape/Landscape.INT.udn
+++ b/Engine/Documentation/Source/Engine/Landscape/Landscape.INT.udn
@@ -115,6 +115,18 @@ order:8
 	[/PARAM]
 [/OBJECT]
 
+[OBJECT:TopicCompact]
+	[PARAM:image]
+		![](Landscape_Collision_Banner.png)
+	[/PARAM]
+	[PARAM:title]
+		%Engine/Landscape/Collision:title%
+	[/PARAM]
+	[PARAM:path]
+		[RELATIVE:Engine/Landscape/Collision]
+	[/PARAM]
+[/OBJECT]
+
 
 
 
diff --git a/Engine/Documentation/Source/Engine/Landscape/TechnicalGuide/LansdscapeTechnicalGuide.INT.udn b/Engine/Documentation/Source/Engine/Landscape/TechnicalGuide/LansdscapeTechnicalGuide.INT.udn
--- a/Engine/Documentation/Source/Engine/Landscape/TechnicalGuide/LansdscapeTechnicalGuide.INT.udn
+++ b/Engine/Documentation/Source/Engine/Landscape/TechnicalGuide/LansdscapeTechnicalGuide.INT.udn
@@ -131,182 +131,4 @@ Selecting a channel will apply a shader to the Landscape which shows the area th
 
 [/EXCERPT:RecommendedLSS]
 
-## Landscape Collision 
-
-Landscapes enable you to specify the detail level of the geometry used for both simple and complex collisions for the entire Landscape or on a per-Component basis. In the following sections, we will go over how to use the system as well as any relevant information you must know before using this in your UE4 projects.
-
-[region:note]
-For this example, we are using the Landscape Mountains project that can be found in the Learn tab of the UE4 Launcher.
-[/region]
-
-### Collision Mip Level
-
-If you select any Landscape Actor that has been placed in a level and go to the **Details** panel, under the **Collision** section you will find two settings called **Collision Mip Level** and **Simple Collision Mip Level**.
-
-| --- | --- |
-| **Collision Mip Level** | The Collision Mip Level sets the complexity of the **Complex** collision that is used for the Landscape. Setting Collision Mip Level to **0**, the default,  will give you a very accurate Landscape collision at the expense of memory. Setting this value to **5**, the maximum setting, will make the Landscape collision cheaper, but the accuracy of the collision will suffer.  [INCLUDE:Engine\Landscape\TechnicalGuide\#VC] |
-| **Simple Collision Mip Level** | The Simple Collision Mip Level sets the complexity of the **Simple** collision that is used for the Landscape. Setting Simple Collision Mip Level to **0**, the default,  will give you a very accurate Landscape collision at the expense of memory. Setting this value to **5**, the maximum setting, will make the Landscape collision cheaper, but the accuracy of the collision will suffer. [INCLUDE:Engine\Landscape\TechnicalGuide\#PC] |
-
-### Viewing Collision Mip Level 
-
-You can visualize the Landscape collision geometry using the Collision Viewmodes. To enable the collision view mode, go to the **View Mode** menu in the Editor viewport toolbar, under the **Collision** section, select either **Player Collision**  or **Visibility Collision**.
-
-![](LSC_00.png)(convert:false)
-
-| --- | --- |
-| **Player Collision** | The **Player Collision** view mode displays the Simple Collision Mip Level. ![](CML_Player_Collision.png)(w:640) |
-| **Visibility Collision** | The **Visibility Collision** view mode displays the Collision Mip Level. ![](CML_Vis_Collision.png)(w:640) |
-
-### Adjusting the Landscape Collision Mip Level
-
-To set the complexity for both simple and complex Landscape collision you will need to do the following:
-
-1. Select your Landscape Terrain in the Editor viewport. In the **Details** panel, expand the **Collision** section.
-
-	[REGION:lightbox]
-	[![](LSC_01.png)(convert:false)](LSC_01.png)
-	[/REGION]
-	[REGION:caption]
-	Click for full image.
-	[/REGION]
-
-1. Under the **Collision** section, find the **Collision Mip Level** option. Set the value  from **0** to **5** and then press the **Enter** key to apply the change. The gray collision mesh in the level updates automatically to reflect the changes.
-
-	[OBJECT:ComparisonSlider]
-		[PARAM:before]
-		![Collision Mip Level 0](LSC_03.png)
-		[/PARAM]
-		[PARAM:after]
-		![Collision Mip Level 5](LSC_02.png)
-		[/PARAM]
-	[/OBJECT]
-
-
-### Mixing Collision Mip Level Options 
-
-You can set the complexity of both the simple and complex Landscape collision meshes to provide a better balance between performance and accuracy. To set the simple and complex collision levels independently in your project you will need to do the following:
-
-1. Select your Landscape and in the **Details** panel and expand the **Collision** section.
-
-	[REGION:lightbox]
-	[![](LSC_01.png)(convert:false)](LSC_01.png)
-	[/REGION]
-	[REGION:caption]
-	Click for full image.
-	[/REGION]
-
-1. Set the **Collision Mip Level** to a value of **0** and **Simple Collision Mip Level** to a value of **2**.
-
-	![](LSC_04.png)
-
-In the following image comparison, you can see what happens to the Landscape collision when the Collision Mip Level and Simple Collision Mip Level are set to different values.
-
-[OBJECT:ComparisonSlider]
-	[PARAM:before]
-	![Player Collision|Simple Collision Mip Level = 2](Player_Collision.png)(w:940)
-	[/PARAM]
-	[PARAM:after]
-	![Visibility Collision|Collision Mip Level = 0](Vis_Collision.png)(w:940)
-	[/PARAM]
-[/OBJECT]
-
-[region:note]
-For most cases, you will want to leave the **Collision Mip Level** at 0 and then use 1 or 2 for the **Simple Collision Mip Level**. Using any higher numbers will start to show inaccuracies between the player and collision.
-[/region]
-
-
-### Setting Collision Mip Level Per Landscape Components
-
-You can set the Collision Mip Level collision for individual Landscape Components which enables you to reduce the Landscape collision complexity even further in non-playable areas of the level.
-
-
- To set the Collision Mip Level for an individual component in your project, you will need to do the following:
-
-1. From the **Modes** panel, click on the Landscape option and select **Manage** mode.
-
-	![](CC_01.png)
-
-1. Select a few Landscape components by clicking on them with the **Left Mouse Button**. The selected Landscape Components are highlighted in red.
-
-	![](CC_02.png)
-
-1. In the **Details** panel, expand the **Landscape Component** section and change both the **Collision Mip Level** and **Simple Collision Mip Level** to **5**.
-
-	![](CC_03.png)
-
-1. In the Landscape **Manage** section under **Tool Settings**, press the **Clear Component Selection** button to deselect any currently selected Landscape components.  
-
-	![](Clear_Selected_Comps.png)
-
-1. Select a few more Landscape components and this time set both Collision Mip levels to a value of 2.
-
-	![](CC_04.png)
-
-In the following image, the Collision Mip level for each of the four outlined Landscape Component has been set to a different level. 
-
-![](CC_05.png)
-
-|Number| Collision Mip Level|
-|----------|-------------------------|
-|1|3|
-|2|4|
-|3|5|
-|4|2|
-
-
-
-<!-----
-[EXCERPT:VC]
-[OBJECT:SliderSequence]
-	[PARAMLITERAL:max]
-	6
-	[/PARAMLITERAL]
-	[PARAMLITERAL:ext]
-	jpg
-	[/PARAMLITERAL]
-	[PARAMLITERAL:id]
-	VC
-	[/PARAMLITERAL]
-	[PARAM:caption]
-	**Drag the slider to adjust the Collision Mip Level from 0 to 5**
-	[/PARAM]
-	[PARAM:image_list]
-	![](VC_1.png)(w:780)
-	![](VC_2.png)(w:780)
-	![](VC_3.png)(w:780)
-	![](VC_4.png)(w:780)
-	![](VC_5.png)(w:780)
-	![](VC_6.png)(w:780)
-	[/PARAM]
-[/OBJECT]
-[/EXCERPT:VC]
--->
-
-<!-----
-[EXCERPT:PC]
-[OBJECT:SliderSequence]
-	[PARAMLITERAL:max]
-	6
-	[/PARAMLITERAL]
-	[PARAMLITERAL:ext]
-	jpg
-	[/PARAMLITERAL]
-	[PARAMLITERAL:id]
-	PC
-	[/PARAMLITERAL]
-	[PARAM:caption]
-	**Drag the slider to adjust the Simple Collision Mip Level from 0 to 5**
-	[/PARAM]
-	[PARAM:image_list]
-	![](PC_1.png)(w:780)
-	![](PC_2.png)(w:780)
-	![](PC_3.png)(w:780)
-	![](PC_4.png)(w:780)
-	![](PC_5.png)(w:780)
-	![](PC_6.png)(w:780)
-	[/PARAM]
-[/OBJECT]
-[/EXCERPT:PC]
--->
-
 
