Incorporating feedback, changing How-To titles.

diff --git a/Engine/Documentation/Source/Engine/LevelStreaming/HowTo/StreamWithBlueprints/StreamBP.INT.udn b/Engine/Documentation/Source/Engine/LevelStreaming/HowTo/StreamWithBlueprints/StreamBP.INT.udn
--- a/Engine/Documentation/Source/Engine/LevelStreaming/HowTo/StreamWithBlueprints/StreamBP.INT.udn
+++ b/Engine/Documentation/Source/Engine/LevelStreaming/HowTo/StreamWithBlueprints/StreamBP.INT.udn
@@ -1,5 +1,5 @@
 Availability: Docs
-Title:Stream Sublevels with Blueprints
+Title:Load and Unload Levels with Blueprints
 Crumbs:%ROOT%, Engine, Engine/UI/LevelEditor
 Description:How to stream levels with a custom streaming Actor created with Blueprints
 parent:Engine/LevelStreaming/HowTo
diff --git a/Engine/Documentation/Source/Engine/LevelStreaming/HowTo/StreamWithC++/StreamC++.INT.udn b/Engine/Documentation/Source/Engine/LevelStreaming/HowTo/StreamWithC++/StreamC++.INT.udn
--- a/Engine/Documentation/Source/Engine/LevelStreaming/HowTo/StreamWithC++/StreamC++.INT.udn
+++ b/Engine/Documentation/Source/Engine/LevelStreaming/HowTo/StreamWithC++/StreamC++.INT.udn
@@ -1,5 +1,5 @@
 Availability: Docs
-Title:Stream Sublevels with C++
+Title:Load and Unload Levels with C++
 Crumbs:%ROOT%, Engine, Engine/UI/LevelEditor
 Description:How to stream levels with a custom streaming Actor created in C++
 parent:Engine/LevelStreaming/HowTo
diff --git a/Engine/Documentation/Source/Engine/LevelStreaming/Overview/Overview.INT.udn b/Engine/Documentation/Source/Engine/LevelStreaming/Overview/Overview.INT.udn
--- a/Engine/Documentation/Source/Engine/LevelStreaming/Overview/Overview.INT.udn
+++ b/Engine/Documentation/Source/Engine/LevelStreaming/Overview/Overview.INT.udn
@@ -53,42 +53,55 @@ tags:Level Streaming
 
 ## The Persistent Level
 
-The first step towards getting your levels to blend together seamlessly is creating the persistent level. This can be thought of as a master level that is primarily used to govern 
-which levels will be streamed in and out. 
+The first step towards getting your Levels to blend together seamlessly is creating the persistent Level. This can be thought of as a master Level that is primarily used to govern 
+which Levels will be streamed in and out. 
+
+## Streaming Levels
+
+Streaming Levels are managed through the [**Levels** window](Engine/Levels/LevelsWindow). They can overlap your persistent Level, or be offset to create larger worlds. Each 
+streaming Level can be set to **Always Loaded** or **Blueprint** for the streaming type. This setting can be toggled in the **Levels** window by right-clicking on the sublevel.
+
+When a streaming Level is set to **Always Loaded**, it is loaded together with the persistent Level. It also becomes visible at the same time as the persistent Level. 
+It ignores streaming volumes assigned to it, as well as any loading or unloading requests sent from Blueprints or C++. This type of sublevel is often used for splitting 
+content that would ordinarily be in the persistent Level into multiple "layers", so artists can collaborate working in parallel without blocking each other. For more information on this workflow,
+see [](Engine/Levels/CollaborateWithSublevels). In game, an **Always Loaded** sublevel is never streamed out, unless the game changes the persistent Level. 
+
+Streaming Levels with the **Blueprint** streaming type can actually be controlled with Level Streaming Volumes, Blueprints, or C++ code. These Levels can be loaded
+or unloaded dynamically. 
 
 
 ## Level Streaming Volumes
 
-Level streaming volumes make controlling level streaming very easy. The idea is simple: load/unload requests for a streaming level are issued based on whether the viewpoint is 
-inside any of the LevelStreamingVolume volumes associated with a level.
+**Level Streaming Volumes** make controlling level streaming very easy. The idea is simple: load/unload requests for a streaming Level are issued based on whether the viewpoint is 
+inside any of the Level Streaming Volumes associated with a Level.
 
-Specifically, level streaming volumes can be used in two ways:
+Specifically, Level Streaming Volumes can be used in two ways:
 
 
-* In the **game**, level streaming volumes cause levels to load when the player is inside the volume and unload when the player is outside the volume.
+* In the **game**, Level Streaming Volumes cause Levels to load when the player's viewpoint is inside the volume and unload when the player's viewpoint is outside the volume.
 
-* In the **editor**, level streaming volumes can be used to preview level streaming by automatically hiding/unhiding levels based on the location of the perspective viewport's camera.
+* In the **editor**, Level Streaming Volumes can be used to preview Level streaming by automatically hiding/unhiding Levels based on the location of the perspective viewport's camera.
 
 
 Volume-based level streaming is simple to use and does not require any scripting, making it an ideal way to control level streaming. Furthermore, volume-based level streaming is 
-much easier to maintain than scripted streaming: when demands on the loading system change, level loading/unloading behavior can be modified by simply resizing the streaming volumes.
+much easier to maintain than scripted streaming: when demands on the loading system change, Level loading/unloading behavior can be modified by simply resizing the streaming volumes.
 
 ### Associating Streaming Volumes With Levels
 
-Volume-based level streaming works as follows: each streaming level can have associated with it a set of volumes of type LevelStreamingVolume. Each frame, the engine iterates over 
-each level and checks to see if the player viewpoint is inside any of the Level Streaming Volumes associated with that level. If the viewpoint is inside at least one 
-LevelStreamingVolume, a request is issued to begin loading that level. If the viewpoint is outside all Level Streaming Volumes, the level is marked for unloading.
+Volume-based level streaming works as follows: each streaming Level can have associated with it a set of Level Streaming Volumes. Each frame, the engine iterates over 
+each Level and checks to see if the player's viewpoint is inside any of the Level Streaming Volumes associated with that Level. If the viewpoint is inside at least one 
+Level Streaming Volume, a request is issued to begin loading that Level. If the viewpoint is outside all Level Streaming Volumes, the Level is marked for unloading.
 
 * [](Engine/LevelStreaming/HowTo/StreamWithVolumes)
 
 ### Important Details
 
-* All Level Streaming Volumes must exist in the persistent level. Level Streaming Volumes that live in other levels cannot be used for level streaming, and will generate warnings 
+* All Level Streaming Volumes must exist in the persistent Level. Level Streaming Volumes that live in other Levels cannot be used for level streaming, and will generate warnings 
 when the map is checked for errors.
 
-* If a level has any streaming volumes associated with it, other methods of streaming the level will not behave correctly.
+* If a Level has any streaming volumes associated with it, other methods of streaming the Level will not behave correctly.
 
-* A single LevelStreamingVolume can affect multiple levels. Similarly, a single level can be affected by multiple Level Streaming Volumes.
+* A single Level Streaming Volume can affect multiple Levels. Similarly, a single Level can be affected by multiple Level Streaming Volumes.
 
 * Volume-based streaming works for split screen. The viewpoints of all local players are considered before any loading/unloading requests are issued.
 
@@ -98,45 +111,41 @@ when the map is checked for errors.
 ### Testing Your Streaming Volume Setup
 
 It is critical that volume-based level streaming be tested in game on the target platform. Streaming in **Play in Editor** mode (PIE) will show where the loads/unloads will happen, 
-but streaming in PIE is not representative of real, in-game loading-unloading. This is because in PIE, the levels are already in memory, and so "loading" a level is simply a matter 
-of unhiding it and is thus instantaneous.
+but streaming in PIE is not representative of real, in-game loading-unloading. This is because in PIE, the Levels are already in memory, and so "loading" a Level is simply a matter 
+of unhiding it instantly.
 
-Running the level in the standalone game on the target platform is critical in making sure your streaming setups work. Note that on some platforms, it can take several seconds to 
-stream in a level. Size your Level Streaming Volumes appropriately, so that the level is loaded by the time the player can reach it. Level loading behaviour can be modified by 
-resizing the Level Streaming Volumes associated with a level. Growing a volume causes associated levels to load sooner and unload later, while shrinking a volume causes later loads 
+Running the Level in the standalone game on the target platform is critical in making sure your streaming setups work. Note that on some platforms, it can take several seconds to 
+stream in a Level. Size your Level Streaming Volumes appropriately, so that the Level is loaded by the time the player can reach it. Level loading behavior can be modified by 
+resizing the Level Streaming Volumes associated with a Level. Growing a volume causes associated Levels to load sooner and unload later, while shrinking a volume causes later loads 
 and earlier unloads.
 
 
 ### Level Streaming Volumes for Previs
 
-Volume-based level streaming can be easily previewed by enabling the "Level Streaming Volume Previs" button on the perspective viewport toolbar. When this mode is enabled, levels 
-will be hidden/shown based on the location of the perspective viewport camera. Note that this hiding/showing will affect the set of visible levels in the Level Browser.
-
-A LevelStreamingVolume can be marked for editor previs only by setting the bEditorPreVisOnly flag on the LevelStreamingVolume. In this manner, volume-based level streaming can 
+A Level Streaming Volume can be marked for editor previsualization only by setting the **Editor Pre Vis Only** flag on the Level Streaming Volume. In this manner, volume-based level streaming can 
 still be used for editor previs while at the same time using another streaming method for in-game streaming.
 
 
-
 ### Cost of Level Streaming Volumes
 
-Each frame, `UWorld::ProcessLevelStreamingVolumes` iterates over each streaming level and for each level, that level will begin loading if any local player is within any of the 
-volumes associated with that level. Likewise, the level will be begin unloading if all local players are outside all volumes.
+Each frame, `UWorld::ProcessLevelStreamingVolumes` iterates over each streaming Level and for each Level, that Level will begin loading if any local player is within any of the 
+volumes associated with that Level. Likewise, the Level will be begin unloading if all local players are outside all volumes.
 
-`UWorld::ProcessLevelStreamingVolumes` exploits coherency in the following manner: for each level, the volume that most recently contained a player is cached.  This cached volume is checked first, so 
-that levels a player is in or returns to are quickly accepted.
+`UWorld::ProcessLevelStreamingVolumes` exploits coherency in the following manner: for each Level, the volume that most recently contained a player is cached.  This cached volume is checked first, so 
+that Levels a player is in or returns to are quickly accepted.
 
-Volumes of any shape are fine, although obviously the fewer the better.  An upper bound on the cost of level streaming volumes can be approximated by the sum of the number of level streaming volumes 
-associated with unloaded levels.
+Volumes of any shape are fine, although obviously the fewer the better.  An upper bound on the cost of Level Streaming Volumes can be approximated by the sum of the number of Level Streaming Volumes 
+associated with unloaded Levels.
 
 Two stats exist under the "Streaming" stats group for monitoring level streaming performance.  The "Streaming Volumes" stat tracks the number of Level Streaming Volumes tested against player viewpoints 
-per frame, while the "Volume Streaming Tick" stat tracks the amount of time spent in UWorld::ProcessLevel Streaming Volumes per frame.
+per frame, while the "Volume Streaming Tick" stat tracks the amount of time spent in `UWorld::ProcessLevelStreamingVolumes` per frame.
 
 
 
 ### Adding Hysteresis to Unloading Requests
 
-A player moving back and forth across a LevelStreamingVolume boundary causes unneeded load/unload requests to be issued. To address this, hysteresis has been added to unloading 
-requests. No hysteresis exists for loading requests, because if a level needs to be loaded, we always want it loaded as soon as possible.
+A player moving back and forth across a Level Streaming Volume boundary causes unneeded load/unload requests to be issued. To address this, hysteresis has been added to unloading 
+requests. No hysteresis exists for loading requests, because if a Level needs to be loaded, we always want it loaded as soon as possible.
 
 The amount of unloading hysteresis can be adjusted by modifying the **Min Time Between Volume Unload Requests** property of the streaming Level in the [Levels window](Engine/Levels/LevelsWindow). 
 The default unloading hysteresis is 2.0 seconds.
