#UE4 Docs: Checking in the new doc on TV Safe Zone.

diff --git a/Engine/Documentation/Source/Platforms/TVSafeZoneDebugging/TVSafeZoneDebugging.INT.udn b/Engine/Documentation/Source/Platforms/TVSafeZoneDebugging/TVSafeZoneDebugging.INT.udn
new file mode 100644
--- /dev/null
+++ b/Engine/Documentation/Source/Platforms/TVSafeZoneDebugging/TVSafeZoneDebugging.INT.udn
@@ -0,0 +1,193 @@
+Availability:Public
+Title:Setting TV Safe Zone Debugging
+Crumbs:%ROOT%, Platforms
+Description:Making sure that UI elements are not too close to the edges of the TV
+Version: 4.12
+
+[TOC(start:2)]
+
+When working on and Unreal Engine 4 project, it is important that UI or text are not too close to the edges of the screen so that it does not become distorted. To help ensure that this does not happen, the TV Safe Zone Debugging tool was created. In the following document we will go over how to use this tool in your UE4 projects:
+
+## What is a TV Safe Zone
+A TV Safe Zone or Safe Area is the term used in television production to describe the area of the television picture that ensures the content can be cleanly viewed by a majority of TV viewers. Generally speaking the most important content for a TV show will be placed in this area to ensure that it will not be distorted when viewed on a wide range of TV sets.
+
+[region:note]
+You can read more about TV Safe Zones / Safe Areas on the following Wikipedia page,
+**[TV Safe area](http://tinyurl.com/zq55mzl)**
+[/region]
+
+## Difference between Title Safe Zone & Action Safe Zone
+UE4 has two different types of TV Safe Zones called, Title Safe Zone and Action Safe Zone which are very similar to one another.
+
+**Title Safe Zone**: The Title Safe Zone is to ensure that UI or Text elements are not cutoff or distorted when viewed. 
+
+**Action Safe Zone**: The Action Safe Zone is used to allow things like characters or enemies to wander partially off screen. 
+
+[region:note]
+The area outside of the Action Safe zone will probably be clipped on displays with overscan, while areas between the Action Safe zone and Title Safe zone may be partially clipped.  It's ok to have some action (e.g., other characters or enemies) occur in that region.  The Action Safe zone is the same size as the Title Safe zone on many platforms.
+[/region]
+
+
+## TV Safe Zone Commands
+In the following section we will go over what commands you can use with the TV Safe Zone as well as the two different ways you can enable and use the TV Safe Zone Debugging in your Unreal Engine 4 (UE4) project.
+
+You can enter the following TV Safe Zone commands right to into the UE4 console.
+
+|Command Name| Value| Description|
+|----------------------|----------------|--------|
+|r.DebugSafeZone|0| Disables all overlays.|
+|r.DebugSafeZone|1| Enables Title safe zone.|
+|r.DebugSafeZone|2| Enables Action safe zone.| 
+|r.DebugSafeZone.OverlayAlpha|0 to 1|Controls how opaque the debug visualization overlay is (Default is 0.3)|
+|r.DebugSafeZone.TitleRatio|0 to 1|Controls the action safe zone margins returned in FDisplayMetrics (Default is 0.9).|
+|r.DebugActionZone.ActionRatio| 0 to 1|Controls the action safe zone margins returned in FDisplayMetrics (Default is 0.9).|
+
+You can also place the TV Safe Zone commands in your projects **DefaultEngine.ini** file under the **[SystemSettings]** section. 
+
+![](SZ_Default_Engine_INI.png)
+
+## Activating the Safe & Action Zones
+
+[region:warning]
+For the ratio of the safe areas to be correctly displayed you **must** run the project in **windowed fullscreen** or **fullscreen**. Failing to do this will result in the safe zone being incorrectly displayed or displayed smaller than expected.
+[/region]
+
+In the following section, we will take a look at what you need to do to activate the Safe and Action Zones.
+
+
+1. From the **Main Toolbar** go to the **Play** section and click on the small white arrow next to the play icon to change the play mode options.
+
+	![](SZ_Play_Mode.png)
+
+1. Once the Play Modes menu has been opened, select the **New Editor Window (PIE)** option so that the game will launch itself in a new stand alone window.
+
+	![](SZ_NEW_PIE.png)
+
+1. Press the Play button to launch the game in its own window and then on the upper right-hand corner of the new game window, click on the **square icon** to maximize the window to full screen. 
+
+	![](SZ_Fullscreen_Window.png)
+
+1. With the game window now full screen, press the **Backtick (`)** key to open the Unreal Console and input either of the following commands to display the needed safe zone, pressing the **Enter** key to apply it.
+
+	|Command Name| Value| Description|
+	|----------------------|----------------|--------|
+	|r.DebugSafeZone|1| Enables Title Safe zone.|
+	|r.DebugSafeZone|2| Enables Action Safe zone.| 
+
+	![](SZ_Enter_Command.png)
+
+1. Once the command has been input, you should now see a slightly transparent red border placed around the entire play window like in the following image.
+
+	[REGION:lightbox]
+	[![](Safe_Zone_00.png)(w:645 h:377)](Safe_Zone_00.png)
+	[/REGION]
+	
+	[REGION:caption]
+	Click for full image.
+	[/REGION]
+
+
+## Adjusting the Size of the Safe & Action Zones
+
+You can adjust the size of either the Safe or Action Zones by entering the following commands into the Unreal Console. 
+
+|Command Name| Value| Description|
+|----------------------|----------------|--------|
+|r.DebugSafeZone.TitleRatio|0 to 1 Range|Controls the Safe Zone margins returned in FDisplayMetrics (Default is 0.9).|
+|r.DebugActionZone.ActionRatio| 0 to 1 Range|Controls the Action Safe Zone margins returned in FDisplayMetrics (Default is 0.9).|
+
+[region:tip]
+It is not possible to have both the Safe and Action Zone active at the same time. To view another zone, you will have to disable the one you are currently viewing and then enable the one you want to view.
+[/region]
+
+In the images below, the command **r.DebugActionZone.ActionRatio** was used to decrease to size of the Action Zone from 0.99 all the way down to 0.1.
+
+[INCLUDE:Platforms\TVSafeZoneDebugging\\#SZ]
+
+## Adjusting the Alpha of the Safe & Action Zone Overlay
+
+You can increase or decrease the opacity of the Safe & Action Zones by entering the following command into the Unreal console. 
+
+|Command Name| Value| Description|
+|----------------------|----------------|--------|
+|r.DebugSafeZone.OverlayAlpha|0 to 1|Controls the opacity of the Safe & Action Zones overlays(Default is 0.3).|
+
+In the images below you can see what happens to the opacity of the Safe & Action Zone overlays when the value of r.DebugSafeZone.OverlayAlpha was set from a value of 1.0 to a value of 0.1.
+
+[INCLUDE:Platforms\TVSafeZoneDebugging\\#SZA]
+
+## TV Safe Zone Troubleshooting
+
+In the following section we will go over some pieces of information you need to be aware of to ensure that you are using the TV Safe Zone tool correctly. 
+
+* Right now the safe areas are defined in terms of the full screen size (the API wasn't really designed for non-consoles) rather than returning the ratios directly to let them be used in windowed mode. This could be changed in the future, but for now you need to be running in **windowed fullscreen** or **fullscreen** for the ratio to be correct.
+
+* The UI isn't currently expected to be able to handle this changing at runtime, so if you change the TitleRatio / ActionRatio via the console while running you'll probably need to switch game modes to see the changes (e.g., go from front-end to in-game or vis-versa, or 'open FrontEndScene' again).
+
+* When using this tool on the PS4, you only need the first command, **r.DebugSafeZone.Mode 1** as the Safe Zone values are driven by the system setting you can pick in the PS4 OS menus.
+
+
+
+
+
+<!-----
+[EXCERPT:SZ]
+[OBJECT:SliderSequence]
+	[PARAMLITERAL:max]
+	10
+	[/PARAMLITERAL]
+	[PARAMLITERAL:ext]
+	jpg
+	[/PARAMLITERAL]
+	[PARAMLITERAL:id]
+	SZ
+	[/PARAMLITERAL]
+	[PARAM:caption]
+	**Dragging the slider will show what happens when adjusting the size of the Safe Zone.**
+	[/PARAM]
+	[PARAM:image_list]
+	![](SZ_1.png)
+	![](SZ_2.png)
+	![](SZ_3.png)
+	![](SZ_4.png)
+	![](SZ_6.png)
+	![](SZ_6.png)
+	![](SZ_7.png)
+	![](SZ_8.png)
+	![](SZ_9.png)
+	![](SZ_10.png)
+	[/PARAM]
+[/OBJECT]
+[/EXCERPT:SZ]
+-->
+
+<!-----
+[EXCERPT:SZA]
+[OBJECT:SliderSequence]
+	[PARAMLITERAL:max]
+	10
+	[/PARAMLITERAL]
+	[PARAMLITERAL:ext]
+	jpg
+	[/PARAMLITERAL]
+	[PARAMLITERAL:id]
+	SZA
+	[/PARAMLITERAL]
+	[PARAM:caption]
+	**Dragging the slider will show what happens when adjusting the opacity of the Overlay Alpha.**
+	[/PARAM]
+	[PARAM:image_list]
+	![](SZA_1.png)
+	![](SZA_2.png)
+	![](SZA_3.png)
+	![](SZA_4.png)
+	![](SZA_5.png)
+	![](SZA_6.png)
+	![](SZA_7.png)
+	![](SZA_8.png)
+	![](SZA_9.png)
+	![](SZA_10.png)
+	[/PARAM]
+[/OBJECT]
+[/EXCERPT:SZA]
+-->
