First pass at including adaptive network update frequency information as well as some background on manually-set update frequency.

diff --git a/Engine/Documentation/Source/Gameplay/Networking/Actors/Properties/Properties.INT.udn b/Engine/Documentation/Source/Gameplay/Networking/Actors/Properties/Properties.INT.udn
--- a/Engine/Documentation/Source/Gameplay/Networking/Actors/Properties/Properties.INT.udn
+++ b/Engine/Documentation/Source/Gameplay/Networking/Actors/Properties/Properties.INT.udn
@@ -42,3 +42,21 @@ In the actor's constructor, make sure you have the bReplicates flag set to true:
 	}
 
 That's about it. The member variable `Owner` will now be synchronized to all connected clients for every copy of this actor type that is currently instantiated (in this case, the base actor class).
+
+## Network Update Optimization
+
+### Data-Driven Network Update Frequency
+
+Actors will observe a maximum update frequency set in their `NetUpdateFrequency` variable. By reducing this variable on less-important or less-frequently-changing Actors, network updates can be made more efficient, potentially leading to smoother play experiences in limited-bandwidth scenarios. Common update frequency values are 10 (updating every 0.1 seconds) for important, unpredictable Actors like player-controlled characters in a shooter, 5 (updating every 0.2 seconds) for slower-moving characters like AI-controlled monsters in cooperative games, or 2 (updating every 0.5 seconds) for background Actors that are not very important to gameplay, but are still synced over the network and/or are controlled by server-side logic and thus need replication.
+
+### Adaptive Network Update Frequency
+
+With Adaptive Network Update Frequency, replicated Actors will increase or reduce their update frequency based on whether or not their updates are meaningful. In this context, a "meaningful" update is any update which initializes the Actor, adds or removes a subobject (i.e. an owned Component), or changes the value of a replicated field on the Actor or any of its subobjects. Each Actor’s possible range of update rates is determined by two variables on the Actor itself: `NetUpdateFrequency` and `MinNetUpdateFrequency`. `NetUpdateFrequency` indicates the maximum number of times per second that the Actor will attempt to update itself, while `MinNetUpdateFrequency` indicates the minimum number of update attempts per second.
+
+#### (Advanced) Update Frequency Decrease Algorithm
+
+During update attempts, Actors determine how long it has been since the most recent meaningful update was sent, and record the new time if they send a meaningful update. If an Actor being considered for update has not sent a meaningful update for more than two seconds, it will start decreasing its update frequency, reaching its minimum frequency after seven seconds without a meaningful update. For example, if an Actor with an update delay between 0.1 and 0.6 seconds had had no meaningful updates for 3 seconds, it would attempt its next update in 0.2 seconds.
+
+#### (Advanced) Update Frequency Increase Algorithm
+
+After sending a meaningful update, an Actor will schedule its next update to happen in 30% less time than the time between the last two meaningful updates, clamped between the minimum and maximum update frequencies. For example, if an Actor went exactly one second between meaningful updates, it would schedule its next update attempt for 0.7 seconds in the future, or as close to that value as its minimum and maximum update frequencies permit. With each successive meaningful update, this calculation will be repeated, quickly reducing the time between updates if an Actor starts to have frequent data or subobject changes.
