Minor text edits to Favorites section

diff --git a/Engine/Documentation/Source/Engine/UI/LevelEditor/Details/LevelEditorDetails.INT.udn b/Engine/Documentation/Source/Engine/UI/LevelEditor/Details/LevelEditorDetails.INT.udn
--- a/Engine/Documentation/Source/Engine/UI/LevelEditor/Details/LevelEditorDetails.INT.udn
+++ b/Engine/Documentation/Source/Engine/UI/LevelEditor/Details/LevelEditorDetails.INT.udn
@@ -72,10 +72,10 @@ If there are properties within an Actor that you frequently change or want quick
 ![](Favorites.png)
 
 [REGION:caption]
-Above we checked two options as favorites indicated by the star icon next to their property name for easy access.
+Above, we checked two options as **Favorites**, indicated by the star icon next to their property name for easy access.
 [/REGION]
 
-You can enable Favorites from inside the **Editor Preferences**. 
+You can enable **Favorites** from inside the **Editor Preferences**. 
 
 ![](EditorPrefs.png)
 
@@ -96,7 +96,7 @@ This will list the property (along with any other marked favorites) under the **
 ![](MarkedFavorites.png)
 
 [REGION:note]
-Some properties may not offer the ability to favorite them, due their complexity of their customization. 
+Some properties may not offer the ability to favorite them, due to the complexity of their customization. 
 [/REGION]
 
 
