Minor text edits to Favorites section.

diff --git a/Engine/Documentation/Source/Engine/UI/LevelEditor/Details/LevelEditorDetails.INT.udn b/Engine/Documentation/Source/Engine/UI/LevelEditor/Details/LevelEditorDetails.INT.udn
--- a/Engine/Documentation/Source/Engine/UI/LevelEditor/Details/LevelEditorDetails.INT.udn
+++ b/Engine/Documentation/Source/Engine/UI/LevelEditor/Details/LevelEditorDetails.INT.udn
@@ -87,7 +87,7 @@ Under **Experimental**, check the **Enable Details Panel Favorites** option.
 You may need to restart the editor for the changes to be applied. 
 [/REGION]
 
-With the option enabled, next to a property you can click the star icon to mark it as a favorite. 
+With the option enabled, you can click the star icon next to a property to mark it as a favorite. 
 
 ![](ClickFavorite.png)
 
