Removed reference to outdated DUMPMOVIE_TILEDSHOT command

UE-33241
#2541

diff --git a/Engine/Documentation/Source/Programming/Basics/CommandLineArguments/CommandLineArguments.INT.udn b/Engine/Documentation/Source/Programming/Basics/CommandLineArguments/CommandLineArguments.INT.udn
--- a/Engine/Documentation/Source/Programming/Basics/CommandLineArguments/CommandLineArguments.INT.udn
+++ b/Engine/Documentation/Source/Programming/Basics/CommandLineArguments/CommandLineArguments.INT.udn
@@ -159,7 +159,6 @@ These arguments can be passed to either the game or the editor, depending on the
 * `NoVSync`: Deactivate the VSYNC via command line
 * `BENCHMARK`: Run game at fixed-step in order to process each frame without skipping any frames. This is useful in conjunction with DUMPMOVIE options.
 * `DUMPMOVIE`: Dump rendered frames to files using current resolution of game.
-* `DUMPMOVIE_TILEDSHOT`: Dump rendered frames to files using the specified resolution multiplier. Uses tiled shots for high resolution frames (e.g. `DUMPMOVIE_TILEDSHOT=4`).
 * `EXEC`: Executes the specified exec file.
 * `FPS`: Set the frames per second for benchmarking.
 * `FULLSCREEN`: Set game to run in fullscreen mode.
