UEDOC-1394 - New how-to added for compiling iOS projects from a Windows machine w/a remote Mac. Legacy page archived w/Docs availability, iOS version information added to device compat page.

diff --git a/Engine/Documentation/Source/Platforms/iOS/GettingStarted/iOSGettingStarted.INT.udn b/Engine/Documentation/Source/Platforms/iOS/GettingStarted/iOSGettingStarted.INT.udn
--- a/Engine/Documentation/Source/Platforms/iOS/GettingStarted/iOSGettingStarted.INT.udn
+++ b/Engine/Documentation/Source/Platforms/iOS/GettingStarted/iOSGettingStarted.INT.udn
@@ -1,4 +1,4 @@
-Availability: Public
+Availability: Docs
 Title: Advanced iOS PC/Mac Workflow
 Crumbs:%ROOT%, Platforms, Platforms/iOS
 Description:Setting up for development for the iOS platform.
diff --git a/Engine/Documentation/Source/Platforms/iOS/Windows/Windows.INT.udn b/Engine/Documentation/Source/Platforms/iOS/Windows/Windows.INT.udn
new file mode 100644
--- /dev/null
+++ b/Engine/Documentation/Source/Platforms/iOS/Windows/Windows.INT.udn
@@ -0,0 +1,77 @@
+Availability: Docs
+Title: Building for iOS on Windows
+Crumbs:%ROOT%, Platforms
+Description:How to remotely compile on a Mac from a Windows machine to create an iOS build. 
+Platform:iOS
+Navigation:platform
+Version: 4.9
+type:howto
+parent:Platforms/iOS
+tags:Mobile
+tags:iOS
+prereq:Platforms/iOS/QuickStart
+
+While building for iOS requires a Mac, using remote compilation means that a whole team that primarily uses Windows machines can build for iOS using one Mac. 
+
+## Mac Setup
+
+1. Enable Remote Compiling on the Mac using [these steps from Apple](https://support.apple.com/kb/PH18726).
+
+1. Install your iOS Developer Certificates to the System keychain. If your certificates are already installed to your login keychain, just drag them to the System keychain and they will be installed there as well.
+
+	[REGION:note]
+	For more on creating iOS Developer Certificates, see [](Platforms/iOS/QuickStart/3).
+	[/REGION]
+
+1. Lastly, make sure the private keys for your certificates are set up to be accessible from any program.  This can be accessed by double-clicking on the key under the Developer Certificate. For more information, see [this Apple Support article](https://support.apple.com/kb/PH18682).
+
+## PC Setup
+
+1. Open your project in the editor, then open **Project Settings**.  
+
+	![](Engine/UI/ProjectSettings/menu_project_settings.png)
+
+1. Click on **iOS** under **Platforms**.
+
+	![](iOSsettings.png)
+
+1. Scroll to the **Build** heading, then expand **Remote Build Options**.
+
+	![](BuildOptions.png)
+
+1. In the **Remote Server Name** field, enter the name of the Mac you are using to build the project.
+
+1. In the **Remote User Name** field, enter the user name that would normally log in to the remote Mac.
+
+1. At this point, you will need to generate the SSH key which will be used to ensure a secure connection to the remote Mac.  Press the Generate SSH Key button and follow the prompts.
+
+### SSH Setup
+
+For the SSH setup, you will primarily follow along with the prompts that are presented in the window that pops up. However, here are some additional steps to help guide you through this process.
+
+1. If your host can't be authenticated, enter "yes" at the prompt.
+
+1. Enter the password which belongs to the user you specified in **Remote User Name**.
+
+1. The prompts will ask you to enter a passphrase.  We recommend not specifying a passphrase so that it can be used without requiring user interaction.
+
+1. You will then be prompted to enter the password for the user name again to complete the process.
+
+At this point, you should successfully have generated your SSH key and are ready to start building iOS games from editor on PC.
+
+## Optional Team Setup
+
+If you wish to pass around this data so that it can be used for all projects in the future, then you will want to specify the **Remote Server Name** and **Remote User Name** in one or more of the `*Engine.ini` files on your computer.  
+
+1. Select the multi-dot button next to the property.  
+
+	![](ConfigFiles.png)
+
+	This will open the **Configuration Editor** 
+	
+1. Set the properties for each individual `.ini` file you wish.  
+
+	![](inifileexample.png)
+
+
+Also, the SSH keys are stored in `../../../Engine/Build/SSHKeys`. Make sure to check that directory into source control to share them with your team.
