iOS version info added to device compat page.

diff --git a/Engine/Documentation/Source/Platforms/iOS/DeviceCompatibility/iOSDeviceCompatibility.INT.udn b/Engine/Documentation/Source/Platforms/iOS/DeviceCompatibility/iOSDeviceCompatibility.INT.udn
--- a/Engine/Documentation/Source/Platforms/iOS/DeviceCompatibility/iOSDeviceCompatibility.INT.udn
+++ b/Engine/Documentation/Source/Platforms/iOS/DeviceCompatibility/iOSDeviceCompatibility.INT.udn
@@ -19,6 +19,8 @@ type:overview
 ![](%ROOT%/icon_error_16x.png)
 [/VAR]
 
+Unreal Engine supports iOS 7.0 and above, and the following charts give information about device-specific compatibility for engine features. 
+
 We have not yet tested all iOS devices but this table will give you an idea what features we expect to work well on each device. We plan to expand and refine this section as new devices are released and we broaden the devices we are testing in house.
 
 These tables reference feature tiers as described on the [](Platforms/Mobile/Performance) page.
