Sequncer How-to (Triggering Level Sequences)

diff --git a/Engine/Documentation/Source/Engine/Sequencer/HowTo/SequencerHowTo.INT.udn b/Engine/Documentation/Source/Engine/Sequencer/HowTo/SequencerHowTo.INT.udn
--- a/Engine/Documentation/Source/Engine/Sequencer/HowTo/SequencerHowTo.INT.udn
+++ b/Engine/Documentation/Source/Engine/Sequencer/HowTo/SequencerHowTo.INT.udn
@@ -29,6 +29,7 @@ Refer to the table below for more information and click each link to view its co
 [](FANCY:Engine\Sequencer\HowTo\CameraRigCrane "%Engine\Sequencer\HowTo\CameraRigCrane:description%")
 [](FANCY:Engine\Sequencer\HowTo\CameraRigRail "%Engine\Sequencer\HowTo\CameraRigRail:description%")
 [](FANCY:Engine\Sequencer\HowTo\CineCameraActors "%Engine\Sequencer\HowTo\CineCameraActors:description%")
+[](FANCY:Engine\Sequencer\HowTo\TriggeringSequences "%Engine\Sequencer\HowTo\TriggeringSequences:description%")
 [PUBLISH:docs]
 [](FANCY:Engine\Sequencer\HowTo\BurnIns "%Engine\Sequencer\HowTo\BurnIns:description%")
 [/PUBLISH:docs]
diff --git a/Engine/Documentation/Source/Engine/Sequencer/HowTo/TriggeringSequences/TriggeringSequences.INT.udn b/Engine/Documentation/Source/Engine/Sequencer/HowTo/TriggeringSequences/TriggeringSequences.INT.udn
new file mode 100644
--- /dev/null
+++ b/Engine/Documentation/Source/Engine/Sequencer/HowTo/TriggeringSequences/TriggeringSequences.INT.udn
@@ -0,0 +1,207 @@
+Availability: Public
+Title:Triggering Sequences from Gameplay
+Crumbs: %ROOT%, Engine, Engine/Sequencer, Engine/Sequencer/HowTo
+Description:Illustrates how to trigger a Sequence from an in-game event. 
+Version: 4.12
+SkillLevel: Intermediate
+parent:Engine/Sequencer/HowTo
+order:1
+checkpoint: editorqs
+Related: Engine/Sequencer/Overview
+Related: Engine\Sequencer\ReferenceEditor\
+
+Once you have created your cinematic sequence, you may want to have it called and played during gameplay as part of a cutscene. 
+Perhaps a character walks into a room and you initiate a cutscene or maybe there is an ending cinematic you want to play when a character kills an enemy. 
+
+In this How-to we use [](Engine/Blueprints) to tell a Level Sequence to play based on an in-game event. 
+We also explore the **Sequence Player** options and how we can use them let our game know that the Level Sequence has ended, at which point we can fire of additional Blueprint script. 
+
+## Steps
+
+[REGION:note]
+For this how-to, we are using the **Blueprint Third Person Template** project with a sample Sequence already created. 
+[/REGION]
+
+1.	If the Level Sequence you want to play is not inside your level, **Drag-and-drop** it from the **Content Browser** into the level. 
+
+	![](Trigger01.png)(w:720)
+
+1.	From the **Modes** panel, drag a **Trigger** into the level and adjust it based on your needs (we used a **Box Trigger** below). 
+
+	![](Trigger02.png)
+
+	When your player enters this trigger, we will tell our Level Sequence to start playing. 
+
+1.	With the **Trigger** still selected, click the **Blueprints** button from the Main Toolbar then select **Open Level Blueprint**. 
+
+	![](Trigger03.png)
+
+1.	**Right-click** in the graph and under **Collision**, select **Add On Actor Begin Overlap**. 
+
+	![](Trigger04.png)
+
+1.	Return to the level and select your **Level Sequence**.
+
+1.  Click the **Blueprints** button from the Main Toolbar and **Open Level Blueprint** again.
+
+1.	**Right-click** and select **Create a Reference to...** your Level Sequence. 
+
+	![](Trigger05.png)
+
+	Above the Level Sequence we want to play is called **CameraFlyBy**. 
+
+1.	**Left-click** and drag off the new node and add the **Get Sequence Player** node. 
+
+	![](Trigger06.png)
+
+1.	Off the **Sequence Player** node, add the **Play** node and connect as shown below. 
+
+	![](Trigger07.png)
+
+	If we were to play our game now, the sequence we created will play when the character enters the trigger. 
+
+	[OBJECT:EmbeddedVideo]
+	[PARAMLITERAL:width]
+	640
+	[/PARAMLITERAL]
+	[PARAMLITERAL:height]
+	360
+	[/PARAMLITERAL]
+	[PARAMLITERAL:videoid]
+	sjsxR2nwcNg
+	[/PARAMLITERAL]
+	[/OBJECT]
+
+	With our current setup, our character can still freely move while our cinematic is playing. 
+	While we have managed to successfully playback our cinematic, there may be times where we want to lock our players in place while the cinematic is playing or to **Set Hidden in Game** so they do not appear. 
+	
+1.	**Right-click** in the **Level Blueprint** and add a **Disable Input** node. 
+
+	![](Trigger08.png)
+
+	If you have your Level Sequence still selected in the level, you may see two options (choose the one that is under **Input** by itself. 
+
+1.	Connect the node following **Play** then drag off **Player Controller** and add a **Get Player Controller** node. 
+
+	![](Trigger09.png)
+
+	This will tell the Player Controller associated with your player to stop receiving input. 
+
+1.	Drag off the **Target** pin of the **Disable Input** node and add a **Get Player Character** node. 
+
+	![](Trigger10.png)
+
+	We have disabled input when the sequence starts however we need to enable input once it is over. 
+
+1.	Drag off the **Sequencer Player** node and add a **Get Length** node. 
+
+	![](Trigger11.png)
+
+1.	**Right-click** in the graph and add a **Set Timer by Event** node. 
+
+	![](Trigger12.png)
+
+1.	Connect the **Set Timer by Event** node following **Disable Input** and feed the **Length** into the **Time** pin. 
+
+	![](Trigger13.png)
+
+	This will allow us to call an event at the end of our Level Sequence. 
+
+1.	**Right-click** in the graph and create a **Custom Event** called **Sequence End**.  
+
+	![](Trigger14.png)
+
+1.	Connect the bind pin to the **Event** pin on the **Set Timer by Event** node. 
+
+	![](Trigger15.png)
+
+	Whatever is connected to the **Sequence End** event will now get fired when the Level Sequence completes. 
+
+1.	Add an **Enable Input** node and connect a **Get Player Character** and **Get Player Controller** node. 
+	
+	[REGION:lightbox]
+	[![](Trigger16.png)](Trigger16.png)
+	[/REGION]
+
+	[REGION:caption]
+	Click for full view. 
+	[/REGION]
+
+	If we play now, our playable character will stop and not receive input while the Level Sequence is active. 
+	Once the sequence completes, the player is able to provide input to the character again. 
+
+	[OBJECT:EmbeddedVideo]
+	[PARAMLITERAL:width]
+	640
+	[/PARAMLITERAL]
+	[PARAMLITERAL:height]
+	360
+	[/PARAMLITERAL]
+	[PARAMLITERAL:videoid]
+	nuw87s4vNro
+	[/PARAMLITERAL]
+	[/OBJECT]
+
+	With our current setup, our cinematic is not re-triggerable. If we want to re-trigger the sequence, we need to add one more node to our existing setup. 
+
+1.	In the **Level Blueprint**, drag off the **Sequence Player** node and add a **Set Playback Position** node before the **Play** node. 
+
+	[REGION:lightbox]
+	[![](Trigger17.png)](Trigger17.png)
+	[/REGION]
+
+	[REGION:caption]
+	Click for full view. 
+	[/REGION]
+
+	With the **Playback Position** defaulting to start **0.0**, the sequence will start from the beginning each time it is played. 
+
+1.	Click the **Play** button to play in the editor. 	
+			
+
+## End Result
+
+[OBJECT:EmbeddedVideo]
+[PARAMLITERAL:width]
+640
+[/PARAMLITERAL]
+[PARAMLITERAL:height]
+360
+[/PARAMLITERAL]
+[PARAMLITERAL:videoid]
+dWCbJb5YJhI
+[/PARAMLITERAL]
+[/OBJECT]
+
+We now have a cinematic sequence that is triggered from gameplay and locks the player in place. Our sequence can also be triggered multiple times and input is granted back to the player after the cinematic plays. 
+
+You can access other information or functions from the **Sequencer Player** node under the **Cinematic** section. 
+
+![](Trigger18.png)
+
+
+
+
+
+
+
+
+
+
+
+
+
+
+
+
+
+
+
+
+
+
+
+
+
+
+
