Adding page for Actor Merging. Needs initial flags at top of the page.

diff --git a/Engine/Documentation/Source/Engine/Actors/Merging/ActorMerging.INT.udn b/Engine/Documentation/Source/Engine/Actors/Merging/ActorMerging.INT.udn
new file mode 100644
--- /dev/null
+++ b/Engine/Documentation/Source/Engine/Actors/Merging/ActorMerging.INT.udn
@@ -0,0 +1,88 @@
+Availability: Docs
+Title:Actor Merging
+Crumbs: %ROOT%, Engine, Engine/Actors, Engine/Actors/Merging
+Description:How to use the Actor Merge tool to combine multiple Static Meshes into one in UE4.
+Version: 4.12
+
+
+
+[TOC(start:2 end:2)]
+
+The Merge Actors tool enables the user to combine multiple Static Meshes into a single new Actor. Materials can also be baked together into a single Material with a series of custom textures using a new UV layout, and the original Static Meshes used can even be completely replaced by the newly merged asset.
+
+[REGION:warning]
+The Merge Actor feature is officially supported in 4.12 and later. In 4.8 to 4.11, it must be activated from the Editor Settings under the Experimental tab. If you are using a version of the engine earlier than 4.12, this feature is considered experimental and is not officially supported.
+[/REGION]
+
+![](MergeActorsWindow.png)
+
+## Actor Merging Worklow
+
+[OBJECT:EmbeddedVideo]
+[PARAMLITERAL:videoid]
+_fTSGN8HGu8
+[/PARAMLITERAL]
+[PARAMLITERAL:width]
+930
+[/PARAMLITERAL]
+[PARAMLITERAL:height]
+568
+[/PARAMLITERAL]
+[/OBJECT]
+
+Merging Static Meshes in your scene is handled like so:
+
+1. Select a number of Static Meshes in your level that your would like to merge. 
+
+2. From the main menubar choose **Window > Developer Tools > Merge Actors.**
+
+3. Choose the appropriate options (explained below).
+
+4. Click the Merge Actors button.
+
+[REGION:tip]If you select fewer than **two** Static Meshes, you will get a warning saying there are insufficient meshes to merge.[/REGION]
+
+## Merge Actors Tool Options
+
+At the top of the Merge Actors window is a list of all the Static Meshes selected in the viewport. Here you can check which meshes will be included in the final merge. 
+
+If you have one or fewer meshes selected, you will see a yellow warning near the bottom of the window stating that you have insufficient mesh components for merging.
+
+
+| Option | Description |
+|--------|-------------|
+|[REGION:tablesection]Mesh Settings[/REGION]|||
+|**Pivot Point at Zero**|If set to true, the pivot point of the final merged actor will be set to the scene's origin <0,0,0>. If set to false, the pivot will be located at the pivot location of the first Static Mesh Actor selected for the merge.
+|**Merge Physics Data**|If true, this will merge together all of the physics primitives into a single physics object.|
+|**Export Specific LOD**| Use this to force the selected meshes to a specific LOD prior to merging.|
+|[REGION:tablesection]Material Settings[/REGION]||
+|**Merge Materials**|If checked, the Materials of all the merged assets will be combined into a single Material with new UV layout. The following subsettings apply only if Merge Materials is checked. |
+| |[INCLUDE:#materialsettings]|
+
+
+<!--
+[EXCERPT:MaterialSettings]
+[REGION:simpletable]
+|----|----|
+|**Texture Size**|This provides access to the X and Y pixel dimensions (resolution) for the final merged Material's textures.|
+|**Texture Sizing Type**| The settings for this property are explained in the table below.|
+| **Normal Map**| If checked, a Normal map will be baked out for the resulting merged Actor.|
+|** Metallic Map** |If checked, a Metallic map will be baked out for the resulting merged Actor. This nullifies the Metallic Constant property.|
+| **Metallic Constant**|Only enabled if Metallic Map is unchecked, this sets a constant value for the Metallic property of the resulting merged Actor's Material.|
+| **Roughness Map** |If checked, a Roughness map will be baked out for the resulting merged Actor. This nullifies the Roughness Constant property.|
+| **Roughness Constant**|Only enabled if Roughness Map is unchecked, this sets a constant value for the Roughness property of the resulting merged Actor's Material.| 
+| **Specular Map** |If checked, a Specular map will be baked out for the resulting merged Actor. This nullifies the Specular Constant property.|
+| **Specular Constant**|Only enabled if Specular Map is unchecked, this sets a constant value for the Specular property of the resulting merged Actor's Material.|
+| **Emissive Map** |If checked, an Emissive map will be baked out for the resulting merged Actor.|
+|**Bake Vertex Data**|Bakes vertex data including vertex color and normals into the new resulting mesh.|
+[/REGION]
+[/EXCERPT:MaterialSettings]
+-->
+
+|[REGION:tablesection]Texture Sizing Type Settings[/REGION]||
+|-----|-----|
+|**Use TextureSize for all material properties**|This uses the resolution set in the Texture Size property for all baked textures.|
+|**Use automatically biased texture sizes based on TextureSize**|This creates a 1:0.5:0.25 bias for all textures such that Normal gets 100% of the resolution set in Texture Size, Diffuse (Base Color) gets 50% of Texture Size, and all other properties get 25% of TextureSize|
+|**Use per property manually overriden texture sizes**|This provides individual number entry fields for explicitly providing the resolution for each map being baked.|
+|**Use Simplygon's automatic texture sizing**|This uses Simplygon's algorithm for texture sizing based on its decimation settings. This should only be used if you are already using Simplygon.|
+
