#UE4 Docs: Fixed the layout of the tonemapper film section to make the commands more visable.

diff --git a/Engine/Documentation/Source/Platforms/Mobile/PostProcessEffects/MobilePostProcessEffects.INT.udn b/Engine/Documentation/Source/Platforms/Mobile/PostProcessEffects/MobilePostProcessEffects.INT.udn
--- a/Engine/Documentation/Source/Platforms/Mobile/PostProcessEffects/MobilePostProcessEffects.INT.udn
+++ b/Engine/Documentation/Source/Platforms/Mobile/PostProcessEffects/MobilePostProcessEffects.INT.udn
@@ -32,7 +32,11 @@ To use the Tonemapper Post Processing options on Mobile Devices, you will need t
 
 
 ### Enabling Tonemapper Film via the Console
-You can enable the Tonemapper in the UE4 Editor by first opening up the console by pressing the **Backtick key(`)** and then entering **r.tonemapperfilm 1** to enable the Tonemapper. To disable the Tonemapper, open up the console and enter **r.tonemapperfilm 0**. 
+You can enable the Tonemapper in the UE4 Editor by first opening up the console by pressing the **Backtick key(`)** and then entering the following command to enable the Tonemapper.
+* **r.tonemapperfilm 1** 
+
+To disable the Tonemapper, open up the console and enter the following command to disable the Tonemapper.
+*  **r.tonemapperfilm 0**. 
 
 [region:note]
 This method will enable the Tonemapper only for this session. If you want the Tonemapper to always be on you will need to add the Tonemapper settings you want to use to the DefaultEngine.INI.
@@ -194,6 +198,7 @@ Click for full image.
 |Sky Distance|No|n/a|
 |Vignette Size|No|n/a|
 
+
 ### Misc
 |Post Process Property| Supported|R.tonemapperfilm Value| 
 |--------------------|-----------|--------------|
