#UE4 Docs: updated metadata

 - use dynamic list for landing page links

diff --git a/Engine/Documentation/Source/Engine/Landscape/Collision/LandscapeCollision.INT.udn b/Engine/Documentation/Source/Engine/Landscape/Collision/LandscapeCollision.INT.udn
--- a/Engine/Documentation/Source/Engine/Landscape/Collision/LandscapeCollision.INT.udn
+++ b/Engine/Documentation/Source/Engine/Landscape/Collision/LandscapeCollision.INT.udn
@@ -3,9 +3,38 @@ Title:Landscape Collision Guide
 Crumbs: %ROOT%, Engine, Engine/Landscape
 Description:Collision settings for Landscape.
 version: 4.12
+parent:Engine/Landscape
+order:7
 
-
-
+[VAR:Topic]
+[OBJECT:Topic]
+	[PARAM:image]
+		![](Engine/Landscape/Landscape_Collision_Banner.png)
+	[/PARAM]
+	[PARAM:title]
+		%Engine/Landscape/Collision:title%
+	[/PARAM]
+	[PARAM:path]
+		[RELATIVE:Engine/Landscape/Collision]
+	[/PARAM]
+[/OBJECT]
+[/VAR]
+[VAR:TopicCompact]
+[OBJECT:TopicCompact]
+	[PARAM:image]
+		![](Engine/Landscape/Landscape_Collision_Banner.png)
+	[/PARAM]
+	[PARAM:title]
+		%Engine/Landscape/Collision:title%
+	[/PARAM]
+	[PARAM:description]
+		%Engine/Landscape/Collision:description%
+	[/PARAM]
+	[PARAM:path]
+		[RELATIVE:Engine/Landscape/Collision]
+	[/PARAM]
+[/OBJECT]
+[/VAR]
 
 [TOC(start:2 end:3)]
 
diff --git a/Engine/Documentation/Source/Engine/Landscape/Creation/Creation.INT.udn b/Engine/Documentation/Source/Engine/Landscape/Creation/Creation.INT.udn
--- a/Engine/Documentation/Source/Engine/Landscape/Creation/Creation.INT.udn
+++ b/Engine/Documentation/Source/Engine/Landscape/Creation/Creation.INT.udn
@@ -3,11 +3,32 @@ Title:Creating Landscapes
 Crumbs: %ROOT%, Engine, Engine/Landscape
 Description:Guide to creating new Landscape terrains.
 version: 4.9
+parent:Engine/Landscape
+order:2
 
 [VAR:Topic]
-[OBJECT:Topiccompact]
+[OBJECT:Topic]
 	[PARAM:image]
-		![%Engine/Landscape/Creation:title%](Engine/Landscape/Creation/landscape_create_topic.png)
+		![%Engine/Landscape/Creation:title%](landscape_create_topic.png)
+	[/PARAM]
+	[PARAM:icon]
+		![](%ROOT%/landscape_icon.png)(convert:false)
+	[/PARAM]
+	[PARAM:title]
+		%Engine/Landscape/Creation:title%
+	[/PARAM]
+	[PARAM:description]
+		%Engine/Landscape/Creation:description%
+	[/PARAM]
+	[PARAM:path]
+		[RELATIVE:Engine/Landscape/Creation]
+	[/PARAM]
+[/OBJECT]
+[/VAR]
+[VAR:TopicCompact]
+[OBJECT:TopicCompact]
+	[PARAM:image]
+		![%Engine/Landscape/Creation:title%](landscape_create_topic.png)
 	[/PARAM]
 	[PARAM:icon]
 		![](%ROOT%/landscape_icon.png)(convert:false)
diff --git a/Engine/Documentation/Source/Engine/Landscape/Custom/CreatingCustomHeightmapsAndLayers.INT.udn b/Engine/Documentation/Source/Engine/Landscape/Custom/CreatingCustomHeightmapsAndLayers.INT.udn
--- a/Engine/Documentation/Source/Engine/Landscape/Custom/CreatingCustomHeightmapsAndLayers.INT.udn
+++ b/Engine/Documentation/Source/Engine/Landscape/Custom/CreatingCustomHeightmapsAndLayers.INT.udn
@@ -3,6 +3,7 @@ Title:Creating and Using Custom Heightmaps and Layers
 Crumbs: %ROOT%, Engine, Engine/Landscape
 Description:Creating and using custom Heightmaps and layers for your landscapes.
 version: 4.9
+parent:Engine/Landscape/Creation
 
 
 [TOC(start:2 end:3)]
diff --git a/Engine/Documentation/Source/Engine/Landscape/Editing/Editing.INT.udn b/Engine/Documentation/Source/Engine/Landscape/Editing/Editing.INT.udn
--- a/Engine/Documentation/Source/Engine/Landscape/Editing/Editing.INT.udn
+++ b/Engine/Documentation/Source/Engine/Landscape/Editing/Editing.INT.udn
@@ -3,9 +3,30 @@ Title:Editing Landscapes
 Crumbs: %ROOT%, Engine, Engine/Landscape
 Description:Guide to using the tools for editing Landscape terrains.
 Version: 4.9
+parent:Engine/Landscape
+order:3
 
 [VAR:Topic]
-[OBJECT:Topiccompact]
+[OBJECT:Topic]
+	[PARAM:image]
+		![%Engine/Landscape/Editing:title%](Engine/Landscape/Editing/landscape_editing_topic.png)
+	[/PARAM]
+	[PARAM:icon]
+		![](%ROOT%/landscape_icon.png)(convert:false)
+	[/PARAM]
+	[PARAM:title]
+		%Engine/Landscape/Editing:title%
+	[/PARAM]
+	[PARAM:description]
+		%Engine/Landscape/Editing:description%
+	[/PARAM]
+	[PARAM:path]
+		[RELATIVE:Engine/Landscape/Editing]
+	[/PARAM]
+[/OBJECT]
+[/VAR]
+[VAR:TopicCompact]
+[OBJECT:TopicCompact]
 	[PARAM:image]
 		![%Engine/Landscape/Editing:title%](Engine/Landscape/Editing/landscape_editing_topic.png)
 	[/PARAM]
diff --git a/Engine/Documentation/Source/Engine/Landscape/Landscape.INT.udn b/Engine/Documentation/Source/Engine/Landscape/Landscape.INT.udn
--- a/Engine/Documentation/Source/Engine/Landscape/Landscape.INT.udn
+++ b/Engine/Documentation/Source/Engine/Landscape/Landscape.INT.udn
@@ -52,85 +52,7 @@ order:8
 
 [EXCERPT:landscape002]
 
-[REGION:topics third]
-%Engine/Landscape/QuickStart:topiccompact%
-
-[OBJECT:TopicCompact]
-	[PARAM:image]
-		![%Engine/Landscape/Creation:title%](Engine/Landscape/Creation/landscape_create_topic.png)
-	[/PARAM]
-	[PARAM:title]
-		%Engine/Landscape/Creation:title%
-	[/PARAM]
-	[PARAM:path]
-		[RELATIVE:Engine/Landscape/Creation]
-	[/PARAM]
-[/OBJECT]
-
-[OBJECT:TopicCompact]
-	[PARAM:image]
-		![%Engine/Landscape/Editing:title%](Engine/Landscape/Editing/landscape_editing_topic.png)
-	[/PARAM]
-	[PARAM:title]
-		%Engine/Landscape/Editing:title%
-	[/PARAM]
-	[PARAM:path]
-		[RELATIVE:Engine/Landscape/Editing]
-	[/PARAM]
-[/OBJECT]
-
-[OBJECT:TopicCompact]
-	[PARAM:image]
-		![%Engine/Landscape/UI:title%](Engine/Landscape/UI/Landscape_UI_topic.png)
-	[/PARAM]
-	[PARAM:title]
-		%Engine/Landscape/UI:title%
-	[/PARAM]
-	[PARAM:path]
-		[RELATIVE:Engine/Landscape/UI]
-	[/PARAM]
-[/OBJECT]
-
-[OBJECT:TopicCompact]
-	[PARAM:image]
-		![%Engine/Landscape/Materials:title%](Engine/Landscape/Materials/landscape_materials_topic.png)
-	[/PARAM]
-	[PARAM:title]
-		%Engine/Landscape/Materials:title%
-	[/PARAM]
-	[PARAM:path]
-		[RELATIVE:Engine/Landscape/Materials]
-	[/PARAM]
-[/OBJECT]
-
-[OBJECT:TopicCompact]
-	[PARAM:image]
-		![%Engine/Landscape/TechnicalGuide:title%](Engine/Landscape/Materials/LS_Technical.png)
-	[/PARAM]
-	[PARAM:title]
-		%Engine/Landscape/TechnicalGuide:title%
-	[/PARAM]
-	[PARAM:path]
-		[RELATIVE:Engine/Landscape/TechnicalGuide]
-	[/PARAM]
-[/OBJECT]
-
-[OBJECT:TopicCompact]
-	[PARAM:image]
-		![](Landscape_Collision_Banner.png)
-	[/PARAM]
-	[PARAM:title]
-		%Engine/Landscape/Collision:title%
-	[/PARAM]
-	[PARAM:path]
-		[RELATIVE:Engine/Landscape/Collision]
-	[/PARAM]
-[/OBJECT]
-
-
-
-
-[/REGION]
+[DIR(output:"topiccompact" parent:"Engine/Landscape")]
 
 [/EXCERPT:landscape002]
 
diff --git a/Engine/Documentation/Source/Engine/Landscape/Materials/Materials.INT.udn b/Engine/Documentation/Source/Engine/Landscape/Materials/Materials.INT.udn
--- a/Engine/Documentation/Source/Engine/Landscape/Materials/Materials.INT.udn
+++ b/Engine/Documentation/Source/Engine/Landscape/Materials/Materials.INT.udn
@@ -3,9 +3,30 @@ Title:Landscape Materials
 Crumbs: %ROOT%, Engine, Engine/Landscape
 Description:Setting up materials to work with Landscape terrains.
 version: 4.9
+parent:Engine/Landscape
+order:5
 
 [VAR:Topic]
-[OBJECT:Topiccompact]
+[OBJECT:Topic]
+	[PARAM:image]
+		![%Engine/Landscape/Materials:title%](Engine/Landscape/Materials/landscape_materials_topic.png)
+	[/PARAM]
+	[PARAM:icon]
+		![](%ROOT%/landscape_icon.png)(convert:false)
+	[/PARAM]
+	[PARAM:title]
+		%Engine/Landscape/Materials:title%
+	[/PARAM]
+	[PARAM:description]
+		%Engine/Landscape/Materials:description%
+	[/PARAM]
+	[PARAMLITERAL:path]
+		[RELATIVE:Engine/Landscape/Materials]
+	[/PARAMLITERAL]
+[/OBJECT]
+[/VAR]
+[VAR:TopicCompact]
+[OBJECT:TopicCompact]
 	[PARAM:image]
 		![%Engine/Landscape/Materials:title%](Engine/Landscape/Materials/landscape_materials_topic.png)
 	[/PARAM]
diff --git a/Engine/Documentation/Source/Engine/Landscape/QuickStart/LandscapeQuickStart.INT.udn b/Engine/Documentation/Source/Engine/Landscape/QuickStart/LandscapeQuickStart.INT.udn
--- a/Engine/Documentation/Source/Engine/Landscape/QuickStart/LandscapeQuickStart.INT.udn
+++ b/Engine/Documentation/Source/Engine/Landscape/QuickStart/LandscapeQuickStart.INT.udn
@@ -3,6 +3,8 @@ Title: Landscape Quick Start Guide
 Crumbs: %ROOT%, GettingStarted
 Description: Getting up and running with the basics of the Landscape System in Unreal Engine.
 version: 4.9
+parent:Engine/Landscape
+order:1
 
 [VAR:Topic]
 [OBJECT:Topic]
diff --git a/Engine/Documentation/Source/Engine/Landscape/TechnicalGuide/LansdscapeTechnicalGuide.INT.udn b/Engine/Documentation/Source/Engine/Landscape/TechnicalGuide/LansdscapeTechnicalGuide.INT.udn
--- a/Engine/Documentation/Source/Engine/Landscape/TechnicalGuide/LansdscapeTechnicalGuide.INT.udn
+++ b/Engine/Documentation/Source/Engine/Landscape/TechnicalGuide/LansdscapeTechnicalGuide.INT.udn
@@ -3,7 +3,28 @@ Title:Landscape Technical Guide
 Crumbs: %ROOT%, Engine, Engine/Landscape
 Description:Technical Settings for Landscape.
 version: 4.9
+parent:Engine/Landscape
+order:6
 
+[VAR:Topic]
+[OBJECT:Topic]
+	[PARAM:image]
+		![%Engine/Landscape/TechnicalGuide:title%](Engine/Landscape/Materials/LS_Technical.png)
+	[/PARAM]
+	[PARAM:icon]
+		![](%ROOT%/landscape_icon.png)(convert:false)
+	[/PARAM]
+	[PARAM:title]
+		%Engine/Landscape/TechnicalGuide:title%
+	[/PARAM]
+	[PARAM:description]
+		%Engine/Landscape/TechnicalGuide:description%
+	[/PARAM]
+	[PARAMLITERAL:path]
+		[RELATIVE:Engine/Landscape/TechnicalGuide]
+	[/PARAMLITERAL]
+[/OBJECT]
+[/VAR]
 [VAR:TopicCompact]
 [OBJECT:TopicCompact]
 	[PARAM:image]
diff --git a/Engine/Documentation/Source/Engine/Landscape/UI/InterfaceReference.INT.udn b/Engine/Documentation/Source/Engine/Landscape/UI/InterfaceReference.INT.udn
--- a/Engine/Documentation/Source/Engine/Landscape/UI/InterfaceReference.INT.udn
+++ b/Engine/Documentation/Source/Engine/Landscape/UI/InterfaceReference.INT.udn
@@ -3,8 +3,29 @@ Title:Landscape UI Reference
 Crumbs:%ROOT%, Engine, Engine/Landscape
 Description:Description of the Landscape tool's user interface.
 version: 4.9
+parent:Engine/Landscape
+order:4
 
 [VAR:Topic]
+[OBJECT:Topic]
+	[PARAM:image]
+		![](Landscape_UI_topic.png)
+	[/PARAM]
+	[PARAM:icon]
+		![](%ROOT%/landscape_icon.png)(convert:false)
+	[/PARAM]
+	[PARAM:title]
+		%Engine/Landscape/UI:title%
+	[/PARAM]
+	[PARAM:description]
+		%Engine/Landscape/UI:description%
+	[/PARAM]
+	[PARAM:path]
+		[RELATIVE:Engine/Landscape/UI]
+	[/PARAM]
+[/OBJECT]
+[/VAR]
+[VAR:Topiccompact]
 [OBJECT:Topiccompact]
 	[PARAM:image]
 		![](Landscape_UI_topic.png)
