#UE4 Docs: set pages to publish

diff --git a/Engine/Documentation/Source/Engine/LevelStreaming/LevelStreaming.INT.udn b/Engine/Documentation/Source/Engine/LevelStreaming/LevelStreaming.INT.udn
--- a/Engine/Documentation/Source/Engine/LevelStreaming/LevelStreaming.INT.udn
+++ b/Engine/Documentation/Source/Engine/LevelStreaming/LevelStreaming.INT.udn
@@ -2,6 +2,7 @@ Availability: Public
 Title:Level Streaming
 Crumbs:%ROOT%, Engine, Engine/UI/LevelEditor
 Description:Asynchronously loading and unloading levels during play to decrease memory usage and create seamless worlds.
+Version: 4.12
 parent:Engine
 order:10
 tags:Level Streaming
@@ -46,15 +47,15 @@ tags:Level Streaming
 [/VAR]
 
 [EXCERPT:Intro]
-The Level Streaming feature makes it possible to load and unload map files into memory as well as toggle their visibility during play. This makes it possible to have worlds broken up into 
-smaller chunks so that only the relevant parts of the world are taking up resources and being rendered at any point. If done properly, this allows for the creation of very large, seamless 
-games that can make the player feel as if they are playing within a world that dwarfs them in size.
+The Level Streaming feature makes it possible to load and unload map files into memory as well as toggle their visibility during play. This makes it possible to have worlds broken up 
+into smaller chunks so that only the relevant parts of the world are taking up resources and being rendered at any point. If done properly, this allows for the creation of very 
+large, seamless games that can make the player feel as if they are playing within a world that dwarfs them in size.
 [/EXCERPT:Intro]
 
-[DIR(output:"listbutton" parent:"Engine/LevelStreaming" tags:"Level Streaming" org:"hierarchy" end:0)]
+[DIR(output:"listbutton" parent:"Engine/LevelStreaming" tags:"Level Streaming" org:"hierarchy")]
 
 ## World Composition
 
 World Composition is a specific form of level streaming used to create large worlds. Levels are arranged in a planar grid, and streamed in as the player approaches them.
 
-[DIR(output:"listbutton" parent:"Engine/LevelStreaming" tags:"World Composition")]
\ No newline at end of file
+[DIR(output:"listbutton" parent:"Engine/LevelStreaming" tags:"World Composition" org:"hierarchy")]
\ No newline at end of file
diff --git a/Engine/Documentation/Source/Engine/LevelStreaming/WorldBrowser/WorldBrowser.INT.udn b/Engine/Documentation/Source/Engine/LevelStreaming/WorldBrowser/WorldBrowser.INT.udn
--- a/Engine/Documentation/Source/Engine/LevelStreaming/WorldBrowser/WorldBrowser.INT.udn
+++ b/Engine/Documentation/Source/Engine/LevelStreaming/WorldBrowser/WorldBrowser.INT.udn
@@ -5,7 +5,7 @@ Description:System for managing large worlds including origin shifting technolog
 SkillLevel: Intermediate
 Version: 4.9
 Parent:Engine/LevelStreaming
-Label:World Composition
+tags:World Composition
 
 [TOC (start:2 end:3)]
 
diff --git a/Engine/Documentation/Source/Engine/Levels/HowTo/ChangeDefaultLevel/ChangeDefaultLevel.INT.udn b/Engine/Documentation/Source/Engine/Levels/HowTo/ChangeDefaultLevel/ChangeDefaultLevel.INT.udn
--- a/Engine/Documentation/Source/Engine/Levels/HowTo/ChangeDefaultLevel/ChangeDefaultLevel.INT.udn
+++ b/Engine/Documentation/Source/Engine/Levels/HowTo/ChangeDefaultLevel/ChangeDefaultLevel.INT.udn
@@ -1,4 +1,4 @@
-Availability: Docs
+Availability: Public
 Title:Change Default Levels
 Crumbs: %ROOT%, Engine
 Description:How to set the default editor and game Levels for your project
diff --git a/Engine/Documentation/Source/Engine/Levels/HowTo/WorkWithLevelAssets/WorkWithLevelAssets.INT.udn b/Engine/Documentation/Source/Engine/Levels/HowTo/WorkWithLevelAssets/WorkWithLevelAssets.INT.udn
--- a/Engine/Documentation/Source/Engine/Levels/HowTo/WorkWithLevelAssets/WorkWithLevelAssets.INT.udn
+++ b/Engine/Documentation/Source/Engine/Levels/HowTo/WorkWithLevelAssets/WorkWithLevelAssets.INT.udn
@@ -1,4 +1,4 @@
-Availability: Docs
+Availability: Public
 Title:Work with Level Assets
 Crumbs: %ROOT%, Engine
 Description:How to create, save, and open Level assets in UE4
diff --git a/Engine/Documentation/Source/Engine/Levels/LevelsWindow/LevelsWindow.INT.udn b/Engine/Documentation/Source/Engine/Levels/LevelsWindow/LevelsWindow.INT.udn
--- a/Engine/Documentation/Source/Engine/Levels/LevelsWindow/LevelsWindow.INT.udn
+++ b/Engine/Documentation/Source/Engine/Levels/LevelsWindow/LevelsWindow.INT.udn
@@ -1,4 +1,4 @@
-Availability: Docs
+Availability: Public
 Title:Managing Multiple Levels
 Description:Manage your persistent level and sublevels with the Levels window.
 Version: 4.12
diff --git a/Engine/Documentation/Source/Engine/Levels/WorldSettings/WorldSettings.INT.udn b/Engine/Documentation/Source/Engine/Levels/WorldSettings/WorldSettings.INT.udn
--- a/Engine/Documentation/Source/Engine/Levels/WorldSettings/WorldSettings.INT.udn
+++ b/Engine/Documentation/Source/Engine/Levels/WorldSettings/WorldSettings.INT.udn
@@ -1,4 +1,4 @@
-Availability: Docs
+Availability: Public
 Title:World Settings
 Description:The **World Settings** tab enables you to set and override Level-specific settings like base Game Mode and gravity. 
 Version: 4.12
