#UE4 Docs: minor formatting updates

diff --git a/Engine/Documentation/Source/Engine/Sequencer/HowTo/BurnIns/BurnIns.INT.udn b/Engine/Documentation/Source/Engine/Sequencer/HowTo/BurnIns/BurnIns.INT.udn
--- a/Engine/Documentation/Source/Engine/Sequencer/HowTo/BurnIns/BurnIns.INT.udn
+++ b/Engine/Documentation/Source/Engine/Sequencer/HowTo/BurnIns/BurnIns.INT.udn
@@ -8,7 +8,7 @@ parent:Engine/Sequencer/HowTo
 order:1
 checkpoint: editorqs
 Related: Engine/Sequencer/Overview
-Related: Engine\Sequencer\ReferenceEditor\
+Related: Engine/Sequencer/ReferenceEditor
 
 When creating and rendering out cinematic movies, you may want to include overlays with information about the scene you are viewing such as the name of the shot, the date, time or frame information. 
 These overlays are referred to as **Burn Ins** as they are burned into the movie when it is rendered out. 
@@ -46,7 +46,11 @@ For this how-to we are using the **Blueprint Third Person Template** project and
 
 ## End Result
 
-When the capturing of your movie is complete, upon playing back your movie you will notice similar overlays applied to the cinematic. 
+When the capturing of your movie is complete, upon playing back your movie you will notice similar overlays applied to the cinematic. The Default Burn in applies the name of your Level Sequence across the Top Center followed by the current date. In the Bottom Left, the current Shot Name is applied (or Level Sequence name if no shots are used). The Bottom Center displays the time and frame count for the Master Sequence while the Bottom Right shows the current Shot Frame. 
+
+Each of these options can be changed to display your own text inside the **Settings** section of the **Burn in Options**. Additionally, you can apply Watermarks or make Font changes within the Default Burn In options. 
+If the default information is not enough for you, you can also take a look at [](Engine\Sequencer\HowTo\CustomBurnIns) to include your desired content. 
+
 
 [OBJECT:EmbeddedVideo]
 [PARAMLITERAL:width]
@@ -64,11 +68,3 @@ DRIykPNZBAs
 You may want to view the video above in full screen or on YouTube to see the overlays more clearly.
 [/REGION]
 
-The Default Burn in applies the name of your Level Sequence across the Top Center followed by the current date. In the Bottom Left, the current Shot Name is applied (or Level Sequence name if no shots are used). 
-The Bottom Center displays the time and frame count for the Master Sequence while the Bottom Right shows the current Shot Frame. 
-
-Each of these options can be changed to display your own text inside the **Settings** section of the **Burn in Options**. Additionally, you can apply Watermarks or make Font changes within the Default Burn In options. 
-If the default information is not enough for you, you can also take a look at [](Engine\Sequencer\HowTo\CustomBurnIns) to include your desired content. 
-
-
-
diff --git a/Engine/Documentation/Source/Engine/Sequencer/HowTo/CustomBurnIns/CustomBurnIns.INT.udn b/Engine/Documentation/Source/Engine/Sequencer/HowTo/CustomBurnIns/CustomBurnIns.INT.udn
--- a/Engine/Documentation/Source/Engine/Sequencer/HowTo/CustomBurnIns/CustomBurnIns.INT.udn
+++ b/Engine/Documentation/Source/Engine/Sequencer/HowTo/CustomBurnIns/CustomBurnIns.INT.udn
@@ -21,36 +21,27 @@ For this how-to we are using the **Blueprint Third Person Template** project and
 
 ## Steps
 
-1.	In the **Content Browser**, click the **View Options** and **Show Engine Content**. 
+1.	In the **Content Browser**, click the **View Options** and **Show Engine Content**. This will grant us the ability to see and use Font options that are provided with the Engine Content. 
 
 	![](Custom01.png)
 
-	This will grant us the ability to see and use Font options that are provided with the Engine Content. 
+1.	**Right-click** in the **Content Browser** and create a **Blueprint Class** of the **LevelSequenceBurnInInitSettings** type. Name the new Blueprint **CustomSettings** as this will hold all the settings you wish to configure inside the Burn In Options. 
 
-1.	**Right-click** in the **Content Browser** and create a **Blueprint Class** of the **LevelSequenceBurnInInitSettings** type. 
+	![](Custom02.png) 
 
-	![](Custom02.png)
-
-	Name the new Blueprint **CustomSettings** as this will hold all the settings you wish to configure inside the Burn In Options.  
-
-1.	**Right-click** in the **Content Browser** and create a **Blueprint Class** of the **LevelSequenceBurnIn** type. 
+1.	**Right-click** in the **Content Browser** and create a **Blueprint Class** of the **LevelSequenceBurnIn** type. Name the new Blueprint **CustomBurnIn** as this will be the [UMG](Engine/UMG/) widget that we will create our custom burn in content in. 
 
 	![](Custom03.png)
 
-	Name the new Blueprint **CustomBurnIn** as this will be the [UMG](Engine/UMG/) widget that we will create our custom burn in content in. 
-
 1.	Open the **CustomSettings** Blueprint. 
 
 	![](Custom04.png)
 
-1.	Click the **+** sign next to **Variables** and add the options you want to expose to the Burn In Options. 
+1.	Click the **+** sign next to **Variables** and add the options you want to expose to the Burn In Options. Below we have created two variables, one called **CustomText** of the **Text** variable type and another called **CustomFont** of the **Slate Font Info** type (this will allow us to define the text and font style of burn in we create). We've set both variables to **Editable**, **Compiled** then set their **Default Values** (use any default values you wish). 
 
 	![](Custom05.png)
 
-	Above we have created two variables, one called **CustomText** of the **Text** variable type and another called **CustomFont** of the **Slate Font Info** type (this will allow us to define the text and font style of burn in we create). 
-	We've set both variables to **Editable**, **Compiled** then set their **Default Values** (use any default values you wish). 
-
-1.	Open the **CustomSettings** Blueprint, then arrange your burn in content on the **Designer** tab. 
+1.	Open the **CustomSettings** Blueprint, then arrange your burn in content on the **Designer** tab. Below we added a Text widget to the graph and anchored it to the bottom center of the screen. We also changed the name of the widget to **DisplayedText** and made it a **Variable** so that we can access it through script. You can arrange your content here just as you would any UMG Widget Blueprint (apply images, animation, etc.), for this example we are only going to display some custom text. 
 
 	[REGION:lightbox]
 	[![](Custom06.png)](Custom06.png)
@@ -60,28 +51,20 @@ For this how-to we are using the **Blueprint Third Person Template** project and
 	Click image for full view. 
 	[/REGION]
 
-	Above we added a Text widget to the graph and anchored it to the bottom center of the screen. 
-	We also changed the name of the widget to **DisplayedText** and made it a **Variable** so that we can access it through script. 
-	You can arrange your content here just as you would any UMG Widget Blueprint (apply images, animation, etc.), for this example we are only going to display some custom text. 
-
-1.	On the **Graph** tab, in the **My Blueprint** window click the **Function Override** button and select **Get Settings Class**. 
+1.	On the **Graph** tab, in the **My Blueprint** window click the **Function Override** button and select **Get Settings Class**. We are going to point this Blueprint toward our **CustomSettings** Blueprint Class we created that way it knows which settings to use. 
 
 	![](Custom07.png)
 
-	We are going to point this Blueprint toward our **CustomSettings** Blueprint Class we created that way it knows which settings to use. 
-
 1.	In the graph that opens, click the **Return Value** dropdown and select **CustomSettings**. 
 
 	![](Custom08.png)
 
 1.	Return to the **Event Graph** and **Right-click** in the graph window and add the **Event Set Settings** node. 
 
-1.	Drag off the **In Settings** and **Cast To CustomSettings**. 
+1.	Drag off the **In Settings** and **Cast To CustomSettings**. This will allow us to access the options we've defined within it. 
 
 	![](Custom09.png)
 
-	This will allow us to access the options we've defined within it. 
-
 1.	**Right-click** on the **As Custom Settings** pin and **Promote to Variable** called **Options**. 
 
 	![](Custom10.png)
@@ -119,7 +102,7 @@ For this how-to we are using the **Blueprint Third Person Template** project and
 
 ## End Result
 
-When the capturing of your movie is complete, upon playing back your movie you will notice the overlays applied to the cinematic. 
+When the capturing of your movie is complete, upon playing back your movie you will notice the overlays applied to the cinematic. While our example allows us to enter custom text and adjust the font style, you can expose various properties inside the **Level Sequence Burn In Init Settings** class to generate your desired effects. Inside your UMG Widget you can then script the functionality you want your burn ins to take, perhaps even animating them upon **Event Construct**.
 
 [OBJECT:EmbeddedVideo]
 [PARAMLITERAL:width]
@@ -133,4 +116,4 @@ mZzVTyb0r34
 [/PARAMLITERAL]
 [/OBJECT]
 
-While our example allows us to enter custom text and adjust the font style, you can expose various properties inside the **Level Sequence Burn In Init Settings** class to generate your desired effects. Inside your UMG Widget you can then script the functionality you want your burn ins to take, perhaps even animating them upon **Event Construct**.
+
