#UE4 Docs: fixed missing image in slider

 - updated metadata

diff --git a/Engine/Documentation/Source/Platforms/TVSafeZoneDebugging/TVSafeZoneDebugging.INT.udn b/Engine/Documentation/Source/Platforms/TVSafeZoneDebugging/TVSafeZoneDebugging.INT.udn
--- a/Engine/Documentation/Source/Platforms/TVSafeZoneDebugging/TVSafeZoneDebugging.INT.udn
+++ b/Engine/Documentation/Source/Platforms/TVSafeZoneDebugging/TVSafeZoneDebugging.INT.udn
@@ -1,8 +1,11 @@
-Availability:Public
+Availability:Docs
 Title:Setting TV Safe Zone Debugging
 Crumbs:%ROOT%, Platforms
 Description:Making sure that UI elements are not too close to the edges of the TV
 Version: 4.12
+skilllevel:Intermediate
+parent:Platforms
+order:1000
 
 [TOC(start:2)]
 
@@ -150,7 +153,7 @@ In the following section we will go over some pieces of information you need to
 	![](SZ_2.png)
 	![](SZ_3.png)
 	![](SZ_4.png)
-	![](SZ_6.png)
+	![](SZ_5.png)
 	![](SZ_6.png)
 	![](SZ_7.png)
 	![](SZ_8.png)
