Draft Beginner Linux Developers Directory Page.

This Landing Page contains a link to the Supported Graphics Drivers Overview page.

#UE4docs

diff --git a/Engine/Documentation/Source/Platforms/Linux/BeginnerLinuxDeveloper/BeginnerLinuxDeveloper.INT.udn b/Engine/Documentation/Source/Platforms/Linux/BeginnerLinuxDeveloper/BeginnerLinuxDeveloper.INT.udn
new file mode 100644
--- /dev/null
+++ b/Engine/Documentation/Source/Platforms/Linux/BeginnerLinuxDeveloper/BeginnerLinuxDeveloper.INT.udn
@@ -0,0 +1,44 @@
+Availability: Docs
+Title: Beginner Linux Developers
+Crumbs:%ROOT%, Platforms, Platforms/Linux
+Description:Learn how to setup a typical workflow for developing games on Linux
+Platform:Linux
+Navigation:platform
+Version: 4.11
+type:landing
+parent:Platforms/Linux
+order: 1
+tags:Linux
+SkillLevel:Beginner
+
+[VAR:Params]
+[PARAM:image]
+	![%Platforms/Linux/BeginnerLinuxDeveloper:title%](beginner_topics.png)
+[/PARAM]
+[PARAM:icon]
+	![](%ROOT%/tutorial_icon.png)(convert:false)
+[/PARAM]
+[PARAM:title]
+	%Platforms/Linux/BeginnerLinuxDeveloper:title%
+[/PARAM]
+[PARAM:description]
+	%Platforms/Linux/BeginnerLinuxDeveloper:description%
+[/PARAM]
+[PARAM:path]
+	[RELATIVE:Platforms/Linux/BeginnerLinuxDeveloper]
+[/PARAM]
+[/VAR]
+
+[VAR:Topic]
+[OBJECT:Topic]
+%params%
+[/OBJECT]
+[/VAR]
+
+If you're a fledgling developer, you'll want to start off with this learning path, where you'll setup a typical workflow for Linux game development. 
+By the end of this learning path, you'll setup an Integrated Development Environment (IDE), you'll have a thorough understanding of the Unreal Engine 4 (UE4) 
+user interface, and you'll gain some experience by developing a game on Linux.
+
+[REGION:topics third]
+%Platforms/Linux/BeginnerLinuxDeveloper/SupportedGfxDrivers:topic%
+[/REGION]
\ No newline at end of file
