Added the following Metadata in order to comply with the latest version of the Style & Standards Guide:
- Parent
- Order
- Type

Updated engine version to reflect latest stable release.
Added Table of Contents in order to comply with the latest version of the Style & Standards Guide.
Various minor grammatical and stylistic revisions per my Peer Review comments.

#UE4docs

diff --git a/Engine/Documentation/Source/Gameplay/Networking/Actors/Properties/Properties.INT.udn b/Engine/Documentation/Source/Gameplay/Networking/Actors/Properties/Properties.INT.udn
--- a/Engine/Documentation/Source/Gameplay/Networking/Actors/Properties/Properties.INT.udn
+++ b/Engine/Documentation/Source/Gameplay/Networking/Actors/Properties/Properties.INT.udn
@@ -5,20 +5,22 @@ Description:Detailed information about how Actor properties are replicated.
 Related: Gameplay/Networking/Actors/Properties/ObjectReferences
 Related: Gameplay/Networking/Actors/Properties/Conditions
 Related: Gameplay/Networking/Actors/RPCs
-Version:4.9
+Version:4.11
+type: Reference
+parent: Gameplay/Networking/Actors
+order: 1
 
+[TOC(start:1 end:4)]
 
+Each actor maintains a list of properties that can be marked for replication to clients. Whenever the value of the variable changes on the server side, the server sends the client the updated value. The variable may have changed on the client side, in which case the new value will overwrite it. Property updates only come from the server (i.e., the client will never send property updates to the server).
 
-
-Each actor maintains a list of properties, and each property can be marked to be replicated to clients. Whenever the value of the variable changes on the server side, the server sends the client the updated value. The variable may have changed on client side too - in which case the new value will overwrite it. Property updates only come from the server, the client will never send property updates to the server.
-
-Some properties replicate by default (Location, Rotation, Owner, etc). These are properties that the framework generally depends on to make basic multiplayer work. From there, you can mark additional properties to replicate as you need.
+Some properties replicate by default (e.g., Location, Rotation, Owner, etc.). These are properties that the framework generally depends on to make basic multiplayer work. From there, you can mark additional properties to replicate as needed.
 
 Actor property replication is reliable. This means that the property of the client version of the Actor will eventually reflect the value on the server - but it's possible that some individual property value changes will be skipped.
 
 ## Setting up properties for replication
 
-To replicate a property you need to do a few things:
+To replicate a property, you need to do a few things:
 In the header of the actor class where the property is defined, you need to make sure you have the `replicated` keyword as one of the parameters to the UPROPERTY declaration:
 
 	class ENGINE_API AActor : public UObject
@@ -41,7 +43,7 @@ In the actor's constructor, make sure you have the bReplicates flag set to true:
 		bReplicates = true;
 	}
 
-That's about it. The member variable `Owner` will now be synchronized to all connected clients for every copy of this actor type that is currently instantiated (in this case, the base actor class).
+The member variable `Owner` will now be synchronized to all connected clients for every copy of this actor type that is currently instantiated (in this case, the base actor class).
 
 ## Network Update Optimization
 
@@ -55,8 +57,8 @@ With Adaptive Network Update Frequency, replicated Actors will increase or reduc
 
 #### (Advanced) Update Frequency Decrease Algorithm
 
-During update attempts, Actors determine how long it has been since the most recent meaningful update was sent, and record the new time if they send a meaningful update. If an Actor being considered for update has not sent a meaningful update for more than two seconds, it will start decreasing its update frequency, reaching its minimum frequency after seven seconds without a meaningful update. For example, if an Actor with an update delay between 0.1 and 0.6 seconds had had no meaningful updates for 3 seconds, it would attempt its next update in 0.2 seconds.
+During update attempts, Actors determine how long it has been since the most recent meaningful update was sent, and record the new time if they send a meaningful update. If an Actor being considered for update has not sent a meaningful update for more than two seconds, it will start decreasing its update frequency, reaching its minimum frequency after seven seconds without a meaningful update. For example, if an Actor with an update delay between 0.1 and 0.6 seconds had no meaningful updates for 3 seconds, it would attempt its next update in 0.2 seconds.
 
 #### (Advanced) Update Frequency Increase Algorithm
 
-After sending a meaningful update, an Actor will schedule its next update to happen in 30% less time than the time between the last two meaningful updates, clamped between the minimum and maximum update frequencies. For example, if an Actor went exactly one second between meaningful updates, it would schedule its next update attempt for 0.7 seconds in the future, or as close to that value as its minimum and maximum update frequencies permit. With each successive meaningful update, this calculation will be repeated, quickly reducing the time between updates if an Actor starts to have frequent data or subobject changes.
+After sending a meaningful update, an Actor will schedule its next update to happen in 30-percent less time than the time between the last two meaningful updates, clamped between the minimum and maximum update frequencies. For example, if an Actor went exactly one second between meaningful updates, it would schedule its next update attempt for 0.7 seconds in the future (or a time near the specified minimum and maximum update frequencies). With each successive meaningful update, this calculation will be repeated, quickly reducing the time between updates if an Actor starts to have frequent data or subobject changes.
