#UE4 docs adding reference to LOD Threshold as a temporary measure, will need to refactor all the reference pages in the Skeletal Controls section

diff --git a/Engine/Documentation/Source/Engine/Animation/AimOffset/AimOffset.INT.udn b/Engine/Documentation/Source/Engine/Animation/AimOffset/AimOffset.INT.udn
--- a/Engine/Documentation/Source/Engine/Animation/AimOffset/AimOffset.INT.udn
+++ b/Engine/Documentation/Source/Engine/Animation/AimOffset/AimOffset.INT.udn
@@ -97,8 +97,10 @@ This is solved with Mesh space rotation. Since rotation is based on the Skeletal
 
 Mesh space is set as a property of an Animation Sequence. Under the **Additive Settings** category, you will see the **Additive Anim Type** property. Setting this to Mesh Space will utilize the Mesh coordinate space as described. This is a fairly expensive blending operation, and should generally only be used by certain types of Blend Spaces where you know you need the animation to move in a constant direction, such as in the case of an Aim Offset.
 
+## Performance
 
+![](animOffsetNode.png)
 
+When an Aim Offset node is placed in the **Anim Graph** you can adjust their performance using the **LOD Threshold** property, introduced in 4.11, you can find the setting for this in the Details Panel for the **Aim Offset** node.
 
-
-
+[INCLUDE:Engine/Animation/NodeReference/SkeletalControls#lodthresh]
\ No newline at end of file
diff --git a/Engine/Documentation/Source/Engine/Animation/NodeReference/Blend/BlendNodes.INT.udn b/Engine/Documentation/Source/Engine/Animation/NodeReference/Blend/BlendNodes.INT.udn
--- a/Engine/Documentation/Source/Engine/Animation/NodeReference/Blend/BlendNodes.INT.udn
+++ b/Engine/Documentation/Source/Engine/Animation/NodeReference/Blend/BlendNodes.INT.udn
@@ -14,9 +14,9 @@ version: 4.9
 **Blend Nodes** are used to blend multiple animations together. These nodes are only available to be placed in the **Anim Graph** of a Animation Blueprint. Each blend node takes in multiple poses and an alpha or weight value that is used to calculate the weighting for each pose in the final output. Some blend nodes may be more complex than others and require additional data to be input as well.
 [/EXCERPT:Intro]
 
-## Apply Additive
+## Apply Additive and Apply Mesh Space Additive
 
-The **Apply Additive** node adds an additive pose to a base normal animation pose based on an alpha value.
+The **Apply Additive** adn **Apply Mesh Space Additive** nodes add an additive pose to a base normal animation pose based on an alpha value.
 
 ![](ApplyAdditive.png)
 
@@ -29,6 +29,10 @@ The **Apply Additive** node adds an additive pose to a base normal animation pos
 | [REGION:tablesection]**Output Pins**[/REGION] ||
 | **Pose** | The final pose after the blending has been applied. |
 
+**Apply Additive** nodes are also affected by the **LOD Threshold** system introduced in 4.11, you can find the setting for this in the Details Panel for the **Apply Additive** node.
+
+[INCLUDE:Engine/Animation/NodeReference/SkeletalControls#lodthresh]
+
 ## Blend
 
 The standard **Blend** node is a straight mixing of the two input poses based on an alpha input.
diff --git a/Engine/Documentation/Source/Engine/Animation/NodeReference/SkeletalControls/AnimDynamics/AnimDynamics.INT.udn b/Engine/Documentation/Source/Engine/Animation/NodeReference/SkeletalControls/AnimDynamics/AnimDynamics.INT.udn
--- a/Engine/Documentation/Source/Engine/Animation/NodeReference/SkeletalControls/AnimDynamics/AnimDynamics.INT.udn
+++ b/Engine/Documentation/Source/Engine/Animation/NodeReference/SkeletalControls/AnimDynamics/AnimDynamics.INT.udn
@@ -112,8 +112,71 @@ To quickly get started, create an **AnimDynamics** node inside the AnimGraph of
 | [REGION:tablesection]**Output Pins**[/REGION] ||
 | Pose (Blank) | The final pose after the Transform has been applied. |
 
+**AnimDynamics** nodes are also affected by the **LOD Threshold** system introduced in 4.11, you can find the setting for this in the Details Panel for the **AnimDynamics** node.
+
+[INCLUDE:Engine/Animation/NodeReference/SkeletalControls#lodthresh]
+
 [REGION:note]
 Please see the [**Creating Dyanmic Animations**](Engine\Animation\AnimHowTo\AnimDynamics) how-to for a step-by-step example of using AnimDyamics. 
 [/REGION]
 
 
+<!--
+[OBJECT:modulesectionexpanded]
+	[PARAMLITERAL:id]
+		LODThresholdMSE
+	[/PARAMLITERAL]
+	[PARAM:heading]
+		Performance
+	[/PARAM]
+	[PARAM:content]
+		%p100%
+	[/PARAM]
+[/OBJECT]
+
+[VAR:p100]
+
+[INCLUDE:Engine/Animation/NodeReference/SkeletalControls#lodthresh]
+
+[/VAR]
+
+
+[VAR:spacer]![](%ROOT%/spacer.png)(w:150 convert:false)[/VAR]
+
+[VAR:spacer2]![](%ROOT%/spacer.png)(w:35 convert:false)[/VAR]
+
+[REGION:raw]
+| Property | %spacer% | Description |
+| --- | --- | :--- |
+| **Node** || %doxygen:UAnimGraphNode_AnimDynamics::Node% |
+| **Preview Live** || %doxygen:UAnimGraphNode_AnimDynamics::bPreviewLive% |
+| **Show Linear Limits** || %doxygen:UAnimGraphNode_AnimDynamics::bShowLinearLimits% |
+| **Show Angular Limits** || %doxygen:UAnimGraphNode_AnimDynamics::bShowAngularLimits% |
+| **Show Planar Limit** || %doxygen:UAnimGraphNode_AnimDynamics::bShowPlanarLimit% |
+| **Show Spherical Limit** || %doxygen:UAnimGraphNode_AnimDynamics::bShowSphericalLimit% |
+| **Show Collision Spheres** || %doxygen:UAnimGraphNode_AnimDynamics::bShowCollisionSpheres% |
+| **Last Preview Component** || %doxygen:UAnimGraphNode_AnimDynamics::LastPreviewComponent% |
+| **Cached Node Titles** || %doxygen:UAnimGraphNode_AnimDynamics::CachedNodeTitles% |
+| **Edit Preview Shape** || %doxygen:UAnimGraphNode_AnimDynamics::EditPreviewShape% |
+| **Post Load** || %doxygen:UAnimGraphNode_AnimDynamics::PostLoad% |
+| **Reset Button Clicked** || %doxygen:UAnimGraphNode_AnimDynamics::ResetButtonClicked% |
+| **Reset Sim** || %doxygen:UAnimGraphNode_AnimDynamics::ResetSim% |
+| **Serialize** || %doxygen:UAnimGraphNode_AnimDynamics::Serialize% |
+| **Get Node Title** || %doxygen:UAnimGraphNode_AnimDynamics::GetNodeTitle% |
+| **Get Tooltip Text** || %doxygen:UAnimGraphNode_AnimDynamics::GetTooltipText% |
+| **Validate Anim Node During Compilation** || %doxygen:UAnimGraphNode_AnimDynamics::ValidateAnimNodeDuringCompilation% |
+| **Customize Details** || %doxygen:UAnimGraphNode_AnimDynamics::CustomizeDetails% |
+| **Draw** || %doxygen:UAnimGraphNode_AnimDynamics::Draw% |
+| **Get On Screen Debug Info** || %doxygen:UAnimGraphNode_AnimDynamics::GetOnScreenDebugInfo% |
+| **Get Preview Dynamics Node** || %doxygen:UAnimGraphNode_AnimDynamics::GetPreviewDynamicsNode% |
+| **Get Controller Description** || %doxygen:UAnimGraphNode_AnimDynamics::GetControllerDescription% |
+| **Post Edit Change Property** || %doxygen:UAnimGraphNode_AnimDynamics::PostEditChangeProperty% |
+| **Draw Linear Limits** || %doxygen:UAnimGraphNode_AnimDynamics::DrawLinearLimits% |
+| **Draw Angular Limits** || %doxygen:UAnimGraphNode_AnimDynamics::DrawAngularLimits% |
+| **Get Node** || %doxygen:UAnimGraphNode_AnimDynamics::GetNode% |
+[/REGION]
+-->
+
+
+
+
diff --git a/Engine/Documentation/Source/Engine/Animation/NodeReference/SkeletalControls/ApplyPercentRotation/ApplyPercentRotation.INT.udn b/Engine/Documentation/Source/Engine/Animation/NodeReference/SkeletalControls/ApplyPercentRotation/ApplyPercentRotation.INT.udn
--- a/Engine/Documentation/Source/Engine/Animation/NodeReference/SkeletalControls/ApplyPercentRotation/ApplyPercentRotation.INT.udn
+++ b/Engine/Documentation/Source/Engine/Animation/NodeReference/SkeletalControls/ApplyPercentRotation/ApplyPercentRotation.INT.udn
@@ -77,4 +77,6 @@ By changing the **Rotation Axis to Refer** value we can alter which way the pack
 | [REGION:tablesection]**Input Pins**[/REGION] ||
 | **Multiplier** | A float value that represents the percentage of the source bone's Rotation to apply to the target bone. When **Multiplier** is 0, none of the source bone's Rotation is applied to the target bone, and when it is 1.0, all of the source bone's Rotation is applied. |
 
+**Apply a Percentage of Rotation** nodes are also affected by the **LOD Threshold** system introduced in 4.11, you can find the setting for this in the Details Panel for the **Apply a Percentage of Rotation** node.
 
+[INCLUDE:Engine/Animation/NodeReference/SkeletalControls#lodthresh]
diff --git a/Engine/Documentation/Source/Engine/Animation/NodeReference/SkeletalControls/BoneDrivenController/BoneDrivenController.INT.udn b/Engine/Documentation/Source/Engine/Animation/NodeReference/SkeletalControls/BoneDrivenController/BoneDrivenController.INT.udn
--- a/Engine/Documentation/Source/Engine/Animation/NodeReference/SkeletalControls/BoneDrivenController/BoneDrivenController.INT.udn
+++ b/Engine/Documentation/Source/Engine/Animation/NodeReference/SkeletalControls/BoneDrivenController/BoneDrivenController.INT.udn
@@ -75,6 +75,9 @@ CDG4WcQ7cIY
 | **(As Pin) Alpha** | Enable/Disable the display of the Alpha pin on the node. |
 | **Alpha Scale Bias** | Set the minimal and maximal input scaling values. |
 
+**Bone Driven Controller** nodes are also affected by the **LOD Threshold** system introduced in 4.11, you can find the setting for this in the Details Panel for the **Bone Driven Controller** node.
+
+[INCLUDE:Engine/Animation/NodeReference/SkeletalControls#lodthresh]
 
 <!--
 [EXCERPT:LocationSpace]
diff --git a/Engine/Documentation/Source/Engine/Animation/NodeReference/SkeletalControls/CopyBone/CopyBone.INT.udn b/Engine/Documentation/Source/Engine/Animation/NodeReference/SkeletalControls/CopyBone/CopyBone.INT.udn
--- a/Engine/Documentation/Source/Engine/Animation/NodeReference/SkeletalControls/CopyBone/CopyBone.INT.udn
+++ b/Engine/Documentation/Source/Engine/Animation/NodeReference/SkeletalControls/CopyBone/CopyBone.INT.udn
@@ -78,3 +78,6 @@ Be cautious when selecting the checkboxes for Transform. For instance, if you en
 | **Copy Rotation** | A bool specifying whether to copy the Rotation component of the source bone's Transform. |
 | **Copy Scale** | A bool specifying whether to copy the Scale component of the source bone's Transform. |
 
+**Copy Bone** nodes are also affected by the **LOD Threshold** system introduced in 4.11, you can find the setting for this in the Details Panel for the **Copy Bone** node.
+
+[INCLUDE:Engine/Animation/NodeReference/SkeletalControls#lodthresh]
diff --git a/Engine/Documentation/Source/Engine/Animation/NodeReference/SkeletalControls/HandIKRetargeting/HandIKRetargeting.INT.udn b/Engine/Documentation/Source/Engine/Animation/NodeReference/SkeletalControls/HandIKRetargeting/HandIKRetargeting.INT.udn
--- a/Engine/Documentation/Source/Engine/Animation/NodeReference/SkeletalControls/HandIKRetargeting/HandIKRetargeting.INT.udn
+++ b/Engine/Documentation/Source/Engine/Animation/NodeReference/SkeletalControls/HandIKRetargeting/HandIKRetargeting.INT.udn
@@ -83,3 +83,6 @@ See the **Animation** map inside the **Content Examples** project (available for
 | **Component Pose** | The incoming pose. |
 | **Alpha** | The current strength of the control. |
 
+**Hand IK Retargeting** nodes are also affected by the **LOD Threshold** system introduced in 4.11, you can find the setting for this in the Details Panel for the **Hand IK Retargeting** node.
+
+[INCLUDE:Engine/Animation/NodeReference/SkeletalControls#lodthresh]
diff --git a/Engine/Documentation/Source/Engine/Animation/NodeReference/SkeletalControls/LookAt/LookAt.INT.udn b/Engine/Documentation/Source/Engine/Animation/NodeReference/SkeletalControls/LookAt/LookAt.INT.udn
--- a/Engine/Documentation/Source/Engine/Animation/NodeReference/SkeletalControls/LookAt/LookAt.INT.udn
+++ b/Engine/Documentation/Source/Engine/Animation/NodeReference/SkeletalControls/LookAt/LookAt.INT.udn
@@ -89,3 +89,6 @@ You may need to enlarge the video to see the debug line or try it out in the edi
 | **Component Pose** | The incoming pose. |
 | **Alpha** | The current strength of the control. |
 
+**Look At** nodes are also affected by the **LOD Threshold** system introduced in 4.11, you can find the setting for this in the Details Panel for the **Look At** node.
+
+[INCLUDE:Engine/Animation/NodeReference/SkeletalControls#lodthresh]
diff --git a/Engine/Documentation/Source/Engine/Animation/NodeReference/SkeletalControls/ObserveBone/ObserveBone.INT.udn b/Engine/Documentation/Source/Engine/Animation/NodeReference/SkeletalControls/ObserveBone/ObserveBone.INT.udn
--- a/Engine/Documentation/Source/Engine/Animation/NodeReference/SkeletalControls/ObserveBone/ObserveBone.INT.udn
+++ b/Engine/Documentation/Source/Engine/Animation/NodeReference/SkeletalControls/ObserveBone/ObserveBone.INT.udn
@@ -58,4 +58,8 @@ Below, a bone by the name of **b_MF_Hand_L** is being observed and debug informa
 | [REGION:tablesection]**Input Pins**[/REGION] ||
 | **Component Pose** | The incoming pose. |
 
+**Observe Bone** nodes are also affected by the **LOD Threshold** system introduced in 4.11, you can find the setting for this in the Details Panel for the **Observe Bone** node.
+
+[INCLUDE:Engine/Animation/NodeReference/SkeletalControls#lodthresh]
+
 
diff --git a/Engine/Documentation/Source/Engine/Animation/NodeReference/SkeletalControls/SkeletalControls.INT.udn b/Engine/Documentation/Source/Engine/Animation/NodeReference/SkeletalControls/SkeletalControls.INT.udn
--- a/Engine/Documentation/Source/Engine/Animation/NodeReference/SkeletalControls/SkeletalControls.INT.udn
+++ b/Engine/Documentation/Source/Engine/Animation/NodeReference/SkeletalControls/SkeletalControls.INT.udn
@@ -30,6 +30,15 @@ Some pins and properties are common to all SkeletalControls. These are described
 | [REGION:tablesection]**Output Pins**[/REGION] ||
 | **Pose (Blank)** | The final pose after the Transform has been applied |
 
+[EXCERPT:LODThresh]
+![](perf.png)
+
+| Property | Description |
+| --- | --- |
+| **LOD Threshold** | This is the max Level of Detail (**LOD**) that this node is permitted to run on. For example: if you have LOD Threshold set to **2**, it will run until **LOD 2**, but disable itself once the component's LOD becomes **3**. |
+[/EXCERPT]
+
+
 ## Component Space
 
 [INCLUDE:Engine/Animation/NodeReference/SpaceConversion#intro]
diff --git a/Engine/Documentation/Source/Engine/Animation/NodeReference/SkeletalControls/SpringController/SpringController.INT.udn b/Engine/Documentation/Source/Engine/Animation/NodeReference/SkeletalControls/SpringController/SpringController.INT.udn
--- a/Engine/Documentation/Source/Engine/Animation/NodeReference/SkeletalControls/SpringController/SpringController.INT.udn
+++ b/Engine/Documentation/Source/Engine/Animation/NodeReference/SkeletalControls/SpringController/SpringController.INT.udn
@@ -71,4 +71,8 @@ We then connect the nodes to include the Spring Controller that is targeting the
 | **Error Reset Thresh** | A threshold displacement - if the spring stretches more than this amount, it resets in order to avoid errors introduced by sudden, large displacements such as those caused by teleporting Actors. |
 | **No ZSpring** | If enabled, the displacement in the Z axis is ignored by the spring and is always correct. |
 
+**Spring Controller** nodes are also affected by the **LOD Threshold** system introduced in 4.11, you can find the setting for this in the Details Panel for the **Spring Controller** node.
+
+[INCLUDE:Engine/Animation/NodeReference/SkeletalControls#lodthresh]
+
 
diff --git a/Engine/Documentation/Source/Engine/Animation/NodeReference/SkeletalControls/TrailController/TrailController.INT.udn b/Engine/Documentation/Source/Engine/Animation/NodeReference/SkeletalControls/TrailController/TrailController.INT.udn
--- a/Engine/Documentation/Source/Engine/Animation/NodeReference/SkeletalControls/TrailController/TrailController.INT.udn
+++ b/Engine/Documentation/Source/Engine/Animation/NodeReference/SkeletalControls/TrailController/TrailController.INT.udn
@@ -65,5 +65,9 @@ The **Trail Controller** can be used to affect a series of bones in a hierarchy.
 | **Component Pose** | The incoming pose. |
 | **Alpha** | The current strength of the control. |
 
+**Trail Controller** nodes are also affected by the **LOD Threshold** system introduced in 4.11, you can find the setting for this in the Details Panel for the **Trail Controller** node.
+
+[INCLUDE:Engine/Animation/NodeReference/SkeletalControls#lodthresh]
+
 
 
diff --git a/Engine/Documentation/Source/Engine/Animation/NodeReference/SkeletalControls/TransformBone/TransformBone.INT.udn b/Engine/Documentation/Source/Engine/Animation/NodeReference/SkeletalControls/TransformBone/TransformBone.INT.udn
--- a/Engine/Documentation/Source/Engine/Animation/NodeReference/SkeletalControls/TransformBone/TransformBone.INT.udn
+++ b/Engine/Documentation/Source/Engine/Animation/NodeReference/SkeletalControls/TransformBone/TransformBone.INT.udn
@@ -73,3 +73,7 @@ The scale of our character is also driven by a float variable (**Scale Multiplie
 | **Rotation** | A Rotator specifying the Rotation to apply to the target bone. |
 | **Scale** | A Vector specifying the Scale to apply to the target bone. |
 
+**Transform (Modify) Bone** nodes are also affected by the **LOD Threshold** system introduced in 4.11, you can find the setting for this in the Details Panel for the **Transform (Modify) Bone** node.
+
+[INCLUDE:Engine/Animation/NodeReference/SkeletalControls#lodthresh]
+
diff --git a/Engine/Documentation/Source/Engine/Animation/NodeReference/SkeletalControls/TwistCorrective/TwistCorrective.INT.udn b/Engine/Documentation/Source/Engine/Animation/NodeReference/SkeletalControls/TwistCorrective/TwistCorrective.INT.udn
--- a/Engine/Documentation/Source/Engine/Animation/NodeReference/SkeletalControls/TwistCorrective/TwistCorrective.INT.udn
+++ b/Engine/Documentation/Source/Engine/Animation/NodeReference/SkeletalControls/TwistCorrective/TwistCorrective.INT.udn
@@ -63,4 +63,8 @@ You can add the Twist Corrective node in the **AnimGraph** of your **Animation B
 | [REGION:tablesection]**Output Pins**[/REGION] ||
 | Pose (Blank) | The final pose after the Transform has been applied. |
 
+**Twist Corrective** nodes are also affected by the **LOD Threshold** system introduced in 4.11, you can find the setting for this in the Details Panel for the **Twist Corrective** node.
+
+[INCLUDE:Engine/Animation/NodeReference/SkeletalControls#lodthresh]
+
 
diff --git a/Engine/Documentation/Source/Engine/Animation/NodeReference/SkeletalControls/TwoBoneIK/TwoBoneIK.INT.udn b/Engine/Documentation/Source/Engine/Animation/NodeReference/SkeletalControls/TwoBoneIK/TwoBoneIK.INT.udn
--- a/Engine/Documentation/Source/Engine/Animation/NodeReference/SkeletalControls/TwoBoneIK/TwoBoneIK.INT.udn
+++ b/Engine/Documentation/Source/Engine/Animation/NodeReference/SkeletalControls/TwoBoneIK/TwoBoneIK.INT.udn
@@ -84,6 +84,10 @@ By adjusting these values we can determine how the character holds the weapon wi
 | **Effector Location** | A Vector specifying the location effector in the Effector Location Space. |
 | **Joint Target Location** | A Vector specifying the location joint target in the Joint Target Location Space. |
 
+**Two Bone IK** nodes are also affected by the **LOD Threshold** system introduced in 4.11, you can find the setting for this in the Details Panel for the **Two Bone IK** node.
+
+[INCLUDE:Engine/Animation/NodeReference/SkeletalControls#lodthresh]
+
 <!--
 [EXCERPT:LocationSpace]
 [REGION:simpletable]
