Added VR Editor Controls for Mesh Paint. Added note w/link to Mesh Paint controls page.

diff --git a/Engine/Documentation/Source/Engine/Editor/VR/Controls/VREditorControls.INT.udn b/Engine/Documentation/Source/Engine/Editor/VR/Controls/VREditorControls.INT.udn
--- a/Engine/Documentation/Source/Engine/Editor/VR/Controls/VREditorControls.INT.udn
+++ b/Engine/Documentation/Source/Engine/Editor/VR/Controls/VREditorControls.INT.udn
@@ -115,6 +115,16 @@ Each Controller emits a beam called the laser pointer which is used to determine
 | ![Aim Controller at Window Grip + Squeeze Trigger + Track Pad Towards or Away](scale_window.png)(convert:false) | Aim the controller at the Grip of a window and squeeze the Trigger to grab the window. Move your thumb up or down on the Track Pad to scale the window up or down. |
 [/REGION]
 
+## Mesh Paint Controls
+
+[REGION:raw]
+| Control | Description |
+| ------- | ----------- |
+| ![Aim Controller at Actor + Squeeze Trigger + Move Controller](MeshPaint_2_NoModifier.png)(convert:false) | Aim the controller at a Static Mesh in the world and squeeze the Trigger to begin painting. Move the controller to move the brush along the Mesh. |
+| ![Modifier + Aim Controller at Actor + Squeeze Trigger + Move Controller](ModifierMeshPaint.png)(convert:false) | To erase instead of paint, hold the modifier button. Aim the controller at a Static Mesh in the world and squeeze the Trigger to begin erasing. Move the controller to move the brush along the Mesh.  |
+[/REGION]
+
+
 <!--
 [EXCERPT:MultiSelect]
 [PUBLISH:Rocket]
diff --git a/Engine/Documentation/Source/Engine/UI/LevelEditor/Modes/MeshPaintMode/VertexColor/PaintingOnMeshes/PaintingOnMeshes.INT.udn b/Engine/Documentation/Source/Engine/UI/LevelEditor/Modes/MeshPaintMode/VertexColor/PaintingOnMeshes/PaintingOnMeshes.INT.udn
--- a/Engine/Documentation/Source/Engine/UI/LevelEditor/Modes/MeshPaintMode/VertexColor/PaintingOnMeshes/PaintingOnMeshes.INT.udn
+++ b/Engine/Documentation/Source/Engine/UI/LevelEditor/Modes/MeshPaintMode/VertexColor/PaintingOnMeshes/PaintingOnMeshes.INT.udn
@@ -24,6 +24,10 @@ The circular brush outline shows you the radius of the brush.  The inner circle
 
 To paint, hold down the **Left Mouse Button** and drag over the mesh's surface.  You can also hold down the **Shift** key while painting to erase instead.
 
+[REGION:tip]
+Mesh painting also works in the VR Editor. For the controls for Mesh Paint in VR, see the [](Engine/Editor/VR/Controls).
+[/REGION]
+
 Paint will be applied to the surface each time you click, and also each time the mouse position changes by dragging.  Also, if **Brush Flow** is enabled, paint will be applied every time the scene is rendered.
 
 ![](PaintStroke_Green.png)(w:720)
