#UE4 Docs: fixed links to videos

diff --git a/Engine/Documentation/Source/Resources/Concepts/GameplayConceptExamples.INT.udn b/Engine/Documentation/Source/Resources/Concepts/GameplayConceptExamples.INT.udn
--- a/Engine/Documentation/Source/Resources/Concepts/GameplayConceptExamples.INT.udn
+++ b/Engine/Documentation/Source/Resources/Concepts/GameplayConceptExamples.INT.udn
@@ -63,7 +63,7 @@ Feel free to open up and analyze any element of these examples, as well as use t
 	Blueprint Splines Track
 	[/PARAMLITERAL]
 	[PARAMLITERAL:playlistid]
-	PLZlv_N0_O1ga0aV9jVqJgog0VWz1cLL5f
+	PLZlv_N0_O1gYuWkyTSbfcuyPNOQsK_8bu
 	[/PARAMLITERAL]
 	[PARAMLITERAL:videoid]
 	wR0fH6O9jD8
@@ -79,6 +79,9 @@ Feel free to open up and analyze any element of these examples, as well as use t
 	[PARAMLITERAL:title]
 	Unreal Stick Figure 2D
 	[/PARAMLITERAL]
+	[PARAMLITERAL:playlistid]
+	PLZlv_N0_O1ga0aV9jVqJgog0VWz1cLL5f
+	[/PARAMLITERAL]
 	[PARAMLITERAL:videoid]
 	qWUEkTd2Xu8
 	[/PARAMLITERAL]
@@ -111,7 +114,7 @@ Feel free to open up and analyze any element of these examples, as well as use t
 	Turn Based Strategy
 	[/PARAMLITERAL]
 	[PARAMLITERAL:playlistid]
-	PLZlv_N0_O1ga0aV9jVqJgog0VWz1cLL5f
+	PLZlv_N0_O1gYuWkyTSbfcuyPNOQsK_8bu
 	[/PARAMLITERAL]
 	[PARAMLITERAL:videoid]
 	mI7eYXMJ5eI
