UMG How-to : Creating/Displaying Widgets

diff --git a/Engine/Documentation/Source/Engine/UMG/HowTo/CreatingWidgets/CreatingWidgets.INT.udn b/Engine/Documentation/Source/Engine/UMG/HowTo/CreatingWidgets/CreatingWidgets.INT.udn
new file mode 100644
--- /dev/null
+++ b/Engine/Documentation/Source/Engine/UMG/HowTo/CreatingWidgets/CreatingWidgets.INT.udn
@@ -0,0 +1,142 @@
+Availability: Public
+Title:Creating and Displaying UI
+Crumbs: %ROOT%, Engine, Engine/UMG, Engine/UMG/HowTo
+Description:An example of how to create, display and remove UI elements from the screen during gameplay. 
+Related: Engine/UMG
+Related: Engine/UMG/UserGuide
+Version: 4.12
+SkillLevel: Beginner
+
+Typically in most games you will want to communicate some information to your players through your game's User Interface. 
+This can range from things like Main Menus, in-game Pause Menus, Health or Experience, Inventory Items or help text to instruct players what to do in a given situation. 
+
+With **Unreal Motion Graphics (UMG)**, you can create a [Widget Blueprint](Engine/UMG/UserGuide/WidgetBlueprints) to handle the display of UI elements in your project. 
+The Widget Blueprint uses [](Engine/Blueprints) to design the layout as well as script functionality for the UI element such as what happens when a button is clicked or if a value changes, etc. 
+
+In this how-to we will create a Widget Blueprint, add some Text to it, then through the Level Blueprint display it when a key is pressed. 
+
+
+## Steps
+
+[REGION:note]
+For this how-to, we are using the **Blueprint Third Person Template**. However, you can use any project you wish. 
+[/REGION]
+
+1.	In the [](Engine/Content/Browser/) click the **Add New** button, then under **User Interface** select **Widget Blueprint**. 
+
+	![](CreatingWidgets01.png)
+
+	Regardless of the type of UI element you wish to create in UMG, you will need to do so inside a Widget Blueprint. 
+
+1.	Enter the name **SampleUI** for the name of the Widget Blueprint, then **Double-click** and open it. 
+
+	![](CreatingWidgets02.png)
+
+	Upon opening a Widget Blueprint you will be presented with the UMG Editor Interface
+
+	[REGION:lightbox]
+	[![](CreatingWidgets03.png)](CreatingWidgets03.png)
+	[/REGION:lightbox]
+
+	[REGION:caption]
+	Click for full view. 
+	[/REGION]
+
+	The Widget Blueprint opens by default to the **Designer** tab which is where you perform your visual layout of your UI elements. To script functionality for the Widget Bluperint you can click the **Graph** button in the top right corner of the editor. 
+
+1.	Inside the **Palette** window, **Left-click** and drag a **Text** widget into the **Visual Designer** graph. 
+
+	![](CreatingWidgets04.png)
+
+	The Palette window contains several different types of widgets that you can use to construct your UI elements. 
+
+1.	In the **Details** panel for the Text Widget, click the [](Engine/UMG/UserGuide/Anchors) dropdown and select to Anchor to the Center. 
+
+	![](CreatingWidgets05.png)
+
+	This will ensure that as the viewport size changes, the Text will stay anchored relative to the center of the screen. 
+
+1.	In the **Details** panel click the **Size to Content** option, change the **Text**, **Color and Opacity** and **Font Size** options. 
+
+	![](CreatingWidgets06.png)
+
+	The Size to Content option as the name suggests will adjust the size of the widget to fit its content. We also changed the text to "Hello World!", changed the color of the text to Green and set the Font Size to 56.
+
+1.	Click the **Compile** button then **Save** and close the Widget Blueprint. 
+
+	![](CreatingWidgets09.png)
+
+1.	From the Main Toolbar, click the **Blueprints** button then select **Open Level Blueprint**. 
+
+	![](CreatingWidgets10.png)
+
+	We are now going to determine when to display the SampleUI Widget Blueprint we created. 
+
+1.	**Right-click** in the graph and search for and add the **Create Widget** node. 
+
+	![](CreatingWidgets11.png)
+
+1.	On the node, click the **Class** drop-down menu and search for and add our Widget Blueprint **SampleUI**. 
+
+	![](CreatingWidgets12.png)
+
+1.	**Right-click** on the **Return Value** pin and select **Promote to Variable** and call it **SampleUI**. 
+
+	![](CreatingWidgets13.png)
+
+	Whenever you create the Widget Blueprint (your UI), it's a good idea to promote that to a variable so you can access it later through Blueprint Script. 
+
+1.	Drag off the out pin on the **Set** node and search for and add the **Add to Viewport** node. 
+
+	![](CreatingWidgets14.png)
+
+	This is how we tell the Widget Blueprint to be displayed in-game. 
+
+1.	**Right-click** in the graph and search for and add the **Enter** Keyboard Event. 
+
+	![](CreatingWidgets15.png)
+
+1.	Connect the **Pressed** pin to the executable pin of the **Create Sample UI Widget** node. 
+
+	![](CreatingWidgets16.png)
+
+	At this point we have stated that when we press Enter, create the widget and store it as a variable before displaying it. 
+
+1.	Hold **Ctrl** and drag in the **SampleUI** variable into the graph from the **My Blueprint** panel.
+
+	![](CreatingWidgets17.png)
+
+1.	Drag off the pin of the new node and search for and add the **Remove from Parent** node. 
+
+	![](CreatingWidgets18.png)
+
+	The Remove from Parent node is how we tell Widget Blueprints to be removed from the viewport. 
+
+1.	Connect the **Release** pin of the **Enter** node to the **Remove from Parent** executable in pin. 
+
+	![](CreatingWidgets19.png)
+
+	Now when we press **Enter** the text will be displayed and when we release Enter, the text will be removed. 
+
+1.	Close the **Level Blueprint**, then click the **Play** button from the Main Toolbar. 
+
+	![](CreatingWidgets20.png)
+
+
+
+## End Result
+
+When you play in the editor you should have something similar to below whenever you press/release the **Enter** key. 
+
+[OBJECT:EmbeddedVideo]
+[PARAMLITERAL:width]
+640
+[/PARAMLITERAL]
+[PARAMLITERAL:height]
+360
+[/PARAMLITERAL]
+[PARAMLITERAL:videoid]
+cyUWhFts_3E
+[/PARAMLITERAL]
+[/OBJECT]
+
diff --git a/Engine/Documentation/Source/Engine/UMG/HowTo/UMGHowTo.INT.udn b/Engine/Documentation/Source/Engine/UMG/HowTo/UMGHowTo.INT.udn
--- a/Engine/Documentation/Source/Engine/UMG/HowTo/UMGHowTo.INT.udn
+++ b/Engine/Documentation/Source/Engine/UMG/HowTo/UMGHowTo.INT.udn
@@ -12,6 +12,9 @@ The **Unreal Motion Graphics UI Designer (UMG) - How To...** pages will show you
 Refer to the table below for more information and click each link to view its contents. 
 
 ## Beginner
+[REGION:wrapper transparent]
+[](FANCY:Engine\UMG\HowTo\CreatingWidgets "%Engine\UMG\HowTo\CreatingWidgets:description%")
+[/REGION]
 
 ## Intermediate
 [REGION:wrapper transparent]
