Adding images to Collaborate with Sublevels

diff --git a/Engine/Documentation/Source/Engine/Levels/CollaborateWithSublevels/CollaborateWithSublevels.INT.udn b/Engine/Documentation/Source/Engine/Levels/CollaborateWithSublevels/CollaborateWithSublevels.INT.udn
--- a/Engine/Documentation/Source/Engine/Levels/CollaborateWithSublevels/CollaborateWithSublevels.INT.udn
+++ b/Engine/Documentation/Source/Engine/Levels/CollaborateWithSublevels/CollaborateWithSublevels.INT.udn
@@ -10,15 +10,22 @@ order:5
 Levels are binary assets, which can lead to some challenges when collaborating with teammates working on the same area of the project.
 
 However, by setting up **Always Streaming** sublevels, multiple people can work in the same game "space" at the same time. For instance, you might have a sublevel that contains all 
-the lights for your scene, as well as a sublevel that contains just the Static Meshes laying out your environment. Each sublevel can still only have one person working on it at a 
+the audio for your scene, as well as a sublevel that contains just the Static Meshes laying out your environment. Each sublevel can still only have one person working on it at a 
 time, but with careful segmenting of your scene's Actors, you can minimize overlapping work areas. 
 
+For example, ShooterGame has **Always Streaming** sublevels for audio, collisions, gameplay items, lights, and the surrounding vista environment that the player can see but not interact with. 
+
+![](ShooterGameAlwaysLoaded.png)
+
+You can toggle on **Level Coloration** to visualize what Actors in your scene live in which currently loaded level. 
+
 [REGION:tip]
 You can even have sublevels' Level Blueprints communicate with each other, so that allows a switch in one sublevel, for instance, to toggle a light in another.
 [/REGION]
 
-In the **Levels** window, just right-click on any Level you want to have always loaded and hover over **Change Streaming Method** in the context menu. Then, click on **Always Loaded**. 
-The **Levels** window will also allow you to add more streaming Levels, move Actors between Levels, and change the debug color of Levels so you can visualize the setup more easily.
+In the [**Levels** window](Engine/Levels/HowTo/WorkWithLevelAssets), just right-click on any Level you want to have always loaded and hover over **Change Streaming Method** in the context menu. Then, click on **Always Loaded**. 
+
+![](AlwaysLoaded.png)
 
 [REGION:note]
 If your Levels are going to be always visible in game, make sure they are all toggled on in the **Levels** window before rebuilding lighting or navigation.
