Branch snapshot for CL 3059681

[CL 3059681 in Dev-VREditor branch]

diff --git a/Engine/Documentation/Source/Engine/Actors/Volumes/Volumes.INT.udn b/Engine/Documentation/Source/Engine/Actors/Volumes/Volumes.INT.udn
--- a/Engine/Documentation/Source/Engine/Actors/Volumes/Volumes.INT.udn
+++ b/Engine/Documentation/Source/Engine/Actors/Volumes/Volumes.INT.udn
@@ -163,7 +163,8 @@ The **Kill Z Volume** or "Kill-Brush" essentially destroys any Actor that enters
 
 ### Level Streaming Volumes
 
-**Level Streaming Volumes** are used to aid in the Level Streaming process. They provide for a simple way to encapsulate a level, as well as control when it streams in and out of memory, based on when a Player enters and exits the volume. 
+[EXCERPT:LSV]
+**Level Streaming Volumes** are used to aid in the [level streaming](Engine/LevelStreaming) process. They provide for a simple way to encapsulate a level, as well as control when it streams in and out of memory, based on when a Player enters and exits the volume. 
 
 You can adjust how a Level Streaming Volume handles level streaming by adjusting the properties from the **Details** panel, pictured below. 
 
@@ -176,6 +177,9 @@ You can adjust how a Level Streaming Volume handles level streaming by adjusting
 | **Disabled** | If true, the streaming volume is ignored by the streaming volume code. <br> Also used to either disable a Level Streaming Volume without disassociating it from the level, or to toggle the control of a level's streaming between Blueprints and Volume Streaming. |
 | **Streaming Usage** | Determines what the volume is used for, e.g. whether to control loading, loading and visibility, or just visibility (blocking on load). |
 
+[/EXCERPT]
+
+For more on Level Streaming Volumes, see [](Engine/LevelStreaming/StreamingVolumes).
 
 ### Cull Distance Volume
 
diff --git a/Engine/Documentation/Source/Engine/LevelStreaming/Overview/Overview.INT.udn b/Engine/Documentation/Source/Engine/LevelStreaming/Overview/Overview.INT.udn
--- a/Engine/Documentation/Source/Engine/LevelStreaming/Overview/Overview.INT.udn
+++ b/Engine/Documentation/Source/Engine/LevelStreaming/Overview/Overview.INT.udn
@@ -47,7 +47,7 @@ tags:Level Streaming
 [/OBJECT]
 [/VAR]
 
-[TOC]
+
 
 [INCLUDE:Engine/LevelStreaming#intro]
 
@@ -66,13 +66,27 @@ It ignores streaming volumes assigned to it, as well as any loading or unloading
 content that would ordinarily be in the persistent Level into multiple "layers", so artists can collaborate working in parallel without blocking each other. For more information on this workflow,
 see [](Engine/Levels/CollaborateWithSublevels). In game, an **Always Loaded** sublevel is never streamed out, unless the game changes the persistent Level. 
 
-Streaming Levels with the **Blueprint** streaming type can actually be controlled with Level Streaming Volumes, Blueprints, or C++ code. These Levels can be loaded
-or unloaded dynamically. 
 
+[OBJECT:ComparisonSlider]
+	[PARAM:before]
+	![Persistent Level Alone](PersistentLevel.png)(w:700)
+	[/PARAM]
+	[PARAM:after]
+	![After Streaming In Levels](StreamedInLevels.png)(w:700)
+	[/PARAM]
+[/OBJECT]
 
-## Level Streaming Volumes
+### Dynamic Streaming Methods
 
-**Level Streaming Volumes** make controlling level streaming very easy. The idea is simple: load/unload requests for a streaming Level are issued based on whether the viewpoint is 
+Streaming Levels with the **Blueprint** streaming type can actually be controlled with Level Streaming Volumes, Blueprints, or C++ code. These Levels can be loaded
+or unloaded dynamically. 
+
+[OBJECT:TopicList]
+	[PARAM:title]
+	Level Streaming Volumes
+	[/PARAM]
+	[PARAM:description]
+	**Level Streaming Volumes** make controlling level streaming very easy. The idea is simple: load/unload requests for a streaming Level are issued based on whether the viewpoint is 
 inside any of the Level Streaming Volumes associated with a Level.
 
 Specifically, Level Streaming Volumes can be used in two ways:
@@ -85,87 +99,29 @@ Specifically, Level Streaming Volumes can be used in two ways:
 
 Volume-based level streaming is simple to use and does not require any scripting, making it an ideal way to control level streaming. Furthermore, volume-based level streaming is 
 much easier to maintain than scripted streaming: when demands on the loading system change, Level loading/unloading behavior can be modified by simply resizing the streaming volumes.
+	[/PARAM]
+	[PARAM:links]
+	* [](Engine/LevelStreaming/StreamingVolumes "%Engine/LevelStreaming/StreamingVolumes:description%")
+	* [](Engine/LevelStreaming/HowTo/StreamWithVolumes "%Engine/LevelStreaming/HowTo/StreamWithVolumes:description%")
+	[/PARAM]
+[/OBJECT]
 
-### Associating Streaming Volumes With Levels
-
-Volume-based level streaming works as follows: each streaming Level can have associated with it a set of Level Streaming Volumes. Each frame, the engine iterates over 
-each Level and checks to see if the player's viewpoint is inside any of the Level Streaming Volumes associated with that Level. If the viewpoint is inside at least one 
-Level Streaming Volume, a request is issued to begin loading that Level. If the viewpoint is outside all Level Streaming Volumes, the Level is marked for unloading.
-
-* [](Engine/LevelStreaming/HowTo/StreamWithVolumes)
-
-### Important Details
-
-* All Level Streaming Volumes must exist in the persistent Level. Level Streaming Volumes that live in other Levels cannot be used for level streaming, and will generate warnings 
-when the map is checked for errors.
-
-* If a Level has any streaming volumes associated with it, other methods of streaming the Level will not behave correctly.
-
-* A single Level Streaming Volume can affect multiple Levels. Similarly, a single Level can be affected by multiple Level Streaming Volumes.
-
-* Volume-based streaming works for split screen. The viewpoints of all local players are considered before any loading/unloading requests are issued.
-
-
-
-
-### Testing Your Streaming Volume Setup
-
-It is critical that volume-based level streaming be tested in game on the target platform. Streaming in **Play in Editor** mode (PIE) will show where the loads/unloads will happen, 
-but streaming in PIE is not representative of real, in-game loading-unloading. This is because in PIE, the Levels are already in memory, and so "loading" a Level is simply a matter 
-of unhiding it instantly.
-
-Running the Level in the standalone game on the target platform is critical in making sure your streaming setups work. Note that on some platforms, it can take several seconds to 
-stream in a Level. Size your Level Streaming Volumes appropriately, so that the Level is loaded by the time the player can reach it. Level loading behavior can be modified by 
-resizing the Level Streaming Volumes associated with a Level. Growing a volume causes associated Levels to load sooner and unload later, while shrinking a volume causes later loads 
-and earlier unloads.
-
-
-### Level Streaming Volumes for Previs
-
-A Level Streaming Volume can be marked for editor previsualization only by setting the **Editor Pre Vis Only** flag on the Level Streaming Volume. In this manner, volume-based level streaming can 
-still be used for editor previs while at the same time using another streaming method for in-game streaming.
-
-
-### Cost of Level Streaming Volumes
-
-Each frame, `UWorld::ProcessLevelStreamingVolumes` iterates over each streaming Level and for each Level, that Level will begin loading if any local player is within any of the 
-volumes associated with that Level. Likewise, the Level will be begin unloading if all local players are outside all volumes.
-
-`UWorld::ProcessLevelStreamingVolumes` exploits coherency in the following manner: for each Level, the volume that most recently contained a player is cached.  This cached volume is checked first, so 
-that Levels a player is in or returns to are quickly accepted.
-
-Volumes of any shape are fine, although obviously the fewer the better.  An upper bound on the cost of Level Streaming Volumes can be approximated by the sum of the number of Level Streaming Volumes 
-associated with unloaded Levels.
-
-Two stats exist under the "Streaming" stats group for monitoring level streaming performance.  The "Streaming Volumes" stat tracks the number of Level Streaming Volumes tested against player viewpoints 
-per frame, while the "Volume Streaming Tick" stat tracks the amount of time spent in `UWorld::ProcessLevelStreamingVolumes` per frame.
-
-
-
-### Adding Hysteresis to Unloading Requests
-
-A player moving back and forth across a Level Streaming Volume boundary causes unneeded load/unload requests to be issued. To address this, hysteresis has been added to unloading 
-requests. No hysteresis exists for loading requests, because if a Level needs to be loaded, we always want it loaded as soon as possible.
-
-The amount of unloading hysteresis can be adjusted by modifying the **Min Time Between Volume Unload Requests** property of the streaming Level in the [Levels window](Engine/Levels/LevelsWindow). 
-The default unloading hysteresis is 2.0 seconds.
-
-
-
-### Disabling Level Streaming Volumes
-
-There is a property on Level Streaming Volumes called **bDisabled**. When set to true, the volume will be ignored by the streaming volume code, both in-game and in the editor. **bDisabled** 
-can be used to disable a Level Streaming Volume without dissociating it from the Level.  
-
-As an example of where the **bDisabled** flag is useful, imagine a door leading to a Level whose streaming is controlled with streaming volumes. The streaming volume extends out past 
-the door, so that the Level will be streamed in by the time the player can reach the door and open it. Initially, however, the door is locked, and will become unlocked when the 
-player achieves an objective in another part of the Level. So, even though the streaming volume extends past the door, we don't want the Level on the other side of the door to 
-stream in unless the door is actually unlocked ("openable").
 
-Blueprints and C++ could both be used to toggle the disabled state of **Level Streaming Volumes**.
 
-## Programmatic Level Streaming
+[OBJECT:TopicList]
+	[PARAM:title]
+	Scripted Level Streaming
+	[/PARAM]
+	[PARAM:description]
+	If you want to set up more complex behavior for loading and unloading Levels, similar to the door unlocking example above, you can also set up Blueprint or C++ logic for Level 
+loading and unloading. The key functions are **Load Stream Level** and **Unload Stream Level**, and you designate the Level to load by its name. 
 
-If you want to set up more complex behavior for loading and unloading Levels, similar to the door unlocking example above, you can also set up Blueprint or C++ logic for Level 
-loading and unloading. The key functions are **Load Stream Level** and **Unload Stream Level**, and you designate the Level to load by its name. See 
-[](Engine/LevelStreaming/HowTo/StreamWithBlueprints) and [](Engine/LevelStreaming/HowTo/StreamWithC++) for examples. 
\ No newline at end of file
+In addition to **Load Stream Level** and **Unload Stream Level**, the **Get Streaming Level** function enables even more dynamic streaming behavior. Since it gives you access
+to the actual level streaming object, you can modify and query its state. Using this function in combination with **Create Instance**, you can also create and stream in copies of a specific sublevel.
+By applying transforms and sending parameters to these copies, you can create procedural worlds. 
+	[/PARAM]
+	[PARAM:links]
+	* [](Engine/LevelStreaming/HowTo/StreamWithBlueprints "%Engine/LevelStreaming/HowTo/StreamWithBlueprints:description%")
+	* [](Engine/LevelStreaming/HowTo/StreamWithC++ "%Engine/LevelStreaming/HowTo/StreamWithC++:description%")
+	[/PARAM]
+[/OBJECT]
\ No newline at end of file
diff --git a/Engine/Documentation/Source/Engine/LevelStreaming/StreamingVolumes/StreamingVolumes.INT.udn b/Engine/Documentation/Source/Engine/LevelStreaming/StreamingVolumes/StreamingVolumes.INT.udn
--- a/Engine/Documentation/Source/Engine/LevelStreaming/StreamingVolumes/StreamingVolumes.INT.udn
+++ b/Engine/Documentation/Source/Engine/LevelStreaming/StreamingVolumes/StreamingVolumes.INT.udn
@@ -1,150 +1,89 @@
 Availability: Docs
 Title:Level Streaming Volumes
 Crumbs:
-Description:Guide to using volumes to control the streaming of levels based on the player's location.
-
-[REGION:todo]
-	This page is converted from Unreal Engine 3 documentation and currently being reviewed. Information may be outdated.
-[/REGION]
-
-[TOC(start:2)]
-
-
-
-## Overview
-
-This document contains an overview of how to use volume-based level streaming.  It assumes the reader is already familiar with the basics of level streaming described in the 
-[Level Streaming How-To](Engine/LevelStreaming/UserGuide).
-
-Level streaming volumes make controlling level streaming very easy.  The idea is simple: load/unload requests for a streaming level are issued based on whether the viewpoint is inside any of the 
-LevelStreamingVolume volumes associated with a level.
-
-Specifically, level streaming volumes can be used in two ways:
-
-
-* In the **game**, level streaming volumes cause levels to load when the player is inside the volume and unload when the player is outside the volume.
-
-* In the **editor**, level streaming volumes can be used to preview level streaming by automatically hiding/unhiding levels based on the location of the perspective viewport's camera.
-
-
-Volume-based level streaming is simple to use and does not require any scripting, making it an ideal way to control level streaming.  Furthermore, volume-based level streaming is much easier to 
-maintain than scripted streaming: when demands on the loading system change, level loading/unloading behaviour can be modified by simply resizing the streaming volumes.
+Description:Guide to using volumes to control the streaming of Levels based on the player's viewpoint.
+Version: 4.12
+parent:Engine/LevelStreaming
+type:reference
+order:1
+tags:Level Streaming
 
+[TOC]
 
+[INCLUDE:Engine/Actors/Volumes#lsv]
 
 ## Associating Streaming Volumes With Levels
 
-Volume-based level streaming works as follows: each streaming level can have associated with it a set of volumes of type LevelStreamingVolume.  Each frame, the engine iterates over each level and 
-checks to see if the player viewpoint is inside any of the LevelStreamingVolumes associated with that level.  If the viewpoint is inside at least one LevelStreamingVolume, a request is issued to 
-begin loading that level.  If the viewpoint is outside all LevelStreamingVolumes, the level is marked for unloading.
-
-To associate streaming volumes with a level:
-
-1. Drop a LevelStreamingVolume in the Persistent Level (as you would any other volume by using the "Add Volume" button on the main editor toolbar);
-1. Select the level(s) to associate with this volume in the Level Browser, right-click to bring up the Level Browser context menu, and select "Add Streaming Volumes".
-
-That's it!
-
-The full list of LevelStreamingVolume operations appearing in the Level Browser context menu is as follows:
-
-![LevelBrowserMenu.jpg](LevelBrowserMenu.jpg)
-
-
-* **Add Streaming Volumes**: Adds to the level(s)'s volumes, retaining any existing associations.
-
-* **Set Streaming Volumes**: Sets the level(s)'s volumes to be exactly the selected volumes, destroying any existing associations.
-
-* **Clear Streaming Volumes**: Clears all streaming volume associations from the selected level(s).
-
-* **Select Associated Streaming Volumes**: Selects all volumes associated with the selected level(s).
-
-
-Since it's easy for streaming volumes to clutter the view, LevelStreamingVolumes are only visible for levels that are **1) visible**; and **2) selected** in the level browser.  So, if you deselect the 
-level you've just associated that volume with, you'll see the volume disappear in the editor viewports.  Selecting the level again in the Level Browser will cause the associated volumes to reappear.
-
+Volume-based level streaming works as follows: each streaming Level can have associated with it a set of Level Streaming Volumes. Each frame, the engine iterates over 
+each Level and checks to see if the player's viewpoint is inside any of the Level Streaming Volumes associated with that Level. If the viewpoint is inside at least one 
+Level Streaming Volume, a request is issued to begin loading that Level. If the viewpoint is outside all Level Streaming Volumes, the Level is marked for unloading.
 
+* [](Engine/LevelStreaming/HowTo/StreamWithVolumes)
 
 ## Important Details
 
+* All Level Streaming Volumes must exist in the persistent Level. Level Streaming Volumes that live in other Levels cannot be used for level streaming, and will generate warnings 
+when the map is checked for errors.
 
-* All LevelStreamingVolumes must exist in the persistent level.  LevelStreamingVolumes that live in other levels cannot be used for level streaming, and will generate warnings when the map is checked 
-for errors.
-
-* Only levels that were added to the world using the "Blueprint" streaming method can use volume-based level streaming.
-
-* If a level has any streaming volumes associated with it, other methods of streaming the level (ie distance streaming or Blueprint streaming) will not behave correctly.
+* If a Level has any streaming volumes associated with it, other methods of streaming the Level will not behave correctly.
 
-* A single LevelStreamingVolume can affect multiple levels.  Similarly, a single level can be affected by multiple LevelStreamingVolumes.
-
-* Volume-based streaming works for split screen.  The viewpoints of all local players are considered before any loading/unloading requests are issued.
+* A single Level Streaming Volume can affect multiple Levels. Similarly, a single Level can be affected by multiple Level Streaming Volumes.
 
+* Volume-based streaming works for split screen. The viewpoints of all local players are considered before any loading/unloading requests are issued.
 
 
 
 ## Testing Your Streaming Volume Setup
 
-It is critical that volume-based level streaming be tested in game on the target platform.  Streaming in "Play in Editor" (PIE) will show <i>where</i> the loads/unloads will happen, but **streaming in 
-PIE is not representative of real, in-game loading-unloading**.  This is because in PIE, the levels are already in memory, and so "loading" a level is simply a matter of unhiding it and is thus instantaneous.
-
-Running the level in the standalone game on the target platform is critical in making sure your streaming setups work.  Note that on some platforms, it can take several seconds to stream in a level.  
-Size your LevelStreamingVolumes appropriately, so that the level is loaded by the time the player can reach it.  Level loading behaviour can be modified by resizing the LevelStreamingVolumes associated 
-with a level.  Growing a volume causes associated levels to load sooner and unload later, while shrinking a volume causes later loads and earlier unloads.
+It is critical that volume-based level streaming be tested in game on the target platform. Streaming in **Play in Editor** mode (PIE) will show where the loads/unloads will happen, 
+but streaming in PIE is not representative of real, in-game loading-unloading. This is because in PIE, the Levels are already in memory, and so "loading" a Level is simply a matter 
+of unhiding it instantly.
 
+Running the Level in the standalone game on the target platform is critical in making sure your streaming setups work. Note that on some platforms, it can take several seconds to 
+stream in a Level. Size your Level Streaming Volumes appropriately, so that the Level is loaded by the time the player can reach it. Level loading behavior can be modified by 
+resizing the Level Streaming Volumes associated with a Level. Growing a volume causes associated Levels to load sooner and unload later, while shrinking a volume causes later loads 
+and earlier unloads.
 
 
 ## Level Streaming Volumes for Previs
 
-![PrevisButton.jpg](PrevisButton.jpg)
-
-Volume-based level streaming can be easily previewed by enabling the "Level Streaming Volume Previs" button on the perspective viewport toolbar.  When this mode is enabled, levels will be hide/unhide 
-based on the location of the perspective viewport camera.  Note that this hiding/unhiding will affect the set of visible levels in the Level Browser.
-
-A LevelStreamingVolume can be marked for editor previs only by setting the bEditorPreVisOnly flag on the LevelStreamingVolume.  In this manner, volume-based level streaming can still be used for editor 
-previs while at the same time using another streaming method for in-game streaming.
-
+A Level Streaming Volume can be marked for editor previsualization only by setting the **Editor Pre Vis Only** flag on the Level Streaming Volume. In this manner, volume-based level streaming can 
+still be used for editor previs while at the same time using another streaming method for in-game streaming.
 
 
 ## Cost of Level Streaming Volumes
 
-Each frame, UWorld::ProcessLevelStreamingVolumes iterates over each streaming level and for each level, that level will begin loading if any local player is within any of the volumes associated with 
-that level.  Likewise, the level will be begin unloading if all local players are outside all volumes.
+Each frame, `UWorld::ProcessLevelStreamingVolumes` iterates over each streaming Level and for each Level, that Level will begin loading if any local player is within any of the 
+volumes associated with that Level. Likewise, the Level will be begin unloading if all local players are outside all volumes.
 
-UWorld::ProcessLevelStreamingVolumes exploits coherency in the following manner: for each level, the volume that most recently contained a player is cached.  This cached volume is checked first, so 
-that levels a player is in or returns to are quickly accepted.
+`UWorld::ProcessLevelStreamingVolumes` exploits coherency in the following manner: for each Level, the volume that most recently contained a player is cached.  This cached volume is checked first, so 
+that Levels a player is in or returns to are quickly accepted.
 
-Volumes of any shape are fine, although obivously the fewer the better.  An upper bound on the cost of level streaming volumes can be approximated by the sum of the number of level streaming volumes 
-associated with unloaded levels.
+Volumes of any shape are fine, although obviously the fewer the better.  An upper bound on the cost of Level Streaming Volumes can be approximated by the sum of the number of Level Streaming Volumes 
+associated with unloaded Levels.
 
-Two stats exist under the "Streaming" stats group for monitoring level streaming performance.  The "Streaming Volumes" stat tracks the number of LevelStreamingVolumes tested against player viewpoints 
-per frame, while the "Volume Streaming Tick" stat tracks the amount of time spent in UWorld::ProcessLevelStreamingVolumes per frame.
+Two stats exist under the "Streaming" stats group for monitoring level streaming performance.  The "Streaming Volumes" stat tracks the number of Level Streaming Volumes tested against player viewpoints 
+per frame, while the "Volume Streaming Tick" stat tracks the amount of time spent in `UWorld::ProcessLevelStreamingVolumes` per frame.
 
 
 
 ## Adding Hysteresis to Unloading Requests
 
-A player moving back and forth across a LevelStreamingVolume boundary causes spurious load/unload requests to be issued.  To address this, hysteresis has been added to unloading requests.   No hysteresis 
-exists for loading requests, because if a level needs to be loaded, we always want it loaded ASAP.
-
-The amount of unloading hystersis can be adjusted by modifying the MinTimeBetweenVolumeUnloadRequests property of the LevelStreaming object (LevelStreaming properties can be accessed from the Level 
-Browser context menu).  The default unloading hysteresis is 2.0 seconds.
+A player moving back and forth across a Level Streaming Volume boundary causes unneeded load/unload requests to be issued. To address this, hysteresis has been added to unloading 
+requests. No hysteresis exists for loading requests, because if a Level needs to be loaded, we always want it loaded as soon as possible.
 
+The amount of unloading hysteresis can be adjusted by modifying the **Min Time Between Volume Unload Requests** property of the streaming Level in the [Levels window](Engine/Levels/LevelsWindow). 
+The default unloading hysteresis is 2.0 seconds.
 
 
 ## Disabling Level Streaming Volumes
 
-There is a property on LevelStreamingVolumes called bDisabled.  When set to true, the volume will be ignored by the streaming volume code, both in-game and in the editor.  bDisabled can be used to disable 
-a LevelStreamingVolume without disassociating it from the level.  
-
-As an example of where the bDisabled flag is useful, imagine a door leading to a level whose streaming is controlled with streaming volumes.
-
-![bDisabledUsage.jpg](bDisabledUsage.jpg)
-
-In the image above, the player approaches from the left, and the level to stream in is on the right side of the door.  The streaming volume extends leftwards out past the door, so that the level will be 
-streamed in by the time the player can reach the door and open it.  Initially, however the door is locked, and will become unlocked when the player achieves an objective in another part of the map.  
-So, even though the streaming volume extends past the door, we don't want the level on the other side of the door to stream in unless the door is actually unlocked ("openable").
-
- 
-
-
+There is a property on Level Streaming Volumes called **bDisabled**. When set to true, the volume will be ignored by the streaming volume code, both in-game and in the editor. **bDisabled** 
+can be used to disable a Level Streaming Volume without dissociating it from the Level.  
 
+As an example of where the **bDisabled** flag is useful, imagine a door leading to a Level whose streaming is controlled with streaming volumes. The streaming volume extends out past 
+the door, so that the Level will be streamed in by the time the player can reach the door and open it. Initially, however, the door is locked, and will become unlocked when the 
+player achieves an objective in another part of the Level. So, even though the streaming volume extends past the door, we don't want the Level on the other side of the door to 
+stream in unless the door is actually unlocked ("openable").
 
+Blueprints and C++ could both be used to toggle the disabled state of **Level Streaming Volumes**.
\ No newline at end of file
