#UE4 Docs: Checking in the fished version of the Vulkan document.

diff --git a/Engine/Documentation/Source/Platforms/Android/VulkanMobileRenderer/VulkanMobileRender.INT.udn b/Engine/Documentation/Source/Platforms/Android/VulkanMobileRenderer/VulkanMobileRender.INT.udn
--- a/Engine/Documentation/Source/Platforms/Android/VulkanMobileRenderer/VulkanMobileRender.INT.udn
+++ b/Engine/Documentation/Source/Platforms/Android/VulkanMobileRenderer/VulkanMobileRender.INT.udn
@@ -12,14 +12,136 @@
 
 [TOC(start:2)]
 
-Unreal Engine 4 (UE4) now offers support for the **Vulkan API**. Vulkan is a low-overhead, cross-platform 3D graphics library that allows developers more control over the GPU and lower level CPU usage for their Android base mobile projects. In the following document we will go over how you can enable and use Vulkan in your UE4 Android projects.
+Unreal Engine 4 (UE4) has built in support for the **Vulkan** graphics API. Vulkan is a low-overhead, cross-platform 3D graphics library that gives developers more control over the GPU and lower level CPU usage for their Android base mobile projects. In the following document, we will go over how you can enable and use Vulkan in your UE4 Android projects.
 
 [region:warning]
-Vulkan is currently under heavy development and should be considered an **experimental feature** at this point in time. If you wish to build and deploy a mobile project using Vulkan you will need to download and compile UE 4.12 from Github. 
+Vulkan is currently under heavy development and should be considered a **experimental feature** at this point in time. If you wish to build and deploy a mobile project using Vulkan, you will need to download and compile UE4 version 4.12 or later from the UE4 Github branch. 
 [/region]
 
+## Vulkan Video Drivers for PC Development
+To make sure that you can visualize the rendering options Vulkan has available on your development PC, you will need to make sure you download and install the latest version of video card drivers for your graphics card. Below you will find the minimal driver version you need to use to be able to preview what Vulkan will look like on your development PC.
+
+* **NVIDIA**: Driver 367.27 or later
+* **AMD**: Driver 16.3 or later  
+
+## Supported Vulkan Devices
+Currently, only the following mobile device chipsets have support for the Vulkan API. 
+
+* **Nexus 6P** (Android N Devices)
+* **Galaxy S7** (Mali & Adreno variants)
+
+[region:note]
+Please note that the ability to use the Vulkan API depends on if your mobile carrier has released the Vulkain update for your specific device variant. To find if this support has been pushed to your mobile device, you will need to contact your mobile carrier.
+[/region]
+
+## Building for Vulkan 
+To build a UE4 project that has support for the Vulkan API you will need to do the following:
+
+1. Before you begin, make sure that you have your Android smartphone plugged into your development PC via USB and that your Android smartphone has developer mode enabled. 
+
+1. Next, download and compile the latest version of UE4 from the UE4 Github branch as the binary version you get from the launcher will not allow you to build for Vulkan.
+
+	[region:note]
+	 If you are unsure how to do this check out the [Download Unreal Engine Source Code](GettingStarted/DownloadingUnrealEngine/) page for detailed instruction on what to do.
+	[/region]
+
+1. Then go to [https://vulkan.lunarg.com/signin](https://vulkan.lunarg.com/signin) and in the download section, click on the **Windows** option.
+
+	[REGION:lightbox]
+	[![](VMP_LunarG_00.png)(w:469 h:577)](VMP_LunarG_00.png)
+	[/REGION]
+
+	[region:note] 
+	You do not need to sign in or create an account to download the Vulkan SDK.
+	[/region]
+
+1. From the LunarG Vulkan SDK License Agreement page, change the version from **Latest** and then press the **I Accept** button to begin the download.
+
+	[REGION:lightbox]
+	[![](VMP_LunarG_01.png)(w:463 h:224)](VMP_LunarG_01.png)
+	[/REGION]
+
+1. Once the download completes, run the **VulkanSDK-1.0.17.0** installer and follow the prompts to install the Vulkan SDK. When the Vulkan SDK has finished installing, reboot your PC to make sure that the environment variables are set properly. 
+
+1. When your PC finished rebooting open up your system's **Environment Variables** and check for the **Vulkan_SDK** in the **Variable** section. 
+
+	![](VMP_Env_Vars.png)
+
+1. Once UE4 has finshed building, Launch the UE4 Editor and create a new project using the following settings:
+
+	* Blueprint Project
+	* Blank Template
+	* Mobile / Tablet
+	* Maximum Quality
+	* No Starter Content
+
+	[REGION:lightbox]
+	[![](VMP_New_Project_Settings.png)(w:555 h:387)](VMP_New_Project_Settings.png)
+	[/REGION]
+
+	When that is completed, press the **Create** button to create and load the new project.
+
+1. Once the project has loaded go to **Edit** > **Project Settings** then under **Engine** go to the **Rendering** section and make sure that **Mobile HDR** has been enabled.
+
+	[REGION:lightbox]
+	[![](VMP_Android_Settings_04.png)(w:639 h:385)](VMP_Android_Settings_04.png)
+	[/REGION]
+
+1. Next, locate the **Android** section which is under **Platforms**.
+	
+	[REGION:lightbox]
+	[![](VMP_Android_Settings_00.png)(w:639 h:385)](VMP_Android_Settings_00.png)
+	[/REGION]
+
+1. In the Android section under **APKPackaging** set the **Minimum SDK Version** to **19** and set the **Target SDK Version** to **19**. 
+
+	[REGION:lightbox]
+	[![](VMP_Android_Settings_01.png)(w:639 h:385)](VMP_Android_Settings_01.png)
+	[/REGION]
+
+1. Below the APKPackaging section in the **Build** section make sure that the following options are enabled.
+
+	* **Support arm7[aka armebi-v7a]**
+	* **Support OpenGL ES2**
+	* **Support Vulkan[Experimental]**
+
+	[REGION:lightbox]
+	[![](VMP_Android_Settings_02.png)(w:639 h:385)](VMP_Android_Settings_02.png)
+	[/REGION]
+
+1. In the **Android SDK** section make sure that **SDKConfig** paths are pointed to the correct SDK on your development PC.
+
+	[REGION:lightbox]
+	[![](VMP_Android_Settings_03.png)(w:639 h:385)](VMP_Android_Settings_03.png)
+	[/REGION]
+
+1. From the menu bar go to **File** > **Package Project** > **Android** and select the **Android(ASTC)** option. 
+	
+	[REGION:lightbox]
+	[![](VMP_Package_Game.png)(w:329 h:266)](VMP_Package_Game.png)
+	[/REGION]
+
+1. Select a location for UE4 to save the Android build and then press the **OK** button to start the packaging process.
+
+	[REGION:lightbox]
+	[![](VMP_Save_Build_Location.png)(w:245 h:267)](VMP_Save_Build_Location.png)
+	[/REGION]
+	![]()
+
+1. When the packaging process has been completed, open up the folder that the packaged build was placed in. Inside this folder, you should see two .BAT files, locate the .BAT file that has the word **Install** in the name and double - click on it to install it to your device.
+	
+	[REGION:lightbox]
+	[![](VMP_Build_Files.png)(w:439 h:198)](VMP_Build_Files.png)
+	[/REGION]
+
+1. Once installed on your device, press on the UE4 icon that has your project name under it to launch the project on the device. 
+
+	[REGION:lightbox]
+	[![](VMP_On_Device.png)(w:818 h:320)](VMP_On_Device.png)
+	[/REGION]
+
 ## Enabling Vulkan Preview Rendering in Editor
-You can enable the ability to use the Vulkan render while inside of the UE4 Editor, you will need to do the following in your project.
+You can enable the ability to use the Vulkan render while inside of the UE4 Editor, you will need to do the following in your project:
 
 1. From the **Main Toolbar** go to the **Edit** option and then select the **Editor Preferences** option from the menu.
 
@@ -35,14 +157,14 @@
 	After doing this, the Editor will need to recompile the entire shader cache to include the required Vulkan options. Depending on the size of your project and power of your development computer, this process could take anywhere from a few minutes to an hour or more to complete.
 	[/region]
 
-1. If you see the following two items listed in the lower right-hand corner of the Level preview **Feature Level: ES3_1** and **Level: Vulkin_Rendering(Persistent)** then everything is setup correctly and you are now seeing what your project will look like when viewed on a device that supports the Vulkan API.
+1. If you see the following two items listed in the lower right-hand corner of the Level preview **Feature Level: ES3_1** and **Level: Vulkan_Rendering(Persistent)** then everything is setup correctly, and you are now seeing what your project will look like when viewed on a device that supports the Vulkan API.
 
 	[REGION:lightbox]
 	[![](VMP_04.png)(w:447 h:308)](VMP_04.png)
 	[/REGION]
 
 ## Enabling the Vulkan Mobile Preview Render
-To enable the ability to use the Vulkan render with the mobile preview render, you will need to do the following in your project.
+To enable the ability to use the Vulkan render with the mobile preview render, you will need to do the following in your project:
 
 1. From the **Main Toolbar** go to the **Edit** option and then select the **Editor Preferences** option from the menu.
 
@@ -55,7 +177,7 @@
 	[/REGION]
 
 
-1. Back in the Main Toolbar, change the **Play Mode** from **Play this level in the selected viewport** to the **Vulkan Mobile Preview (PIE)** option by clicking on the small white triangle next to the play button and selecting the Vulkan Mobile Preview (PIE) from the dropdown list.
+1. Back in the Main Toolbar, change the **Play Mode** from **Play this level in the selected viewport** to the **Vulkan Mobile Preview (PIE)** option by clicking on the small white triangle next to the play button and selecting the Vulkan Mobile Preview (PIE) from the drop down list.
 
 	![](VMP_02.png)
 
@@ -71,74 +193,3 @@
 
 
 
-
-
-
-
-
-
-
-
-
-
-
-
-
-
-
-
-
-
-
-
-
-
-
-
-
-
-
-
-
-
-
-
-
-
-
-
-
-
-
-
-
-
-
-
-
-
-
-
-
-
-
-
-
-
-
-
-
-
-
-
-
-
-
-
-
-
-
-
-
-
-
