#UE4 docs misc fixes

diff --git a/Engine/Documentation/Source/Engine/Animation/NodeReference/SkeletalControls/AnimDynamics/AnimDynamics.INT.udn b/Engine/Documentation/Source/Engine/Animation/NodeReference/SkeletalControls/AnimDynamics/AnimDynamics.INT.udn
--- a/Engine/Documentation/Source/Engine/Animation/NodeReference/SkeletalControls/AnimDynamics/AnimDynamics.INT.udn
+++ b/Engine/Documentation/Source/Engine/Animation/NodeReference/SkeletalControls/AnimDynamics/AnimDynamics.INT.udn
@@ -112,7 +112,7 @@ To quickly get started, create an **AnimDynamics** node inside the AnimGraph of
 | [REGION:tablesection]**Output Pins**[/REGION] ||
 | Pose (Blank) | The final pose after the Transform has been applied. |
 
-**AnimDynamics** nodes are also affected by the **LOD Threshold** system introduced in 4.11, you can find the setting for this in the Details Panel for the **AnimDynamics** node.
+**AnimDynamics** nodes are also affected by the **LOD Threshold** system introduced in 4.11. You can find the setting for this in the **Details Panel** for the **AnimDynamics** node.
 
 [INCLUDE:Engine/Animation/NodeReference/SkeletalControls#lodthresh]
 
diff --git a/Engine/Documentation/Source/Engine/Animation/NodeReference/SkeletalControls/ApplyPercentRotation/ApplyPercentRotation.INT.udn b/Engine/Documentation/Source/Engine/Animation/NodeReference/SkeletalControls/ApplyPercentRotation/ApplyPercentRotation.INT.udn
--- a/Engine/Documentation/Source/Engine/Animation/NodeReference/SkeletalControls/ApplyPercentRotation/ApplyPercentRotation.INT.udn
+++ b/Engine/Documentation/Source/Engine/Animation/NodeReference/SkeletalControls/ApplyPercentRotation/ApplyPercentRotation.INT.udn
@@ -77,6 +77,6 @@ By changing the **Rotation Axis to Refer** value we can alter which way the pack
 | [REGION:tablesection]**Input Pins**[/REGION] ||
 | **Multiplier** | A float value that represents the percentage of the source bone's Rotation to apply to the target bone. When **Multiplier** is 0, none of the source bone's Rotation is applied to the target bone, and when it is 1.0, all of the source bone's Rotation is applied. |
 
-**Apply a Percentage of Rotation** nodes are also affected by the **LOD Threshold** system introduced in 4.11, you can find the setting for this in the Details Panel for the **Apply a Percentage of Rotation** node.
+**Apply a Percentage of Rotation** nodes are also affected by the **LOD Threshold** system introduced in 4.11. You can find the setting for this in the **Details Panel** for the **Apply a Percentage of Rotation** node.
 
 [INCLUDE:Engine/Animation/NodeReference/SkeletalControls#lodthresh]
diff --git a/Engine/Documentation/Source/Engine/Animation/NodeReference/SkeletalControls/BoneDrivenController/BoneDrivenController.INT.udn b/Engine/Documentation/Source/Engine/Animation/NodeReference/SkeletalControls/BoneDrivenController/BoneDrivenController.INT.udn
--- a/Engine/Documentation/Source/Engine/Animation/NodeReference/SkeletalControls/BoneDrivenController/BoneDrivenController.INT.udn
+++ b/Engine/Documentation/Source/Engine/Animation/NodeReference/SkeletalControls/BoneDrivenController/BoneDrivenController.INT.udn
@@ -75,7 +75,7 @@ CDG4WcQ7cIY
 | **(As Pin) Alpha** | Enable/Disable the display of the Alpha pin on the node. |
 | **Alpha Scale Bias** | Set the minimal and maximal input scaling values. |
 
-**Bone Driven Controller** nodes are also affected by the **LOD Threshold** system introduced in 4.11, you can find the setting for this in the Details Panel for the **Bone Driven Controller** node.
+**Bone Driven Controller** nodes are also affected by the **LOD Threshold** system introduced in 4.11. You can find the setting for this in the **Details Panel** for the **Bone Driven Controller** node.
 
 [INCLUDE:Engine/Animation/NodeReference/SkeletalControls#lodthresh]
 
diff --git a/Engine/Documentation/Source/Engine/Animation/NodeReference/SkeletalControls/CopyBone/CopyBone.INT.udn b/Engine/Documentation/Source/Engine/Animation/NodeReference/SkeletalControls/CopyBone/CopyBone.INT.udn
--- a/Engine/Documentation/Source/Engine/Animation/NodeReference/SkeletalControls/CopyBone/CopyBone.INT.udn
+++ b/Engine/Documentation/Source/Engine/Animation/NodeReference/SkeletalControls/CopyBone/CopyBone.INT.udn
@@ -78,6 +78,6 @@ Be cautious when selecting the checkboxes for Transform. For instance, if you en
 | **Copy Rotation** | A bool specifying whether to copy the Rotation component of the source bone's Transform. |
 | **Copy Scale** | A bool specifying whether to copy the Scale component of the source bone's Transform. |
 
-**Copy Bone** nodes are also affected by the **LOD Threshold** system introduced in 4.11, you can find the setting for this in the Details Panel for the **Copy Bone** node.
+**Copy Bone** nodes are also affected by the **LOD Threshold** system introduced in 4.11. You can find the setting for this in the **Details Panel** for the **Copy Bone** node.
 
 [INCLUDE:Engine/Animation/NodeReference/SkeletalControls#lodthresh]
diff --git a/Engine/Documentation/Source/Engine/Animation/NodeReference/SkeletalControls/HandIKRetargeting/HandIKRetargeting.INT.udn b/Engine/Documentation/Source/Engine/Animation/NodeReference/SkeletalControls/HandIKRetargeting/HandIKRetargeting.INT.udn
--- a/Engine/Documentation/Source/Engine/Animation/NodeReference/SkeletalControls/HandIKRetargeting/HandIKRetargeting.INT.udn
+++ b/Engine/Documentation/Source/Engine/Animation/NodeReference/SkeletalControls/HandIKRetargeting/HandIKRetargeting.INT.udn
@@ -83,6 +83,6 @@ See the **Animation** map inside the **Content Examples** project (available for
 | **Component Pose** | The incoming pose. |
 | **Alpha** | The current strength of the control. |
 
-**Hand IK Retargeting** nodes are also affected by the **LOD Threshold** system introduced in 4.11, you can find the setting for this in the Details Panel for the **Hand IK Retargeting** node.
+**Hand IK Retargeting** nodes are also affected by the **LOD Threshold** system introduced in 4.11. You can find the setting for this in the **Details Panel** for the **Hand IK Retargeting** node.
 
 [INCLUDE:Engine/Animation/NodeReference/SkeletalControls#lodthresh]
diff --git a/Engine/Documentation/Source/Engine/Animation/NodeReference/SkeletalControls/LookAt/LookAt.INT.udn b/Engine/Documentation/Source/Engine/Animation/NodeReference/SkeletalControls/LookAt/LookAt.INT.udn
--- a/Engine/Documentation/Source/Engine/Animation/NodeReference/SkeletalControls/LookAt/LookAt.INT.udn
+++ b/Engine/Documentation/Source/Engine/Animation/NodeReference/SkeletalControls/LookAt/LookAt.INT.udn
@@ -89,6 +89,6 @@ You may need to enlarge the video to see the debug line or try it out in the edi
 | **Component Pose** | The incoming pose. |
 | **Alpha** | The current strength of the control. |
 
-**Look At** nodes are also affected by the **LOD Threshold** system introduced in 4.11, you can find the setting for this in the Details Panel for the **Look At** node.
+**Look At** nodes are also affected by the **LOD Threshold** system introduced in 4.11. You can find the setting for this in the **Details Panel** for the **Look At** node.
 
 [INCLUDE:Engine/Animation/NodeReference/SkeletalControls#lodthresh]
diff --git a/Engine/Documentation/Source/Engine/Animation/NodeReference/SkeletalControls/ObserveBone/ObserveBone.INT.udn b/Engine/Documentation/Source/Engine/Animation/NodeReference/SkeletalControls/ObserveBone/ObserveBone.INT.udn
--- a/Engine/Documentation/Source/Engine/Animation/NodeReference/SkeletalControls/ObserveBone/ObserveBone.INT.udn
+++ b/Engine/Documentation/Source/Engine/Animation/NodeReference/SkeletalControls/ObserveBone/ObserveBone.INT.udn
@@ -58,7 +58,7 @@ Below, a bone by the name of **b_MF_Hand_L** is being observed and debug informa
 | [REGION:tablesection]**Input Pins**[/REGION] ||
 | **Component Pose** | The incoming pose. |
 
-**Observe Bone** nodes are also affected by the **LOD Threshold** system introduced in 4.11, you can find the setting for this in the Details Panel for the **Observe Bone** node.
+**Observe Bone** nodes are also affected by the **LOD Threshold** system introduced in 4.11. You can find the setting for this in the **Details Panel** for the **Observe Bone** node.
 
 [INCLUDE:Engine/Animation/NodeReference/SkeletalControls#lodthresh]
 
diff --git a/Engine/Documentation/Source/Engine/Animation/NodeReference/SkeletalControls/SpringController/SpringController.INT.udn b/Engine/Documentation/Source/Engine/Animation/NodeReference/SkeletalControls/SpringController/SpringController.INT.udn
--- a/Engine/Documentation/Source/Engine/Animation/NodeReference/SkeletalControls/SpringController/SpringController.INT.udn
+++ b/Engine/Documentation/Source/Engine/Animation/NodeReference/SkeletalControls/SpringController/SpringController.INT.udn
@@ -71,7 +71,7 @@ We then connect the nodes to include the Spring Controller that is targeting the
 | **Error Reset Thresh** | A threshold displacement - if the spring stretches more than this amount, it resets in order to avoid errors introduced by sudden, large displacements such as those caused by teleporting Actors. |
 | **No ZSpring** | If enabled, the displacement in the Z axis is ignored by the spring and is always correct. |
 
-**Spring Controller** nodes are also affected by the **LOD Threshold** system introduced in 4.11, you can find the setting for this in the Details Panel for the **Spring Controller** node.
+**Spring Controller** nodes are also affected by the **LOD Threshold** system introduced in 4.11. You can find the setting for this in the **Details Panel** for the **Spring Controller** node.
 
 [INCLUDE:Engine/Animation/NodeReference/SkeletalControls#lodthresh]
 
diff --git a/Engine/Documentation/Source/Engine/Animation/NodeReference/SkeletalControls/TrailController/TrailController.INT.udn b/Engine/Documentation/Source/Engine/Animation/NodeReference/SkeletalControls/TrailController/TrailController.INT.udn
--- a/Engine/Documentation/Source/Engine/Animation/NodeReference/SkeletalControls/TrailController/TrailController.INT.udn
+++ b/Engine/Documentation/Source/Engine/Animation/NodeReference/SkeletalControls/TrailController/TrailController.INT.udn
@@ -65,7 +65,7 @@ The **Trail Controller** can be used to affect a series of bones in a hierarchy.
 | **Component Pose** | The incoming pose. |
 | **Alpha** | The current strength of the control. |
 
-**Trail Controller** nodes are also affected by the **LOD Threshold** system introduced in 4.11, you can find the setting for this in the Details Panel for the **Trail Controller** node.
+**Trail Controller** nodes are also affected by the **LOD Threshold** system introduced in 4.11. You can find the setting for this in the **Details Panel** for the **Trail Controller** node.
 
 [INCLUDE:Engine/Animation/NodeReference/SkeletalControls#lodthresh]
 
diff --git a/Engine/Documentation/Source/Engine/Animation/NodeReference/SkeletalControls/TransformBone/TransformBone.INT.udn b/Engine/Documentation/Source/Engine/Animation/NodeReference/SkeletalControls/TransformBone/TransformBone.INT.udn
--- a/Engine/Documentation/Source/Engine/Animation/NodeReference/SkeletalControls/TransformBone/TransformBone.INT.udn
+++ b/Engine/Documentation/Source/Engine/Animation/NodeReference/SkeletalControls/TransformBone/TransformBone.INT.udn
@@ -73,7 +73,7 @@ The scale of our character is also driven by a float variable (**Scale Multiplie
 | **Rotation** | A Rotator specifying the Rotation to apply to the target bone. |
 | **Scale** | A Vector specifying the Scale to apply to the target bone. |
 
-**Transform (Modify) Bone** nodes are also affected by the **LOD Threshold** system introduced in 4.11, you can find the setting for this in the Details Panel for the **Transform (Modify) Bone** node.
+**Transform (Modify) Bone** nodes are also affected by the **LOD Threshold** system introduced in 4.11. You can find the setting for this in the **Details Panel** for the **Transform (Modify) Bone** node.
 
 [INCLUDE:Engine/Animation/NodeReference/SkeletalControls#lodthresh]
 
diff --git a/Engine/Documentation/Source/Engine/Animation/NodeReference/SkeletalControls/TwistCorrective/TwistCorrective.INT.udn b/Engine/Documentation/Source/Engine/Animation/NodeReference/SkeletalControls/TwistCorrective/TwistCorrective.INT.udn
--- a/Engine/Documentation/Source/Engine/Animation/NodeReference/SkeletalControls/TwistCorrective/TwistCorrective.INT.udn
+++ b/Engine/Documentation/Source/Engine/Animation/NodeReference/SkeletalControls/TwistCorrective/TwistCorrective.INT.udn
@@ -63,7 +63,7 @@ You can add the Twist Corrective node in the **AnimGraph** of your **Animation B
 | [REGION:tablesection]**Output Pins**[/REGION] ||
 | Pose (Blank) | The final pose after the Transform has been applied. |
 
-**Twist Corrective** nodes are also affected by the **LOD Threshold** system introduced in 4.11, you can find the setting for this in the Details Panel for the **Twist Corrective** node.
+**Twist Corrective** nodes are also affected by the **LOD Threshold** system introduced in 4.11. You can find the setting for this in the **Details Panel** for the **Twist Corrective** node.
 
 [INCLUDE:Engine/Animation/NodeReference/SkeletalControls#lodthresh]
 
diff --git a/Engine/Documentation/Source/Engine/Animation/NodeReference/SkeletalControls/TwoBoneIK/TwoBoneIK.INT.udn b/Engine/Documentation/Source/Engine/Animation/NodeReference/SkeletalControls/TwoBoneIK/TwoBoneIK.INT.udn
--- a/Engine/Documentation/Source/Engine/Animation/NodeReference/SkeletalControls/TwoBoneIK/TwoBoneIK.INT.udn
+++ b/Engine/Documentation/Source/Engine/Animation/NodeReference/SkeletalControls/TwoBoneIK/TwoBoneIK.INT.udn
@@ -84,7 +84,7 @@ By adjusting these values we can determine how the character holds the weapon wi
 | **Effector Location** | A Vector specifying the location effector in the Effector Location Space. |
 | **Joint Target Location** | A Vector specifying the location joint target in the Joint Target Location Space. |
 
-**Two Bone IK** nodes are also affected by the **LOD Threshold** system introduced in 4.11, you can find the setting for this in the Details Panel for the **Two Bone IK** node.
+**Two Bone IK** nodes are also affected by the **LOD Threshold** system introduced in 4.11. You can find the setting for this in the **Details Panel** for the **Two Bone IK** node.
 
 [INCLUDE:Engine/Animation/NodeReference/SkeletalControls#lodthresh]
 
