Replay overview draft. As simple a front page as I can get without leaving huge questions. Placeholder banners.

diff --git a/Engine/Documentation/Source/Engine/Replay/Replay.INT.udn b/Engine/Documentation/Source/Engine/Replay/Replay.INT.udn
new file mode 100644
--- /dev/null
+++ b/Engine/Documentation/Source/Engine/Replay/Replay.INT.udn
@@ -0,0 +1,69 @@
+Availability:Docs
+Title: Replay Overview
+Crumbs: %ROOT%, Engine
+Description:Overview of the Replay system
+Related:Programming/Replay
+Version: 4.12
+parent:Engine
+order:17
+
+[VAR:Topic]
+[OBJECT:Topic]
+	[PARAM:image]
+		![%Engine/Replay:title%](Engine/Replay/replay_topic.png)
+	[/PARAM]
+	[PARAM:icon]
+		![](%ROOT%/light_icon.png)(convert:false)
+	[/PARAM]
+	[PARAM:title]
+		%Engine/Replay:title%
+	[/PARAM]
+	[PARAM:description]
+		%Engine/Replay:description%
+	[/PARAM]
+	[PARAM:path]
+		[RELATIVE:Engine/Replay]
+	[/PARAM]
+[/OBJECT]
+[/VAR]
+[VAR:TopicCompact]
+[OBJECT:TopicCompact]
+	[PARAM:image]
+		![%Engine/Replay:title%](Engine/Replay/replay_topic.png)
+	[/PARAM]
+	[PARAM:icon]
+		![](%ROOT%/light_icon.png)(convert:false)
+	[/PARAM]
+	[PARAM:title]
+		%Engine/Replay:title%
+	[/PARAM]
+	[PARAM:description]
+		%Engine/Replay:description%
+	[/PARAM]
+	[PARAM:path]
+		[RELATIVE:Engine/Replay]
+	[/PARAM]
+[/OBJECT]
+[/VAR]
+
+[REGION:banner]
+![](Engine/Replay/replay_banner.png)(convert:false)
+[/REGION]
+
+## Replay System Overview
+Unreal Engine 4 features a Replay system which can record gameplay for later viewing. At a high level, the Replay system works by reading data drawn from the built-in replication system, similar to how live gameplay operates in a networked environment. The data is then passed to a Replay Streamer, which handles the process of filtering and storing the data. There are three types of Streamers included with the engine:
+
+### The NULL Streamer
+The NULL Streamer is the default way to record a replay. This streamer writes replay data to disk, e.g. a local hard drive. This is good for making local recordings, especially of single-player games. These recordings can be very helpful for a variety of uses, such as producing gameplay trailers or in-game cutscenes, or enabling your users to view and share speedrun or tutorial videos within your game.
+
+### The Memory Streamer
+The Memory Streamer keeps a user-configurable length of replay data (e.g. 10 seconds, 3 minutes, etc.) in memory on the local machine. This type of stream is best suited to instant replays of recent, dramatic moments, such as scores or penalty calls in a sports game, deaths in a shooter, or the final moments of a boss battle in an action game.
+
+### The HTTP Streamer
+The HTTP Streamer sends replay data to another server, which could be on a LAN or elsewhere on the internet. This is useful for live-streaming matches, or for keeping recordings of matches that can be viewed at any time. This streamer is especially useful for dedicated-server games, where only the server knows everything that's happening everywhere in the game at all times, and where offloading the work of processing replay data will increase the number of simultaneous games that can be hosted on a single server. It can also help to serve as a moderation or cheat detection tool, since the data can be captured from computers completely controlled by the party running the game.
+
+## Replay Data
+In terms of data, a replay contains three types of information. At the start is baseline data describing the starting state of the game world. Checkpoints that act as snapshots of the net changes to the world (from the baseline) appear at regular, user-defined intervals. The space between checkpoints is then filled with incremental changes to individual objects in the world. Any moment in the game can be reconstructed quickly and accurately by initializing the world to the starting state, making the changes described in the last checkpoint before the chosen time, and then applying each incremental change after that checkpoint leading up to the desired point in time.
+
+## Replay System Usage
+The basic functions of the replay system are to start or stop recording a demo, or play back a previously-recorded demo. While in playback mode, the Replay system supports commands to pause, change the playback speed, or jump to a specific point in time. More advanced usage of the system can include adding text tags (usually meta tags or the names of players involved in the game) to a replay, or querying existing replays based on text tags or version information present. The Replay system can be used from C++ code, mainly via the `UGameInstance` and `UWorld` classes, or with console or command-line commands. Note that some of the Replay system's streamers are not intended to be used in PIE or on non-authoritative machines such as clients in networked games.
