Minor typo fixes during peer review.

diff --git a/Engine/Documentation/Source/Engine/UMG/HowTo/CreatingWidgets/CreatingWidgets.INT.udn b/Engine/Documentation/Source/Engine/UMG/HowTo/CreatingWidgets/CreatingWidgets.INT.udn
--- a/Engine/Documentation/Source/Engine/UMG/HowTo/CreatingWidgets/CreatingWidgets.INT.udn
+++ b/Engine/Documentation/Source/Engine/UMG/HowTo/CreatingWidgets/CreatingWidgets.INT.udn
@@ -8,7 +8,7 @@ Version: 4.12
 SkillLevel: Beginner
 
 Typically in most games you will want to communicate some information to your players through your game's User Interface. 
-This can range from things like Main Menus, in-game Pause Menus, Health or Experience, Inventory Items or help text to instruct players what to do in a given situation. 
+This can range from things like Main Menus, in-game Pause Menus, Health or Experience, Inventory Items, or help text to instruct players what to do in a given situation. 
 
 With **Unreal Motion Graphics (UMG)**, you can create a [Widget Blueprint](Engine/UMG/UserGuide/WidgetBlueprints) to handle the display of UI elements in your project. 
 The Widget Blueprint uses [](Engine/Blueprints) to design the layout as well as script functionality for the UI element such as what happens when a button is clicked or if a value changes, etc. 
@@ -28,11 +28,11 @@ For this how-to, we are using the **Blueprint Third Person Template**. However,
 
 	Regardless of the type of UI element you wish to create in UMG, you will need to do so inside a Widget Blueprint. 
 
-1.	Enter the name **SampleUI** for the name of the Widget Blueprint, then **Double-click** and open it. 
+1.	Enter the name **SampleUI** for the name of the Widget Blueprint, then **double-click** and open it. 
 
 	![](CreatingWidgets02.png)
 
-	Upon opening a Widget Blueprint you will be presented with the UMG Editor Interface
+	Upon opening a Widget Blueprint you will be presented with the UMG Editor Interface.
 
 	[REGION:lightbox]
 	[![](CreatingWidgets03.png)](CreatingWidgets03.png)
@@ -42,9 +42,9 @@ For this how-to, we are using the **Blueprint Third Person Template**. However,
 	Click for full view. 
 	[/REGION]
 
-	The Widget Blueprint opens by default to the **Designer** tab which is where you perform your visual layout of your UI elements. To script functionality for the Widget Bluperint you can click the **Graph** button in the top right corner of the editor. 
+	The Widget Blueprint opens by default to the **Designer** tab, which is where you perform your visual layout of your UI elements. To script functionality for the Widget Blueprint, you can click the **Graph** button in the top right corner of the editor. 
 
-1.	Inside the **Palette** window, **Left-click** and drag a **Text** widget into the **Visual Designer** graph. 
+1.	Inside the **Palette** window, **left-click** and drag a **Text** widget into the **Visual Designer** graph. 
 
 	![](CreatingWidgets04.png)
 
@@ -54,15 +54,15 @@ For this how-to, we are using the **Blueprint Third Person Template**. However,
 
 	![](CreatingWidgets05.png)
 
-	This will ensure that as the viewport size changes, the Text will stay anchored relative to the center of the screen. 
+	This will ensure that as the viewport size changes, the Text widget will stay anchored relative to the center of the screen. 
 
 1.	In the **Details** panel click the **Size to Content** option, change the **Text**, **Color and Opacity** and **Font Size** options. 
 
 	![](CreatingWidgets06.png)
 
-	The Size to Content option as the name suggests will adjust the size of the widget to fit its content. We also changed the text to "Hello World!", changed the color of the text to Green and set the Font Size to 56.
+	The Size to Content option as the name suggests will adjust the size of the widget to fit its content. We also changed the text to "Hello World!", changed the color of the text to green, and set the Font Size to 56.
 
-1.	Click the **Compile** button then **Save** and close the Widget Blueprint. 
+1.	Click the **Compile** button, then **Save** and close the Widget Blueprint. 
 
 	![](CreatingWidgets09.png)
 
@@ -84,7 +84,7 @@ For this how-to, we are using the **Blueprint Third Person Template**. However,
 
 	![](CreatingWidgets13.png)
 
-	Whenever you create the Widget Blueprint (your UI), it's a good idea to promote that to a variable so you can access it later through Blueprint Script. 
+	Whenever you create the Widget Blueprint (your UI), it's a good idea to promote that to a variable so you can access it later through Blueprint script. 
 
 1.	Drag off the out pin on the **Set** node and search for and add the **Add to Viewport** node. 
 
@@ -126,7 +126,7 @@ For this how-to, we are using the **Blueprint Third Person Template**. However,
 
 ## End Result
 
-When you play in the editor you should have something similar to below whenever you press/release the **Enter** key. 
+When you play in the editor, you should have something similar to below whenever you press/release the **Enter** key. 
 
 [OBJECT:EmbeddedVideo]
 [PARAMLITERAL:width]
