Minor draft improvement.

diff --git a/Engine/Documentation/Source/Engine/Replay/Replay.INT.udn b/Engine/Documentation/Source/Engine/Replay/Replay.INT.udn
--- a/Engine/Documentation/Source/Engine/Replay/Replay.INT.udn
+++ b/Engine/Documentation/Source/Engine/Replay/Replay.INT.udn
@@ -2,7 +2,6 @@ Availability:Docs
 Title: Replay Overview
 Crumbs: %ROOT%, Engine
 Description:Overview of the Replay system
-Related:Programming/Replay
 Version: 4.12
 parent:Engine
 order:17
@@ -63,7 +62,7 @@ The Memory Streamer keeps a user-configurable length of replay data (e.g. 10 sec
 The HTTP Streamer sends replay data to another server, which could be on a LAN or elsewhere on the internet. This is useful for live-streaming matches, or for keeping recordings of matches that can be viewed at any time. This streamer is especially useful for dedicated-server games, where only the server knows everything that's happening everywhere in the game at all times, and where offloading the work of processing replay data will increase the number of simultaneous games that can be hosted on a single server. It can also help to serve as a moderation or cheat detection tool, since the data can be captured from computers completely controlled by the party running the game.
 
 ## Replay Data
-In terms of data, a replay contains three types of information. At the start is baseline data describing the starting state of the game world. Checkpoints that act as snapshots of the net changes to the world (from the baseline) appear at regular, user-defined intervals. The space between checkpoints is then filled with incremental changes to individual objects in the world. Any moment in the game can be reconstructed quickly and accurately by initializing the world to the starting state, making the changes described in the last checkpoint before the chosen time, and then applying each incremental change after that checkpoint leading up to the desired point in time.
+In terms of data, a replay contains three types of game-state information, as well as some additional text data. At the start is baseline data describing the starting state of the game world. Checkpoints that act as snapshots of the net changes to the world (from the baseline) appear at regular, user-defined intervals. The space between checkpoints is then filled with incremental changes to individual objects in the world. Any moment in the game can be reconstructed quickly and accurately by initializing the world to the starting state, making the changes described in the last checkpoint before the chosen time, and then applying each incremental change after that checkpoint leading up to the desired point in time. The text data contained in a replay consists of a display name, which can be used when making a player-facing list, and user-defined text tags (HTTP Streamer only), which can be used when searching or filtering through lists of games.
 
 ## Replay System Usage
 The basic functions of the replay system are to start or stop recording a demo, or play back a previously-recorded demo. While in playback mode, the Replay system supports commands to pause, change the playback speed, or jump to a specific point in time. More advanced usage of the system can include adding text tags (usually meta tags or the names of players involved in the game) to a replay, or querying existing replays based on text tags or version information present. The Replay system can be used from C++ code, mainly via the `UGameInstance` and `UWorld` classes, or with console or command-line commands. Note that some of the Replay system's streamers are not intended to be used in PIE or on non-authoritative machines such as clients in networked games.
