#ue4 docs

diff --git a/Engine/Documentation/Source/Engine/Actors/Merging/ActorMerging.INT.udn b/Engine/Documentation/Source/Engine/Actors/Merging/ActorMerging.INT.udn
--- a/Engine/Documentation/Source/Engine/Actors/Merging/ActorMerging.INT.udn
+++ b/Engine/Documentation/Source/Engine/Actors/Merging/ActorMerging.INT.udn
@@ -2,6 +2,7 @@ Availability: Docs
 Title:Actor Merging
 Crumbs: %ROOT%, Engine, Engine/Actors, Engine/Actors/Merging
 Description:How to use the Actor Merge tool to combine multiple Static Meshes into one in UE4.
+Related: Engine/Actors/StaticMeshActor
 Version: 4.12
 
 
