#UE4 docs UEDOC-3067 Updated version numbers to 4.13

diff --git a/Engine/Documentation/Source/Engine/Editor/VR/Controls/VREditorControls.INT.udn b/Engine/Documentation/Source/Engine/Editor/VR/Controls/VREditorControls.INT.udn
--- a/Engine/Documentation/Source/Engine/Editor/VR/Controls/VREditorControls.INT.udn
+++ b/Engine/Documentation/Source/Engine/Editor/VR/Controls/VREditorControls.INT.udn
@@ -6,7 +6,7 @@ parent: Engine/Editor/VR
 type: reference
 order: 1
 tags: VR
-version: 4.11
+version: 4.13
 skilllevel: intermediate
 
 [VAR:Topic]
diff --git a/Engine/Documentation/Source/Engine/UI/LevelEditor/Modes/MeshPaintMode/VertexColor/PaintingOnMeshes/PaintingOnMeshes.INT.udn b/Engine/Documentation/Source/Engine/UI/LevelEditor/Modes/MeshPaintMode/VertexColor/PaintingOnMeshes/PaintingOnMeshes.INT.udn
--- a/Engine/Documentation/Source/Engine/UI/LevelEditor/Modes/MeshPaintMode/VertexColor/PaintingOnMeshes/PaintingOnMeshes.INT.udn
+++ b/Engine/Documentation/Source/Engine/UI/LevelEditor/Modes/MeshPaintMode/VertexColor/PaintingOnMeshes/PaintingOnMeshes.INT.udn
@@ -2,7 +2,7 @@ Availability: Public
 Title:Painting on Meshes
 Crumbs:%ROOT%, Engine, Engine/UI/LevelEditor/Modes/MeshPaintMode/VertexColor
 Description:An overview of how to lay paint strokes down on meshes using Mesh Paint Mode.
-Version: 4.9
+Version: 4.13
 SkillLevel: Intermediate
 
 [EXCERPT:Intro]
