[NO NON-PII COMMIT MESSAGE]

diff --git a/Engine/Documentation/Source/Programming/UnrealArchitecture/Reference/Functions/Specifiers/Category/FunctionSpecifier_Category.INT.udn b/Engine/Documentation/Source/Programming/UnrealArchitecture/Reference/Functions/Specifiers/Category/FunctionSpecifier_Category.INT.udn
--- a/Engine/Documentation/Source/Programming/UnrealArchitecture/Reference/Functions/Specifiers/Category/FunctionSpecifier_Category.INT.udn
+++ b/Engine/Documentation/Source/Programming/UnrealArchitecture/Reference/Functions/Specifiers/Category/FunctionSpecifier_Category.INT.udn
@@ -5,4 +5,4 @@ Availability:Public
 Version: 4.9
 
 Specifies the category of the function when displayed in Blueprint editing tools.
-	* Usage: Category=CategoryName or Category="MajorCategory,SubCategory"
+	* Usage: Category=CategoryName or Category="MajorCategory|SubCategory"
