[NO NON-PII COMMIT MESSAGE]

diff --git a/Engine/Documentation/Source/Gameplay/Networking/Actors/Components/Components.INT.udn b/Engine/Documentation/Source/Gameplay/Networking/Actors/Components/Components.INT.udn
--- a/Engine/Documentation/Source/Gameplay/Networking/Actors/Components/Components.INT.udn
+++ b/Engine/Documentation/Source/Gameplay/Networking/Actors/Components/Components.INT.udn
@@ -98,7 +98,7 @@ can be fully replicated, and does not require the items to be Actors (which woul
 
 ## Optimizations
 In the case of having many subobjects to replicate, an Actor can save a lot of time by knowing which, if any, subobjects have recently changed and need to be replicated. This means keeping track of the subobjects changing through accessor functions.
-The interface for doing this is in AActorChannel:
+The interface for doing this is in UActorChannel:
 	bool KeyNeedsToReplicate(int32 ObjID, int32 RepKey);
 This function should be called by the Actor in their `::ReplicateSubobjects` implementation. The Actor class can setup an arbitrary Object ID and Replication Key that the replication system will keep track of for each client. When ReplicateSubobjects returns _true_, the caller is expected to call ReplicateSubobjects on whichever subobjects are being tracked with that Object ID.
 Look at AQAInventory::ReplicateSubobjects for an example. The key thing to remember is that the object ID and replication key are completely arbitrary. Object ID is just used to reference 'things'. It could be the entire list of subobjects, a partial list, or individual objects. The replication key is also arbitrary - it can just be a counter that increments when whatever Object ID is tracking changes. The optimizations discussed here are completely optional.
