Initial pass at Drag/Drop How-to.

diff --git a/Engine/Documentation/Source/Engine/UMG/HowTo/DragAndDrop/DragAndDrop.INT.udn b/Engine/Documentation/Source/Engine/UMG/HowTo/DragAndDrop/DragAndDrop.INT.udn
new file mode 100644
--- /dev/null
+++ b/Engine/Documentation/Source/Engine/UMG/HowTo/DragAndDrop/DragAndDrop.INT.udn
@@ -0,0 +1,48 @@
+Availability: Docs
+Title:Creating Drag and Drop UI
+Crumbs: %ROOT%, Engine, Engine/UMG, Engine/UMG/HowTo
+Description:Shows how you can create draggable and droppable UI widgets with UMG. 
+Version: 4.12
+SkillLevel: Intermediate
+
+There may be instances in your project where you want your players to be able to interact with UI elements by dragging and dropping them onscreen. 
+This could be to customize interface layouts (positioning health bars or unit frames) or interacting with gameplay systems like inventory screens (adding/removing items around). 
+With UMG you can create these types of interactions with the **DragAndDropOperation** node and overriding some of the inherited Blueprint Functions based on your method of interaction. 
+
+In this example, we take a look at how a player can drag and reposition a health bar on screen with the Left Mouse Button.
+
+[OBJECT:EmbeddedVideo]
+[PARAMLITERAL:width]
+640
+[/PARAMLITERAL]
+[PARAMLITERAL:height]
+360
+[/PARAMLITERAL]
+[PARAMLITERAL:videoid]
+qUDIkWcQChI
+[/PARAMLITERAL]
+[/OBJECT]
+
+_At the end of this guide, you will have a UI widget that you can drag around inside the viewport and drop at a new location._
+
+[OBJECT:TopicButtonList]
+	[PARAM:icon]
+		![](%ROOT%/build_icon.png)(convert:false)
+	[/PARAM]
+	[PARAM:title]
+		Steps
+	[/PARAM]
+	[PARAM:description]
+	[/PARAM]
+	[PARAM:links]
+	* [](Engine\UMG\HowTo\DragAndDrop\ProjectSetup)
+	* [](Engine\UMG\HowTo\DragAndDrop\WidgetSetup)
+	* [](Engine\UMG\HowTo\DragAndDrop\OnDragDetected)
+	* [](Engine\UMG\HowTo\DragAndDrop\OnDrop)
+	[/PARAM]
+[/OBJECT]
+
+[REGION:call_to_action]
+	[Click to Start](Engine\UMG\HowTo\DragAndDrop\ProjectSetup "%Engine\UMG\HowTo\DragAndDrop\ProjectSetup:title%")
+[/REGION]
+
diff --git a/Engine/Documentation/Source/Engine/UMG/HowTo/DragAndDrop/OnDragDetected/OnDragDetected.INT.udn b/Engine/Documentation/Source/Engine/UMG/HowTo/DragAndDrop/OnDragDetected/OnDragDetected.INT.udn
new file mode 100644
--- /dev/null
+++ b/Engine/Documentation/Source/Engine/UMG/HowTo/DragAndDrop/OnDragDetected/OnDragDetected.INT.udn
@@ -0,0 +1,62 @@
+Availability: Docs
+Title:3. On Drag Detected
+Crumbs: %ROOT%, Engine, Engine/UMG, Engine/UMG/HowTo
+Description:In this step we set up what happens when drag is detected.
+Related: Engine/UMG
+Related: Engine/UMG/UserGuide
+Version: 4.12
+SkillLevel: Intermediate
+
+[VAR:Steps]
+[OBJECT:Navigation]
+	[PARAM:previous]
+		[Previous Step](Engine\UMG\HowTo\DragAndDrop\WidgetSetup "%Engine\UMG\HowTo\DragAndDrop\WidgetSetup:title%")
+	[/PARAM]
+	[PARAM:current]
+	[/PARAM]
+	[PARAM:home]
+		[Drag and Drop Widgets](Engine\UMG\HowTo\DragAndDrop)
+	[/PARAM]
+	[PARAM:next]
+		[Next Step](Engine\UMG\HowTo\DragAndDrop\OnDrop "%Engine\UMG\HowTo\DragAndDrop\OnDrop:title%")
+	[/PARAM]
+[/OBJECT]
+[/VAR]
+
+%Steps%
+
+In this step we determine what happens when the player is in fact dragging the HealthBar around on screen. 
+
+## Steps
+
+1.	Open the **OnDragDetected** function, then **Right-click** and add a **Create Drag Drop Operation** node set to use **WidgetDrag**. 
+
+	![](DragDrop9.png)
+
+	UMG supplies a default class for drag and drop operations however we have created a custom class that has additional variables in it that we will pass through as part of the process.
+
+1.	**Right-click** in the graph and **Get Reference to Self** and connect as shown and set the **Pivot** drop-down to **Mouse Down**. 
+	
+	![](DragDrop10.png)
+
+	Pivot determines where the Drag Widget Visual will appear when being dragged relative to the pointer performing the drag operation. 
+
+1.	From the **MyBlueprint** panel, drag in the **Drag Offset** variable and connect it and the pins shown below. 
+
+	![](DragDrop11.png)
+
+	Here we are passing through that vector 2D information of where the player clicked on the widget which we will use as part of the OnDrop operation. 
+
+1.	**Right-click** in the graph and add a **Remove from Parent** node and connect as shown. 
+
+	![](DragDrop12.png)
+
+	Before we complete the On Drag Detected process, we remove the widget we are dragging from display.
+	If we don't remove the original HealthBar widget from display, it would remain in the default position while we are dragging a copy around onscreen. 
+	The widget itself is being passed through as part of the process so we are not creating a "new" version of the widget. 
+
+The last few things we need to do is determine what happens when the player releases the Left Mouse Button and performs the OnDrop function. 
+We also need to add the HealthBar Widget to our Main HUD Widget Blueprint and tell that to be displayed before testing everything out.  
+
+%Steps%
+
diff --git a/Engine/Documentation/Source/Engine/UMG/HowTo/DragAndDrop/OnDrop/OnDrop.INT.udn b/Engine/Documentation/Source/Engine/UMG/HowTo/DragAndDrop/OnDrop/OnDrop.INT.udn
new file mode 100644
--- /dev/null
+++ b/Engine/Documentation/Source/Engine/UMG/HowTo/DragAndDrop/OnDrop/OnDrop.INT.udn
@@ -0,0 +1,103 @@
+Availability: Docs
+Title:4. On Drop
+Crumbs: %ROOT%, Engine, Engine/UMG, Engine/UMG/HowTo
+Description:We finish by creating the script for OnDrop functionality before we test. 
+Related: Engine/UMG
+Related: Engine/UMG/UserGuide
+Version: 4.12
+SkillLevel: Intermediate
+
+[VAR:Steps]
+[OBJECT:Navigation]
+	[PARAM:previous]
+		[Previous Step](Engine\UMG\HowTo\DragAndDrop\OnDragDetected "%Engine\UMG\HowTo\DragAndDrop\OnDragDetected:title%")
+	[/PARAM]
+	[PARAM:current]
+	[/PARAM]
+	[PARAM:home]
+		[Drag and Drop Widgets](Engine\UMG\HowTo\DragAndDrop)
+	[/PARAM]
+	[PARAM:next]
+	[/PARAM]
+[/OBJECT]
+[/VAR]
+
+%Steps%
+
+To finish up, we script what happens when we "drop" the widget in our HUD as well as test everything out. 
+
+## Steps
+
+1.	Open the **HUD** Widget Blueprint, then in the **Details** panel for the **Canvas Panel**, set its **Visibility** to **Visible**. 
+
+	![](DragDrop13.png)
+
+	We are doing our drag detection inside the Health Bar widget, however when we release the Left Mouse Button we want to drop it in our HUD Widget Blueprint.
+	In order for our HUD to receive hit detection, we need to set its **Visibility** setting to **Visible**. 
+
+1.	From the **Palette** window, add the **Health Bar** widget to the graph in your desired location. 
+
+	![](DragDrop14.png)
+
+1.	Go to the **Graph** Mode, then add an **OnDrop** function override from the **My Blueprint** panel. 
+
+	![](DragDrop15.png)
+
+1.	Inside **OnDrop** drag off **Operation** and **Cast to WidgetDrag**, then **Get Widget Reference** and **Get Drag Offset**. 
+
+	![](DragDrop16.png)
+
+	We cast here to check which type of operation is being performed as part of the **OnDrop** function, is it **Widget Drag**? 
+	If so, we get the **WidgetReference** (which we passed through as the widget being dragged) as well as the **DragOffset** (or position in which we started the drag from). 
+
+1.	Off **My Geometry** and **Pointer Event**, use the **Absolute to Local** and **GetScreenSpacePosition** nodes and connect as shown.
+
+	![](DragDrop17.png)
+
+	This will tell us the position in 2D space where we released the Left Mouse Button to drop the widget. 
+	We will use this minus the **DragOffset** to determine where we should place the widget when it is dropped. 
+
+1.	Subtract the **DragOffset** from the **Return Value** of **Absolute to Local**. 
+
+1.	For the **Widget Reference**, use the **Add to Viewport** and **Set Position in Viewport** and connect as shown. 
+
+	[REGION:lightbox]
+	[![](DragDrop18.png)](DragDrop18.png)
+	[/REGION]
+
+	[REGION:caption]
+	Click for full view. 
+	[/REGION]
+
+	Make sure to uncheck **Remove DPIScale** and check the **Return Value** checkboxes as we do not need to remove the [](Engine/UMG/UserGuide/DPIScaling) and we have handled this function so we return as true. 
+	With this script we add the widget being dragged to the viewport again and set its position based on where we clicked and started dragging. The last thing we need to do is create and display our HUD Widget Blueprint. 
+
+1.	Inside the **Content/ThirdPersonBP/Blueprints** folder, open the **ThirdPersonCharacter** Blueprint. 
+
+1.	Create the **HUD** Widget Blueprint, **Add to Viewport** and **Get Player Controller** to **Show Mouse Cursor**. 
+
+	![](DragDrop19.png)
+
+1.	**Compile**, **Save** then click the **Play** button to play in the Editor. 
+
+
+## End Result
+
+When you play in the Editor, you can **Left-click** and drag the health bar on screen then drop it in a new location. 
+
+[OBJECT:EmbeddedVideo]
+[PARAMLITERAL:width]
+640
+[/PARAMLITERAL]
+[PARAMLITERAL:height]
+360
+[/PARAMLITERAL]
+[PARAMLITERAL:videoid]
+qUDIkWcQChI
+[/PARAMLITERAL]
+[/OBJECT]
+
+This is just an example of the elements you will need to start the Drag and Drop process. 
+Additional checking may be needed to ensure that players do not drag the widget outside of Safe Zones or potentially on top of other widgets. 
+
+%Steps%
diff --git a/Engine/Documentation/Source/Engine/UMG/HowTo/DragAndDrop/ProjectSetup/ProjectSetup.INT.udn b/Engine/Documentation/Source/Engine/UMG/HowTo/DragAndDrop/ProjectSetup/ProjectSetup.INT.udn
new file mode 100644
--- /dev/null
+++ b/Engine/Documentation/Source/Engine/UMG/HowTo/DragAndDrop/ProjectSetup/ProjectSetup.INT.udn
@@ -0,0 +1,60 @@
+Availability: Docs
+Title:1. Project Setup
+Crumbs: %ROOT%, Engine, Engine/UMG, Engine/UMG/HowTo
+Description:Here we create the assets that we'll need and our DragAndDropOperation Class. 
+Version: 4.12
+SkillLevel: Intermediate
+
+[VAR:Steps]
+[OBJECT:Navigation]
+	[PARAM:previous]
+	[/PARAM]
+	[PARAM:current]
+	[/PARAM]
+	[PARAM:home]
+		[Drag and Drop Widgets](Engine\UMG\HowTo\DragAndDrop)
+	[/PARAM]
+	[PARAM:next]
+		[Next Step](Engine\UMG\HowTo\DragAndDrop\WidgetSetup "%Engine\UMG\HowTo\DragAndDrop\WidgetSetup:title%")
+	[/PARAM]
+[/OBJECT]
+[/VAR]
+
+%Steps%
+
+Before we get started, there are a couple assets that we'll need to create as part of our project setup.
+
+[REGION:note]
+For this how-to, we are using the **Blueprint Third Person Template**. However, you can use any project you wish. 
+[/REGION]
+
+## Steps
+
+1.	In the **Content Browser** of your project, create a folder with two **Widget Blueprints**, one called **HealthBar** and the other **HUD**. 
+
+	![](DragDrop1.png)
+
+	In this example, the element that we want to be drag and droppable is our **HealthBar** widget that will be placed in our **HUD**. 
+
+1.	**Right-click** in the **Content Browser** and create a new **Blueprint Class** of the **DragDropOpertion** class and call it **WidgetDrag**. 
+
+	![](DragDrop2.png)
+
+	This specialized Blueprint Class will allow us to pass through information as part of our drag and drop action. 
+
+1.	Open **WidgetDrag** then create a **User Widget** variable called **WidgetReference** and make it **Editable** and **Expose on Spawn**. 
+
+	![](DragDrop3a.png)
+
+	This will be used to hold the UMG Widget that we want to drag around on screen. 
+
+1.	Create another variable called **DragOffset** of the **Vector 2D** type that is **Editable** and **Exposed on Spawn**. 
+
+	![](DragDrop3b.png)
+
+	This will be used to offset the location in which we start dragging and when we drop our widget.
+
+
+With our project setup done, next we will start working on the HealthBar Widget Blueprint which is the piece of UI that we want to be draggable. 
+
+%Steps%
diff --git a/Engine/Documentation/Source/Engine/UMG/HowTo/DragAndDrop/WidgetSetup/WidgetSetup.INT.udn b/Engine/Documentation/Source/Engine/UMG/HowTo/DragAndDrop/WidgetSetup/WidgetSetup.INT.udn
new file mode 100644
--- /dev/null
+++ b/Engine/Documentation/Source/Engine/UMG/HowTo/DragAndDrop/WidgetSetup/WidgetSetup.INT.udn
@@ -0,0 +1,76 @@
+Availability: Docs
+Title:2. Widget Setup
+Crumbs: %ROOT%, Engine, Engine/UMG, Engine/UMG/HowTo
+Description:In this step we script what happens when the Left Mouse Button is pressed. 
+Version: 4.12
+SkillLevel: Intermediate
+
+[VAR:Steps]
+[OBJECT:Navigation]
+	[PARAM:previous]
+		[Previous Step](Engine\UMG\HowTo\DragAndDrop\ProjectSetup "%Engine\UMG\HowTo\DragAndDrop\ProjectSetup:title%")
+	[/PARAM]
+	[PARAM:current]
+	[/PARAM]
+	[PARAM:home]
+		[Drag and Drop Widgets](Engine\UMG\HowTo\DragAndDrop)
+	[/PARAM]
+	[PARAM:next]
+		[Next Step](Engine\UMG\HowTo\DragAndDrop\OnDragDetected "%Engine\UMG\HowTo\DragAndDrop\OnDragDetected:title%")
+	[/PARAM]
+[/OBJECT]
+[/VAR]
+
+%Steps%
+
+In this step we determine when the Left Mouse Button is pressed and store the location where it was pressed in screen space. 
+We also use the **DetectDragIfPressed** to determine if the player is dragging with a specified key. 
+
+## Steps
+
+1.	Open the **HealthBar** Widget Blueprint and delete the **Canvas Panel** in the **Hierarchy** window. 
+
+1.	Add a **Progress Bar**, then **Right-click** on it and **Wrap With** a **Size Box** set to **500 x 50** in the **Details** panel. 
+
+	![](DragDrop4a.png)
+
+	Optionally you can stylize the Progress Bar. Below we set the adjusted preview to **Desired on Screen**, set the **Percent** to **1.0** so the progress bar is filled and changed the color to green. 
+	
+	![](DragDrop4b.png)
+
+	In this tutorial we will not be creating a functional Health Bar and this is intended for demonstration purposes only. 
+
+1.	Switch to the **Graph** Mode, then in the **My Blueprint** panel add the **OnMouseButtonDown** and **OnDragDetected** overrides. 
+
+	![](DragDrop5.png)
+
+	There are several functions that you can override to call your own custom script. Here we will perform some checking when the Mouse Button is pressed and what happens when a widget drag is detected. 
+
+1.	Inside the **OnMouseButtonDown** node, drag off the **My Geometry** pin and add the **Absolute to Local** node. 
+
+1.	Drag off the **Mouse Event** pin and add the **Get Screen Space Position** node and connect as shown below. 
+
+	![](DragDrop6.png)
+
+	Here we are converting the absolute coordinates of the Health Bar widget to local space by getting the screen space position in which the mouse button registered a button down key press. 
+	This will tell us where on the widget the player actually clicked so that we can use that to determine where we should place the widget when we choose to drop it so we will store it as a variable. 
+	
+1.	**Right-click** the **Return Value** of the **Absolute to Local** pin and **Promote to Variable** called **Drag Offset**. 
+
+	![](DragDrop7.png)
+
+	[REGION:note]
+	Don't worry if you see an error on the On Mouse Button Down node, once we complete the function it will go away. 
+	[/REGION]
+
+1.	**Right-click** in the graph and add a **Detect Drag if Pressed** node. 
+
+1.	Set the **Drag Key** to **Left Mouse Button** and connect the **Mouse Event** and remaining pins as shown below. 
+
+	![](DragDrop8.png)
+
+	Here we are determining if the player is dragging the Left Mouse Button before which will start the On Drag Detected script. 
+
+With the ability to detect if the player is dragging in place, next we will finish off what happens when drag is detected. 
+
+%Steps%
diff --git a/Engine/Documentation/Source/Engine/UMG/HowTo/UMGHowTo.INT.udn b/Engine/Documentation/Source/Engine/UMG/HowTo/UMGHowTo.INT.udn
--- a/Engine/Documentation/Source/Engine/UMG/HowTo/UMGHowTo.INT.udn
+++ b/Engine/Documentation/Source/Engine/UMG/HowTo/UMGHowTo.INT.udn
@@ -19,6 +19,9 @@ Refer to the table below for more information and click each link to view its co
 ## Intermediate
 [REGION:wrapper transparent]
 [](FANCY:Engine\UMG\HowTo\Create3DWidgets "%Engine\UMG\HowTo\Create3DWidgets:description%")
+[PUBLISH:docs]
+[](FANCY:Engine\UMG\HowTo\DragAndDrop "%Engine\UMG\HowTo\DragAndDrop:description%")
+[/PUBLISH:docs]
 [](FANCY:Engine\UMG\HowTo\ScalingUI "%Engine\UMG\HowTo\ScalingUI:description%")
 [/REGION]
 
