Fixing auto-generated link list

diff --git a/Engine/Documentation/Source/Engine/Blueprints/Scripting/Scripting.INT.udn b/Engine/Documentation/Source/Engine/Blueprints/Scripting/Scripting.INT.udn
--- a/Engine/Documentation/Source/Engine/Blueprints/Scripting/Scripting.INT.udn
+++ b/Engine/Documentation/Source/Engine/Blueprints/Scripting/Scripting.INT.udn
@@ -51,9 +51,9 @@ different variable types, working with those variables, and execution of nodes w
 ## Variables
 
 Variables can be created in a variety of different types, including data types such as Boolean, integer, and float. They are color-coded for easy identification within your Blueprint.
-Blueprint variables can also be reference types for holding things like Objects, Actors, and Classes. 
+Blueprint variables can also be reference types for holding things like Objects, Actors, and Classes.
 
-[DIR(output:"listbutton" parent:"Engine/Blueprints" tags:"Data Types")]
+[DIR(output:"listbutton" tags:"Data Types")]
 
 ## Execution Flow
 
