#ue4 docs

diff --git a/Engine/Documentation/Source/Engine/Tools/DevTools/PixelInspector/PixelInspector.INT.udn b/Engine/Documentation/Source/Engine/Tools/DevTools/PixelInspector/PixelInspector.INT.udn
--- a/Engine/Documentation/Source/Engine/Tools/DevTools/PixelInspector/PixelInspector.INT.udn
+++ b/Engine/Documentation/Source/Engine/Tools/DevTools/PixelInspector/PixelInspector.INT.udn
@@ -10,29 +10,31 @@ version: 4.12
 [/REGION]
 
 The **Pixel Inspector** tool is a Developer Tool that will allow you to diagnose the pixels that make up the colors in your scene. 
-Want to know why a pixel is producing an unexpected color or maybe what Material input is driving a pixel's color, you can use the Pixel Inspector's **Inspect Mode** to output information that is driving the pixel's visual result. 
+If you want to know why a pixel is producing an unexpected color (or maybe what Material input is driving a pixel's color), you can use the Pixel Inspector's **Inspect Mode** to output information that is driving the pixel's visual result. 
 
 ## Using the Pixel Inspector
 
-To enable and use the Pixel Inspector, from the Main Editor Window go to the **Window** menu option under **Developer Tools**: 
+To enable and use the Pixel Inspector:
 
-![](EnablePixelInspector.png)
+1. From the Main Editor Window go to the **Window** menu option under **Developer Tools** and select **Pixel Inspector**.  
 
-Clicking on **Pixel Inspector** will open the Pixel Inspector tool: 
+	![](EnablePixelInspector.png)
 
-![](PixelInspectorWindow.png)
+	The Pixel Inspector Window will open. 
 
-Click the ![](InspectIcon.png) button to start pixel inspection: 
+	![](PixelInspectorWindow.png)
 
-![](Inspecting.png)
+1. Click the ![](InspectIcon.png) button to start pixel inspection.
 
-Moving the mouse over any viewport (Level, Material, Blueprint, etc..) will populate the Pixel Inspector data fields in real-time. 
+1. Move the mouse over any viewport (Level, Material, Blueprint, etc..) to populate the Pixel Inspector data fields in real-time. 
 
-Pressing **Escape** will stop inspection and populate the Pixel Inspector data fields with the last inspected pixel. 
+	![](Inspecting.png)
 
+1. Press **Escape** to stop inspection and populate the Pixel Inspector data fields with the last inspected pixel. 
 
 
-## Pixel Inspector Information
+
+## Pixel Inspector Data Reference
 
 Below are the data fields that will become populated with a pixel's information during inspection: 
 
