Added CPP FPS Button Object on the Programming Tutorials Landing Page.
The button object is currently commented out because it is still in Editorial review; however, the block can be uncommented when the CPP FPS Tutorial is ready for publication.
Updated the CPP FPS Landing Page Description to be a bit more descriptive.

#UE4docs

diff --git a/Engine/Documentation/Source/Programming/Tutorials/FirstPersonShooter/FirstPersonShooter.INT.udn b/Engine/Documentation/Source/Programming/Tutorials/FirstPersonShooter/FirstPersonShooter.INT.udn
--- a/Engine/Documentation/Source/Programming/Tutorials/FirstPersonShooter/FirstPersonShooter.INT.udn
+++ b/Engine/Documentation/Source/Programming/Tutorials/FirstPersonShooter/FirstPersonShooter.INT.udn
@@ -1,7 +1,7 @@
 Availability: Docs
 Title: C++ First Person Shooter Tutorial
 Crumbs: %ROOT%, GettingStarted, Programming/Tutorials
-Description: How to develop a First Person Shooter.
+Description: Learn how to implement mechanics for a First Person Shooter game.
 prereq:Programming/Development/VisualStudioSetup
 SkillLevel: Beginner
 Version: 4.11
diff --git a/Engine/Documentation/Source/Programming/Tutorials/ProgrammingTutorialsHome.INT.udn b/Engine/Documentation/Source/Programming/Tutorials/ProgrammingTutorialsHome.INT.udn
--- a/Engine/Documentation/Source/Programming/Tutorials/ProgrammingTutorialsHome.INT.udn
+++ b/Engine/Documentation/Source/Programming/Tutorials/ProgrammingTutorialsHome.INT.udn
@@ -145,6 +145,25 @@ programming
 	[/PARAMLITERAL]
 [/OBJECT]
 
+<!--TO DO: Uncomment the following block when ready for publication.
+[OBJECT:Role]
+	[PARAM:icon]
+		![](%ROOT%/start_icon.png)(convert:false)
+	[/PARAM]
+	[PARAM:title]
+		%Programming/Tutorials/FirstPersonShooter:title%
+	[/PARAM]
+	[PARAM:path]
+		[RELATIVE:Programming/Tutorials/FirstPersonShooter]
+	[/PARAM]
+	[PARAM:description]
+		%Programming/Tutorials/FirstPersonShooter:description%
+	[/PARAM]
+	[PARAMLITERAL:id]
+		cppfps
+	[/PARAMLITERAL]
+[/OBJECT]
+-->
 [/REGION]
 
 ## Intermediate
