#UE4 Docs: Checking in the first pass of this doc.

diff --git a/Engine/Documentation/Source/Platforms/Mobile/Performance/MobilePerformance.INT.udn b/Engine/Documentation/Source/Platforms/Mobile/Performance/MobilePerformance.INT.udn
--- a/Engine/Documentation/Source/Platforms/Mobile/Performance/MobilePerformance.INT.udn
+++ b/Engine/Documentation/Source/Platforms/Mobile/Performance/MobilePerformance.INT.udn
@@ -13,38 +13,37 @@ tags:mobile
 
 
 
-Below are things to watch out for or suggestions on how to set up your content in order to maximize performance 
-on mobile platforms. These are listed from most severe to least.
+Below are things to watch out for or suggestions on how to set up your content to maximize performance on mobile platforms. These are listed from most severe to least.
 
-* Make sure lighting is built before running on device.
-* Make sure most post process features are disabled.  In fact, the only ones you can use for iPad4 are Temporal AA and the Film ones, including vignette.  A bunch of expensive features like Bloom and DOF are enabled by default as we have the same settings between PC and mobile.  Post Process features can easily cost 60+ms with the default settings.  We hope to be able to use DOF and light shafts on higher end devices.  `show postprocessing` can be used on devices to toggle PP and see how much it is costing.
+* Make sure lighting is built before running on a device.
+* Make sure most Post Process features are disabled.  In fact, the only ones you can use for iPad4 are Temporal AA and the Film ones, including vignette.  A bunch of expensive features like Bloom and DOF are enabled by default as we have the same settings between PC and mobile.  Post Process features can easily cost 60+ms with the default settings.  We hope to be able to use DOF and light shafts on higher end devices.  `show postprocessing` can be used on devices to toggle PP and see how much it is costing.
 * Make sure precomputed visibility is setup properly.  To do this, you need to place PrecomputedVisibilityVolumes around anywhere that the player can walk, and build lighting.  You need to make sure the same P level is used when building lighting and when running the game (aka do not build the sublevel by itself).  You can verify that it is working by entering `stat initviews` on device or in previewer and make sure the `Statically occluded primitives` is > 0.  Use `r.ShowPrecomputedVisibilityCells 1` to visualize the cells in the editor.
 * Use masked and translucent Materials sparingly, and only in places where they cover a small part of the screen.  iOS devices are very optimal in shading opaque surfaces, as they will only shade each pixel once, most of the time.  But for masked and translucency, every single layer will have to be shaded.  The total GPU time of a frame can easily be doubled or more by having overdraw.  Use the Shader complexity viewmode in Mobile previewer to investigate hot spots.
-* Draw calls of the entire scene should be &lt;= 700 for any view.  The areas with poor occlusion, like looking over the landscape, are going to be the biggest challenges for this.  This can be seen with `stat openglrhi` on device, `stat d3d11rhi` in ES2 previewer on PC.
-* Triangle count of the entire scene should be &lt;= 500k for any view.  This has been determined to be the poly count needed to hit 30fps on both iPad4 and iPad Air.  This can be seen with `stat openglrhi` on device, `stat d3d11rhi` in ES2 previewer on PC.
+* Draw calls of the entire scene should be &lt;=700 for any view.  The areas with poor occlusion, like looking over the landscape, are going to be the biggest challenges for this.  This can be seen with `stat openglrhi` on device, `stat d3d11rhi` in ES2 previewer on PC.
+* Triangle count of the entire scene should be &lt;=500k for any view.  This has been determined to be the poly count needed to hit 30fps on both iPad4 and iPad Air.  This can be seen with `stat openglrhi` on device, `stat d3d11rhi` in ES2 previewer on PC.
 * Materials should use as few texture lookups and instructions as possible.  You have 5 texture samplers available, but using all of those makes for a fairly expensive Material.  For now, use the PC instruction count to optimize, and ES2 previewer shader complexity to visualize the total cost.
 * Only use independent texture fetches in Materials.  This means UVs in the pixel shader (BaseColor, Roughness, etc) must not be manipulated in any way, such as scaling.  Use the new CustomizedUVs feature instead to do the scaling in the vertex shader.  Some special features like environment mapping require math on the UVs and that is ok for special cases. 
 * Prefer square textures as they have less memory waste, and always use power of 2 dimensions (512, 1024, 2048).  Use `listtextures` in the ES2 previewer to see where all the texture memory is going. 
 
 ## Performance Tiers
 
-UE4 supports a variety of lighting features on mobile devices. Using these features costs performance and may cause your game to perform poorly on slower mobile devices. While it is possible to mix and match most of UE4's mobile lighting features, it can be useful to categorize these features in to the following tiers. When building a mobile game, you should decide which features to use based on the quality of graphics your game requires and the types of devices you need to support. Check the device compatibility pages under [iOS Development](Platforms/iOS/DeviceCompatibility) and [Android Development](Platforms/Android/DeviceCompatibility) for more information on what devices we have tested here at Epic and which tiers we think are most appropriate.
+UE4 supports a variety of lighting features on mobile devices. Using these features costs performance and may cause your game to perform poorly on slower mobile devices. While it is possible to mix and match most of UE4's mobile lighting features, it can be useful to categorize these features into the following tiers. When building a mobile game, you should decide which features to use based on the quality of graphics your game requires and the types of devices you need to support. Check the device compatibility pages under [iOS Development](Platforms/iOS/DeviceCompatibility) and [Android Development](Platforms/Android/DeviceCompatibility) for more information on what devices we have tested here at Epic and which tiers we think are most appropriate.
 
 ###LDR (Low Dynamic Range)
 
-Low dynamic range or LDR mode is the highest performance tier supported in UE4 and is recommended for games that do not require lighting or post processing features. In order to use this mode, you must disable `Mobile HDR` for your project in the Rendering section of the [Project Settings Editor](Engine/UI/ProjectSettings). 
+Low dynamic range or LDR mode is the highest performance tier supported in UE4 and is recommended for games that do not require lighting or Post Processing features. To use this mode, you must disable `Mobile HDR` for your project in the Rendering section of the [Project Settings Editor](Engine/UI/ProjectSettings). 
 
 **Advantages**
 
 * The fastest and lowest overhead mode available for mobile devices. Enables your game to run well on slower mobile devices.
-* Still provides full access to the Material editor to define custom shaders and even perform simple shading to fake lighting.
+* Still provides full access to the Material editor for defining custom shaders and even perform simple shading to fake lighting.
 
 
 **Limitations**
 
 * The scene's color is written out in gamma space with each color channel clamped to the range [0,1].
 * Translucent primitives are blended in gamma space. In most cases, this requires that you author your translucent textures and Materials differently than you would in HDR or for a normal PC game.
-* Post processing features are unavailable in this mode.
+* Post Processing features are unavailable in this mode.
 
 
 **Recommendations**
@@ -61,7 +60,7 @@ In this tier, you leverage static lighting and fully rough Materials to create l
 **Advantages**
 
 * Access to static lighting and global illumination features.
-* Full HDR pipeline with access to some post processing features like tone mapping.
+* Full HDR pipeline with access to some Post Processing features like tone mapping.
 * Translucency is blended in linear space allowing you to author content as you normally would for any other UE4 game.
 
 
@@ -76,17 +75,17 @@ In this tier, you leverage static lighting and fully rough Materials to create l
 * Author all of your Materials with the fully rough flag set.
 * Consider disabling light map directionality in your Materials for additional performance.
 * Use only static lights in your maps.
-* Disable bloom and anti-aliasing. Also stick to the basic set of film and color controls. See [mobile post process](Platforms/Mobile/PostProcessEffects) for more details.
+* Disable bloom and anti-aliasing. Also, stick to the basic set of film and color controls. See [mobile Post Process](Platforms/Mobile/PostProcessEffects) for more details.
 
 
 ### Full HDR Lighting
 
-In this tier, you take advantage of most of the HDR lighting features available for mobile in UE4 as well as some of the post processing features. Using these features requires quite a bit of performance in exchange for high quality lighting features.
+In this tier, you take advantage of most of the HDR lighting features available for mobile in UE4 as well as some of the Post Processing features. Using these features requires quite a bit of performance in exchange for high quality lighting features.
 
 **Advantages**
 
 * Access to static lighting and global illumination features.
-* Full HDR pipeline with access to some post processing features.
+* Full HDR pipeline with access to some Post Processing features.
 * Translucency is blended in linear space allowing you to author content as you normally would for any other UE4 game.
 * Realistic specular reflections on surfaces with support for varying roughness.
 * Full support for normal mapped surfaces allows interesting surface detail without adding a lot of polygons to your meshes.
@@ -95,7 +94,7 @@ In this tier, you take advantage of most of the HDR lighting features available
 **Recommendations**
 
 * Consider enabling bloom to take full advantage of the HDR lighting pipeline.
-* Realistic specular reflections in combination with HDR lighting can lead to specular aliasing. You should take advantage of the normal map to roughness feature to help reduce specular aliasing due to high frequency information in your normal map. You should also consider enabling anti-aliasing in the post process volume to help reduce this artifact.
+* Realistic specular reflections in combination with HDR lighting can lead to specular aliasing. You should take advantage of the normal map to roughness feature to help reduce specular aliasing due to high frequency information in your normal map. You should also consider enabling anti-aliasing in the Post Process volume to help reduce this artifact.
 * Place reflection captures intelligently to get the best results.
 * Use only static lights in your scene.
 
@@ -132,10 +131,10 @@ examples from the mobile showcase:
 
 <!--
 [EXCERPT:Image1]
-The pillar Materials are fairly expensive here, and the translucent volumetric sheets are VERY expensive.  In this scene I asked to have the translucent sheets removed, as they cost too much. 
+The pillar Materials are fairly expensive here, and the translucent volumetric sheets are VERY expensive.  In this scene, I asked to have the translucent sheets removed, as they cost too much. 
 [/EXCERPT:Image1]
 [EXCERPT:Image2]
-Here the pillars are quite expensive as they used all 5 texture lookups and did quite a bit of layering. Otherwise this is just about efficient enough for 30fps.
+Here the pillars are quite expensive as they used all five texture lookups and did quite a bit of layering. Otherwise, this is just about efficient enough for 30fps.
 [/EXCERPT:Image2]
 [EXCERPT:Image3]
 Here the tree is causing pixels to be EXTREMELY expensive.  If the player could ever get up close or cover his screen with that Material, the cost would be enormous.
@@ -145,16 +144,21 @@ Here the tree is causing pixels to be EXTREMELY expensive.  If the player could
 ## Mobile Content Scale Factor 
 Mobile Content Scale Factor is a way to scale the resolution of your project to best fit the screen resolution of the mobile device being used to view your project. 
 You can create and store multiple device profiles by creating a new .INI file called **DefaultDeviceProfiles.INI** in your project's config folder. 
-Inside of this file you can specify what will happen to the resolution of your project by entering the following command. 
+Inside of this file, you can specify what will happen to the resolution of your project by entering the following command. 
 
 	 r.MobileContentScaleFactor = 
 
-In the following image you can see an example of the Mobile Content Scale Factor in action using the command listed above.
+In the following image, you can see an example of the Mobile Content Scale Factor in action using the command listed above.
 
-![](T_MCF_INI_File.png)
+[REGION:lightbox]
+[![](T_MCF_INI_File.png)(w:358 h:505)](T_MCF_INI_File.png)
+[/REGION]
+[REGION:caption]
+Click for full image.
+[/REGION]
 
 The image above was taken from the UE4 Tappy Chicken project and shows what will happen to the resolution of Tappy Chicken when it is played on various mobile devices.
-The top section of the file handles resolution scaling for iOS devices and the bottom section handles resolution scaling for Android devices.
+The top section of the file handles resolution scaling for iOS devices, and the bottom section handles resolution scaling for Android devices.
 Make sure to note that each **r.MobileContentScaleFactor** has a number after it.
 The number that is used after the command is what determines if the projects resolution should be scaled up or down.
 
@@ -185,3 +189,145 @@ Please note that the actual resolution for any scale factor other than 0.0 will
 
 **r.MobileContentScaleFactor 2.0**: This will try and give you a resolution of 2560 x 1440 for Landscape and 1440 x 2560 for portrait.
 
+## Mobile Material Quality Settings
+When building Content for an Unreal Engine 4 (UE4) project that will run on both low and high end mobile devices you will often run into issues with certain features working on one set of devices but not working on another set of devices. You also might encounter issues with your artwork looking correct on one device but completely different on another. While there are a lot of ways to address problems like this, many of theses ways are both time and resource intensive and also extremely error prone. To address these issues UE4 has the Material Quality Level system. This system allows you to build one single Material that can then be used on a wide range of different devices, giving you full control over which devices use which features. In the following sections we will go over these systems and how you can use them in your UE4 projects:
+
+### Previewing Material Quality Level
+You can view what the different Material Quality level settings will look like inside of the UE4 Editor by going to the
+ **Main Toolbar** and then going to **Settings** > **Material Quality Level** and selecting the quality level you want to preview.
+
+[REGION:lightbox]
+[![](MQS_Material_Preview_Level.png)(w:332 h:324)](MQS_Material_Preview_Level.png)
+[/REGION]
+[REGION:caption]
+Click for full image.
+[/REGION]
+
+The following images show what a Material will look like when the Material Quality Level has been set to Low, Medium, and High.
+
+[REGION:fullwidth]
+[INCLUDE:Platforms\Mobile\Performance\#MMQ]
+[/REGION]
+
+### Setting Material Quality Level
+You can set the Material Quality Level in the following ways inside UE4:
+
+**From the Console**: You can set the quality level that your Materials will use in the Editor by pressing the Backtick key to open up the UE4 console and entering **r.MaterialQualityLevel ** followed by one of the following numbers.
+
+![](Platforms\Mobile\Performance\MQS_Console_Command.png)
+
+|Material Quality Level Number| Setting|
+|-------------------------------------|----------|
+|0| Low|
+|1|High|
+|2|Medium|
+
+**From Device Profiles**: You can also adjust the Material Quality by doing the following:
+
+1. From the Main Toolbar go to **Windows** > **Developer Tools** and then click on **Device Profiles**.  
+	[REGION:lightbox]
+	[![](MQS_Device_Profiles.png)(w:217 h:371)](MQS_Device_Profiles.png)
+	[/REGION]
+	[REGION:caption]
+	Click for full image.
+	[/REGION]
+
+1. In the Device Profiles find the **Android_Low** profile and click on the first **Wrench** icon to open up the profile options. 
+	
+	![](MQS_Edit_Profile.png)
+
+1. In the **Console Variables** section under the **Rendering** option, click on the **Plus** icon and in the input window that pops up, input **R.MaterialQualityLevel** and then press the Enter key to add that command to the list. 
+
+	[REGION:lightbox]
+	[![](MQS_Set_Con_Var.png)(w:505 h:242)](MQS_Set_Con_Var.png)
+	[/REGION]
+	[REGION:caption]
+	Click for full image.
+	[/REGION]
+
+1. Once the new entry has been added, change the value from the default value from **1** to **0** so that when this project is viewed on lower end Android devices it will use the fastest possible Material settings.
+
+	[REGION:lightbox]
+	[![](MQS_Set_Quality_Level.png)(w:505 h:242)](MQS_Set_Quality_Level.png)
+	[/REGION]
+	[REGION:caption]
+	Click for full image.
+	[/REGION]
+
+### Previewing Different Quality Levels
+You preview from the UE4 Editor how your project will look when running on different hardware by adjusting the **Preview Rendering Level**. You can adjust this by going to the **Main Toolbar** and then go to **Settings** > **Preview Rendering Level** and select the preview level you want to use.
+
+
+[REGION:lightbox]
+[![](MQS_Material_Quality_Level.png)(w:332 h:324)](MQS_Material_Quality_Level.png)
+[/REGION]
+[REGION:caption]
+Click for full image.
+[/REGION]
+
+### Material Quality Level Switch
+The **Quality Switch** Material Expression allows you to define different levels of complexity within a single Material. For example, if you have a complex math operation or heavy texture reads that will work on higher end mobile devices but not on lower end mobile devices you can use this node to specify simplified versions to be displayed on lower end mobile devices that do not have the power to display higher end graphics.
+
+![](image08.png)
+
+To use the Quality Switch in your Materials, you will need to first add the Quality Switch Material Expression node to the Material Graph and then plug it's output into any if the inputs on the Main Material node you want to change when the quality level is changed. In the following image, the M_Brick_Clay_Old Material that can be found with the Starter Content was set up so that when the Quality Level is switched the complexity of the Material is reduced. 
+
+[REGION:lightbox]
+[![](MQS_Material_Setup.png)(w:789 h:412)](MQS_Material_Setup.png)
+[/REGION]
+[REGION:caption]
+Click for full image.
+[/REGION]
+
+### Mobile Material Options 
+If you open up any Material and click on the Main Material Node then look in the Details Panel under the **Mobile** section you will find the following two properties.
+
+* **Fully Rough**: Forces a Material to be completely rough saving a number of instructions and a single texture sampler.
+* **Use Lightmap Directionality**: Uses lightmap directionally and per pixel normals. If disabled the lightmaps will be flatter but cheaper.
+
+[REGION:lightbox]
+[![](MQS_Mobile_Material_Options.png)(w:556 h:410)](MQS_Mobile_Material_Options.png)
+[/REGION]
+[REGION:caption]
+Click for full image.
+[/REGION]
+
+These two properties can be enabled to help reduce the rendering cost of that Material when viewed on less powerful mobile devices by completely eliminating the rendering path for those features. However this is, an all or nothing option as you either have these options enabled or disabled for this Material on all devices.
+
+
+### Mobile Material Rendering Overrides
+You can override the available rendering option a platform will use by adjusting the **ES2 Quality Settings**. You can find the ES2 Quality Setting for both Android and iOS by going to your **Project Settings** > **Platforms** and then selecting either **Android ES2 Quality** or **iOS ES2 Quality Settings**. To use the overrides you must first click on the **Enable Quality Override** option and then select the option you want to override by clicking on it. Once all of the options you want have been selected, press the **Update preview shaders** button to re-compile all Materials to work with the options you have specified. 
+
+
+[REGION:lightbox]
+[![](MQS_Rendering_Overrides.png)(w:613 h:330)](MQS_Rendering_Overrides.png)
+[/REGION]
+[REGION:caption]
+Click for full image.
+[/REGION]
+
+
+
+<!-----
+[EXCERPT:MMQ]
+[OBJECT:SliderSequence]
+	[PARAMLITERAL:max]
+	3
+	[/PARAMLITERAL]
+	[PARAMLITERAL:ext]
+	jpg
+	[/PARAMLITERAL]
+	[PARAMLITERAL:id]
+	MMQ
+	[/PARAMLITERAL]
+	[PARAM:caption]
+	**Dragging the slider will show what happens when changing the Material Quality Level from Low to Medium and High **
+	[/PARAM]
+	[PARAM:image_list]
+	![](MMQ_1.png)
+	![](MMQ_2.png)
+	![](MMQ_3.png)
+	[/PARAM]
+[/OBJECT]
+[/EXCERPT:MMQ]
+-->
\ No newline at end of file
