Adding a button link to the Setting up and Unreal Workflow How To Series.

#UE4docs

diff --git a/Engine/Documentation/Source/Platforms/Linux/BeginnerLinuxDeveloper/BeginnerLinuxDeveloper.INT.udn b/Engine/Documentation/Source/Platforms/Linux/BeginnerLinuxDeveloper/BeginnerLinuxDeveloper.INT.udn
--- a/Engine/Documentation/Source/Platforms/Linux/BeginnerLinuxDeveloper/BeginnerLinuxDeveloper.INT.udn
+++ b/Engine/Documentation/Source/Platforms/Linux/BeginnerLinuxDeveloper/BeginnerLinuxDeveloper.INT.udn
@@ -1,7 +1,7 @@
 Availability: Docs
 Title: Beginner Linux Developers
 Crumbs:%ROOT%, Platforms, Platforms/Linux
-Description:Learn how to setup a typical workflow for developing games on Linux
+Description:Learn how to setup a typical workflow for developing games on Linux.
 Platform:Linux
 Navigation:platform
 Version: 4.11
@@ -39,6 +39,42 @@ If you're a fledgling developer, you'll want to start off with this learning pat
 By the end of this learning path, you'll setup an Integrated Development Environment (IDE), you'll have a thorough understanding of the Unreal Engine 4 (UE4) 
 user interface, and you'll gain some experience by developing a game on Linux.
 
-[REGION:topics third]
-%Platforms/Linux/BeginnerLinuxDeveloper/SupportedGfxDrivers:topic%
+[REGION:features third]
+
+[OBJECT:Role]
+	[PARAM:icon]
+		![](%ROOT%/linux_icon.png)(convert:false)
+	[/PARAM]
+	[PARAM:title]
+		%Platforms/Linux/BeginnerLinuxDeveloper/SupportedGfxDrivers:title%
+	[/PARAM]
+	[PARAM:path]
+		[RELATIVE:Platforms/Linux/BeginnerLinuxDeveloper/SupportedGfxDrivers]
+	[/PARAM]
+	[PARAM:description]
+		%Platforms/Linux/BeginnerLinuxDeveloper/SupportedGfxDrivers:description%
+	[/PARAM]
+	[PARAMLITERAL:id]
+		supportedgfxdrivers
+	[/PARAMLITERAL]
+[/OBJECT]
+
+[OBJECT:Role]
+	[PARAM:icon]
+		![](%ROOT%/linux_icon.png)(convert:false)
+	[/PARAM]
+	[PARAM:title]
+		%Platforms/Linux/BeginnerLinuxDeveloper/SettingUpAnUnrealWorkflow:title%
+	[/PARAM]
+	[PARAM:path]
+		[RELATIVE:Platforms/Linux/BeginnerLinuxDeveloper/SettingUpAnUnrealWorkflow]
+	[/PARAM]
+	[PARAM:description]
+		%Platforms/Linux/BeginnerLinuxDeveloper/SettingUpAnUnrealWorkflow:description%
+	[/PARAM]
+	[PARAMLITERAL:id]
+		supportedgfxdrivers
+	[/PARAMLITERAL]
+[/OBJECT]
+
 [/REGION]
\ No newline at end of file
