Adding a closing HTML note tag so that the windows and mac sections will show the correct info.

diff --git a/Engine/Documentation/Source/Platforms/Android/LegacyAndroidDevelopment/LegacyAndroidDevelopment.INT.udn b/Engine/Documentation/Source/Platforms/Android/LegacyAndroidDevelopment/LegacyAndroidDevelopment.INT.udn
--- a/Engine/Documentation/Source/Platforms/Android/LegacyAndroidDevelopment/LegacyAndroidDevelopment.INT.udn
+++ b/Engine/Documentation/Source/Platforms/Android/LegacyAndroidDevelopment/LegacyAndroidDevelopment.INT.udn
@@ -209,7 +209,7 @@ TADP installer and Android SDK tools require that both Java 6 and Java 7 are ins
 
 [/EXCERPT:MacOSSetup]
 
-
+-->
 
 
 
