#UE4 Docs: Replaced all referances to old Android SDK's with CodeWorks for Android 1R4.

diff --git a/Engine/Documentation/Source/Platforms/Android/GettingStarted/1/AndroidQS_1.INT.udn b/Engine/Documentation/Source/Platforms/Android/GettingStarted/1/AndroidQS_1.INT.udn
--- a/Engine/Documentation/Source/Platforms/Android/GettingStarted/1/AndroidQS_1.INT.udn
+++ b/Engine/Documentation/Source/Platforms/Android/GettingStarted/1/AndroidQS_1.INT.udn
@@ -1,7 +1,7 @@
 Availability:Public
 Title:1. Install the Android SDK
 Crumbs:%ROOT%, Platforms, Platforms/Android, Platforms/Android/GettingStarted
-Description: Install the Android Works SDK.
+Description: Install the CodeWorks for Android 1R4 SDK.
 Version:4.9
 Skilllevel:Intermediate
 
@@ -24,9 +24,9 @@ Skilllevel:Intermediate
 %Steps%
 
 
-The first thing we need to do to start developing for Android is install Android Works, which will 
+The first thing we need to do to start developing for Android is to install CodeWorks for Android 1R4, which will 
 ensure that you have all the required programs and code needed to deploy to Android devices.
-In the following section, we will cover how you can get Android Works installed on your development PC.
+In the following section, we will cover how you can get CodeWorks for Android 1R4 installed on your development PC.
 
 
 %Globals:OSSelect%
@@ -65,17 +65,17 @@ In the following section, we will cover how you can get Android Works installed
 [EXCERPT:WindowsSetup]
 
 [REGION:note]
-If you already have Android Works installed, but your environment variables have been lost, you can skip all of this and simply point the editor to the install locations using the **Edit Menu** > **Project Settings** > **Android SDK** properties in the Editor.
+If you already have CodeWorks for Android 1R4 installed, but your environment variables have been lost, you can skip all of this and simply point the editor to the install locations using the **Edit Menu** > **Project Settings** > **Android SDK** properties in the Editor.
 [/REGION]
 
 [region:warning]
-It is highly recommended that you install Android Works while the editor and launcher are both closed to ensure that there are not problems with the install.
+It is highly recommended that you install CodeWorks for Android 1R4 while the editor and launcher are both closed to ensure that there are not problems with the install.
 [/region]
 
-1. To install Android Works navigate to the **Engine** > **Extras** > **Android** folder. 
+1. To install CodeWorks for Android 1R4 navigate to the **Engine** > **Extras** > **Android** folder. 
 
 	[region:note]
-	The Android Works installer requires an installation of Visual Studio to complete successfully. Supported versions are 2010, 2012 or 2013 (not 2015).
+	The CodeWorks for Android 1R4 installer requires an installation of Visual Studio to complete successfully. Supported versions are 2010, 2012 or 2013 (not 2015).
 	[/region]
 
 	| Install Type | Location |
@@ -84,13 +84,13 @@ It is highly recommended that you install Android Works while the editor and lau
 	| GitHub | [EngineSyncLocation\UE4\Engine\Extras\Android |
 
 	[region:note]
-	If you downloaded UE4 from **Github** you will need to perform the following steps before you can run the Android Works installer. 
+	If you downloaded UE4 from **Github** you will need to perform the following steps before you can run the CodeWorks for Android 1R4 installer. 
 	Please note that the follow applies to both Windows and Mac based PCs that download the UE4 source code from Github.
-	If you are using the binary build downloaded from the launcher then you can skip this part and go right into installing Android Works. 
+	If you are using the binary build downloaded from the launcher then you can skip this part and go right into installing CodeWorks for Android 1R4. 
 
-	1. Once UE4 has been downloaded from Github and extracted, run the **Setup.bat** file to download the required files as well as the **Android Works** installer.
+	1. Once UE4 has been downloaded from Github and extracted, run the **Setup.bat** file to download the required files as well as the **CodeWorks for Android 1R4** installer.
 
-	1. After  **Setup.bat** has finished running locate and run the Android Works installer following the instructions below.
+	1. After the **Setup.bat** has finished running locate and run the CodeWorks for Android 1R4 installer following the instructions below.
 
 	[/region]
 
@@ -106,7 +106,7 @@ It is highly recommended that you install Android Works while the editor and lau
 
 	![](AndroidWorks_01.png)
 
-1.	Set a location to download and install Android Works and then click the **Next** button when done.
+1.	Set a location to download and install CodeWorks for Android 1R4 and then click the **Next** button.
 
 	![](AndroidWorks_02.png)
 
@@ -114,74 +114,79 @@ It is highly recommended that you install Android Works while the editor and lau
 	It is perfectly fine to leave this set to the default location.
 	[/region]
 
-1.	Make sure that your are installing **AndroidWorks1R1** and then click the **Next** button.
+1.	Make sure that your are installing **CodeWorks for Android 1R4** and then click the **Next** button.
 	
 	![](AndroidWorks_03.png)
 
-1.	Click on the checkmark box, **Accept All** if you agree with the license terms and then press the **Accept** button.
+1.	Click on the check mark box, **Accept All** if you agree with the license terms and then press the **Accept** button.
 	
 	![](AndroidWorks_04.png)
 
-1.	A new window will be displayed that will show you the download and install progress of the required Android Works tools. 
+1.	A new window will be displayed that will show you the download and install progress of the required CodeWorks for Android 1R4 tools. 
 	[region:note]
 	You will need an internet connection to complete this portion. Depending on the speed of that internet connection this could take a few minute to a few hours to complete.
 	[/region]
 	
 	![](AndroidWorks_06.png)
 
-1.	When all required tools have been downloaded and installed the following window will be displayed. Press the **OK** button continue the install process.
+1.	When all required tools have been downloaded and installed the following window will be displayed. Press the **Next** button continue the install process.
 	
 	![](AndroidWorks_08.png)
 
-1.	Now press the **Finish** button to complete the Android Works install.
+1.	Now press the **Next** button to continue the CodeWorks for Android 1R4 install.
 	
 	![](AndroidWorks_09.png)
 
-1.	Finally press the **Yes** button to re-start the computer.
+1.	Then press the **Finish** button to complete the install.
 	
 	![](AndroidWorks_10.png)
+
+1. Finally press the **Yes** button to restart your PC.
+
+	![](AndroidWorks_12.png)
+
 	[region:note]
-	If you downloaded UE4 from Github, when Android Works has finished installing, restart your PC or Mac so the environment variables are set and then run the **Setup.bat** file one more time to download the required Android Development files.
+	If you downloaded UE4 from Github, when CodeWorks for Android 1R4 has finished installing, restart your PC or Mac so the environment variables are set and then run the **Setup.bat** file one more time to download the required Android Development files.
 	[/region]
 [/EXCERPT:WindowsSetup]
 
 [EXCERPT:MacOSSetup]
 
 [REGION:note]
-If you already have Android Works installed, but your environment variables have been lost, you can skip all of this and simply point the editor to the install locations using the **Edit Menu** > **Project Settings** > **Android SDK** properties in the Editor.
+If you already have CodeWorks for Android 1R4 installed, but your environment variables have been lost, you can skip all of this and simply point the editor to the install locations using the **Edit Menu** > **Project Settings** > **Android SDK** properties in the Editor.
 [/REGION]
 
 [region:warning]
-It is highly recommended that you install Android Works while the editor and launcher are both closed to ensure that there are not problems with the install.
+It is highly recommended that you install CodeWorks for Android 1R4 while the editor and launcher are both closed to ensure that there are not problems with the install.
 [/region]
-Before you can install Android Works to your Mac you will need to make sure that both Java 6 and Java 7 are installed. 
+Before you can install CodeWorks for Android 1R4 to your Mac you will need to make sure that both Java 6 and Java 7 are installed. 
 Use the links provide below to download and install the correct versions of Java.
 
 * **Java 6** can be downloaded from [here](http://support.apple.com/kb/DL1572).
 * **Java 7** can be downloaded from [here](http://www.oracle.com/technetwork/java/javase/downloads/jre7-downloads-1880261.html).
 
-1. To install Android Works navigate to the **Engine** > **Extras** > **Android** folder. 
+1. To install CodeWorks for Android 1R4 navigate to the **Engine** > **Extras** > **Android** folder. 
 	| Install Type  | Location  |
 	| --- | --- |
 	| Binary | [EngineInstallLocation]\Unreal Engine[EngineVersion]\Engine\Extras\Android |
 	| GitHub | [EngineSyncLocation\UE4\Engine\Extras\Android |
 
 	[region:note]
-	If you downloaded UE4 from **Github** you will need to perform the following steps before you can run the Android Works installer. 
+	If you downloaded UE4 from **Github** you will need to perform the following steps before you can run the CodeWorks for Android 1R4 installer. 
 	Please note that the follow applies to both Windows and Mac based PCs that download the UE4 source code from Github.
-	If you are using the binary build downloaded from the launcher then you can skip this part and go right into installing Android Works. 
+	If you are using the binary build downloaded from the launcher then you can skip this part and go right into installing CodeWorks for Android 1R4. 
 
-	1. Once UE4 has been downloaded from Github and extracted, run the **Setup.bat** file to download the required files as well as the Android Works installer.
+	1. Once UE4 has been downloaded from Github and extracted, run the **Setup.bat** file to download the required files as well as the CodeWorks for Android 1R4 installer.
 
-	1. After  **Setup.bat** has finished running locate and run the Android Works installer following the instructions below.
+	1. After  **Setup.bat** has finished running locate and run the CodeWorks for Android 1R4 installer following the instructions below.
 
 	[/region]
 
-1. **Double-click** on the **AndroidWorks-1R1-osx.dmg** to load the Android Works installer.
+1. **Double-click** on the **AndroidWorks-1R1-osx.dmg** to load the CodeWorks for Android 1R4 installer.
 	
 	![](AndroidWorks_Mac_00.png)
 
-1. **Double-click** the Android Works installer icon to start the install process.
+1. **Double-click** the CodeWorks for Android 1R4 installer icon to start the install process.
 
 	![](AndroidWorks_Mac_01.png)
 
@@ -193,7 +198,7 @@ Use the links provide below to download and install the correct versions of Java
 
 	![](AndroidWorks_Mac_03.png)
 
-1. Specify a location where Android Works should be installed then press the **Next** button.
+1. Specify a location where CodeWorks for Android 1R4 should be installed then press the **Next** button.
 
 	[region:note]
 	It is perfectly fine to leave this set to the default location.
@@ -201,30 +206,30 @@ Use the links provide below to download and install the correct versions of Java
 
 	![](AndroidWorks_Mac_04.png)
 
-1. When the the Android Works Component Manager is displayed press the **Next** button to begin the download and install process.
+1. When the CodeWorks for Android 1R4 Component Manager is displayed press the **Next** button to begin the download and install process.
 
 	![](AndroidWorks_Mac_05.png)
 
-1. Click on the checkmark box, **Accept All** if you agree with the license terms and then press the **Accept** button.
+1. Click on the check mark box, **Accept All** if you agree with the license terms and then press the **Accept** button.
 
 	![](AndroidWorks_Mac_06.png)
 
-1. The Android Works Component Manager will now begin to download and install the required tools. 
+1. The CodeWorks for Android 1R4 Component Manager will now begin to download and install the required tools. 
 	[region:note]
 	For this part of the setup will take a few minutes to completed and you will need an internet connection to complete it.
 	[/region]
 	![](AndroidWorks_Mac_07.png)
 
-1. Once everything has been installed the following window will be displayed. Click on the **OK** button to continue.
+1. Once everything has been installed the following window will be displayed. Click on the **Next** button to continue.
 
 	![](AndroidWorks_Mac_08.png)
 
-1. Click on the **Finish** button to complete the Android Works install.
+1. Click on the **Finish** button to complete the CodeWorks for Android 1R4 install.
 
 	![](AndroidWorks_Mac_09.png)
 
 	[region:note]
-	If you downloaded UE4 from Github, when Android Works has finished installing, restart your PC or Mac so the environment variables are set and then run the **Setup.bat** file one more time to download the required Android Development files.
+	If you downloaded UE4 from Github, when CodeWorks for Android 1R4 has finished installing, restart your Mac so the environment variables are set and then run the **Setup.bat** file one more time to download the required Android Development files.
 	[/region]
 
 [/EXCERPT:MacOSSetup]
diff --git a/Engine/Documentation/Source/Platforms/Android/GettingStarted/AndroidQuickStart.INT.udn b/Engine/Documentation/Source/Platforms/Android/GettingStarted/AndroidQuickStart.INT.udn
--- a/Engine/Documentation/Source/Platforms/Android/GettingStarted/AndroidQuickStart.INT.udn
+++ b/Engine/Documentation/Source/Platforms/Android/GettingStarted/AndroidQuickStart.INT.udn
@@ -78,7 +78,7 @@ The focus of the Android Quick start is to walk you through the basics of settin
 [/OBJECT]
 
 [region:note]
-Unreal Engine 4 (UE4) now uses a special version of the Android SDK called **Android Works**. This version of Android works was designed specifically for use with UE4 Android based projects. If you have other versions of the Android SDK installed, it is **highly recommend** that you uninstall them and then re-install the Android SDK using the Android Works installer that is bundled with UE4. This way you can ensure that the Android SDK you are using will be compatible with UE4.
+Unreal Engine 4 (UE4) now uses a special version of the Android SDK called **CodeWorks for Android 1R4**. This version of CodeWorks for Android was designed specifically for use with UE4 Android based projects. If you have other versions of the Android SDK installed, it is **highly recommend** that you uninstall them and then re-install CodeWorks for Android using theCodeWorks for Android 1R4 installer that is bundled with UE4. This way you can ensure that the Android SDK you are using will be compatible with UE4.
 [/region]
 
 
diff --git a/Engine/Documentation/Source/Platforms/Android/Reference/AndroidGettingStarted.INT.udn b/Engine/Documentation/Source/Platforms/Android/Reference/AndroidGettingStarted.INT.udn
--- a/Engine/Documentation/Source/Platforms/Android/Reference/AndroidGettingStarted.INT.udn
+++ b/Engine/Documentation/Source/Platforms/Android/Reference/AndroidGettingStarted.INT.udn
@@ -51,7 +51,7 @@ Here is what we currently require for Android game development:
 
 There are several methods for getting the components. If you have already installed one of these, look in the corresponding section for additional requirements:
 
-### TADP (Tegra Android Development Pack)
+### CodeWorks for Android 1R4
 [OBJECT:ToggleButtonContent]
 	[PARAMLITERAL:category]
 	OS
@@ -84,11 +84,11 @@ There are several methods for getting the components. If you have already instal
 <!--
 [EXCERPT:TADPWindowsSetup]
 
-TADP is the Tegra Android Development Pack, by Nvidia. This is a one-stop shop for Android development. If you install TADP, all of our requirements will be met.
+CodeWorks for Android 1R4, by Nvidia is a one-stop shop for Android development needs. If you install CodeWorks for Android 1R4, all of your requirements to develop for Android will be met.
 
-The installer for TADP is included with Unreal Engine 4 for Windows. It can be found here: 
+The installer for CodeWorks for Android 1R4 is included with Unreal Engine 4 for Windows. It can be found here: 
 
-	[ENGINE INSTALL LOCATION]\Engine\Extras\Android\tadp-2.0r8-windows.exe 
+	[ENGINE INSTALL LOCATION]\Engine\Extras\Android\CodeWorksforAndroid-1R4-windows.exe
 
 The important pieces are:
 * SDK
@@ -99,7 +99,7 @@ The important pieces are:
 	
 If you use the default installation set, a pretty minimal set of components will be installed. If you customize the install, you can add more items (like profilers, etc). Note that during installation, 
 one or two of the sub-installers may appear to hang, but they will eventually finish. The install could take as long as 30 to 45 minutes. If you never plan to write code, you may want to 
-select **Customize**, and only install the SDK, Ant, and Java (skipping the NDK and Nsight debugger) to make for the quickest installation. The Nsight portion of TADP is used for debugging an Android game on a device. 
+select **Customize**, and only install the SDK, Ant, and Java (skipping the NDK and Nsight debugger) to make for the quickest installation. The Nsight portion of CodeWorks for Android 1R4 is used for debugging an Android game on a device. 
 It is not required to compile C++ code.
 
 ![](TADPSetup.png)
@@ -107,11 +107,11 @@ It is not required to compile C++ code.
 [/EXCERPT:TADPWindowsSetup]
 
 [EXCERPT:TADPMacOSSetup]
-TADP is the Tegra Android Development Pack, by Nvidia. This is a one-stop shop for Android development. If you install TADP, all of our requirements will be met.
+CodeWorks for Android 1R4, by Nvidia is a one-stop shop for Android development needs. If you install CodeWorks for Android 1R4, all of your requirements to develop for Android will be met.
 
-The installer for TADP is included with Unreal Engine 4 for Mac. It can be found here: 
+The installer for CodeWorks for Android 1R4 is included with Unreal Engine 4 for Mac and can be found here: 
 
-	[ENGINE INSTALL LOCATION]/Engine/Extras/Android/tadp-3.0r4-osx.dmg
+	[ENGINE INSTALL LOCATION]/Engine/Extras/Android/AndroidWorks-1R1-osx.dmg
 
 The important pieces are:
 
@@ -123,7 +123,7 @@ If you use the default installation set, a pretty minimal set of components will
 If you customize the install, you can add more items (like profilers, etc). Note that during installation, one or two of the sub-installers may appear to hang, but they will eventually finish. 
 The install could take as long as 30 to 45 minutes on some computers. 
 If you never plan to write code, you may want to select **Customize**, and only install the SDK, Ant, and Java (skipping the NDK and Nsight debugger) to make for the quickest installation.
-The Nsight portion of TADP is used for debugging an Android game on a device so it is not required to compile C++ code.
+The Nsight portion of CodeWorks for Android 1R4 is used for debugging an Android game on a device so it is not required to compile C++ code.
 
 ![](TADP_Mac_OS.png)
 
@@ -194,10 +194,10 @@ For code projects only, you will need to install the NDK:
 <!--
 [EXCERPT:ENVWindowsSetup]
 
-Unless you installed TADP, you will need to verify that some environment variables are set properly. 
+Unless you installed CodeWorks for Android 1R4 from the provided UE4 installer, you will need to verify that some environment variables are set properly. 
 
 [REGION:note]
-Once you have completed the installation (either by setting environment variables in the System Control Panel, or by installing TADP), you will need to restart both the Editor and Launcher 
+Once you have completed the installation (either by setting environment variables in the System Control Panel, or by installing CodeWorks for Android 1R4), you will need to restart both the Editor and Launcher 
 applications. This is due to how Windows applications store environment variables and then pass the old versions along to applications they start. 
 
 If you continue seeing the message about the Android SDK not being installed properly, it is recommended that you reboot your computer.
@@ -237,10 +237,10 @@ Here are the variables that we depend on, and where they should point:
 [/EXCERPT:ENVWindowsSetup]
 
 [EXCERPT:ENVMacOSSetup]
-Unless you installed TADP, you will need to verify that some environment variables are set properly. 
+Unless you installed CodeWorks for Android 1R4, you will need to verify that some environment variables are set properly. 
 
 [region:note]
-Once you have completed the installation (either by editing your .bash_profile file, or by installing TADP), you will need to restart both the Editor and Launcher applications.
+Once you have completed the installation (either by editing your .bash_profile file, or by installing CodeWorks for Android 1R4), you will need to restart both the Editor and Launcher applications.
 [/region]
 
 To set an environment variable you need to edit your **.bash_profile** file in your home directory. 
@@ -338,7 +338,7 @@ If your device shows up in the **Launch menu**, but gets stuck deploying to the
 If you click on the **Show Log** line in the progress bar, and see that the last line indicates that it is calling adb uninstall try the following.
 * Open up the Terminal app and run the adb devices command. 
 * If the device is listed but shows unauthorized, that indicates that you need to get the device to ask if you want to authorize the Mac. 
-* It is recommended you check the **Always allow** checkbox to always allow your Mac to talk to this device.
+* It is recommended you check the **Always allow** check box to always allow your Mac to talk to this device.
 
 [/EXCERPT:ATTMacOSSetup]
 
diff --git a/Engine/Documentation/Source/Platforms/GearVR/Debugging/GearVRDebugging.INT.udn b/Engine/Documentation/Source/Platforms/GearVR/Debugging/GearVRDebugging.INT.udn
--- a/Engine/Documentation/Source/Platforms/GearVR/Debugging/GearVRDebugging.INT.udn
+++ b/Engine/Documentation/Source/Platforms/GearVR/Debugging/GearVRDebugging.INT.udn
@@ -34,19 +34,43 @@ This can be any kind of application, download or compiled.
 
 1. First on your Samsung Galaxy Note 4 Smartphone go to **Settings** > **Application Manager** > **Gear VR Service**.
 	
-	![](T_GearVR_Dev_Mode_01.png)
+	[REGION:lightbox]
+	[![](T_GearVR_Dev_Mode_01.png)(w:600 h:355)](T_GearVR_Dev_Mode_01.png)
+	[/REGION]
+	
+	[REGION:caption]
+	Click for full image.
+	[/REGION]
 
 1. Next tap on the **Manage Storage** option.
   
-	![](T_GearVR_Dev_Mode_02.png)
+	[REGION:lightbox]
+	[![](T_GearVR_Dev_Mode_02.png)(w:250 h:444)](T_GearVR_Dev_Mode_02.png)
+	[/REGION]
+	
+	[REGION:caption]
+	Click for full image.
+	[/REGION]
 
 1. Then tap on the **VR Service Version** **6** times.
 
-	![](T_GearVR_Dev_Mode_03.png)
+	[REGION:lightbox]
+	[![](T_GearVR_Dev_Mode_03.png)(w:250 h:225)](T_GearVR_Dev_Mode_03.png)
+	[/REGION]
+	
+	[REGION:caption]
+	Click for full image.
+	[/REGION]
 
 1. You should now see a scan process happen and when completed you should see the **Developer Mode** toggle on the screen.
 
-	![](T_GearVR_Dev_Mode_04.png)
+	[REGION:lightbox]
+	[![](T_GearVR_Dev_Mode_04.png)(w:250 h:225)](T_GearVR_Dev_Mode_04.png)
+	[/REGION]
+	
+	[REGION:caption]
+	Click for full image.
+	[/REGION]
 
 1. Now you should be able to launch your Gear VR Project and have it display without having to use the Gear VR Headset.
 
@@ -67,26 +91,56 @@ You can get a full list of available console commands by entering **DumpConsoleC
 1. First make sure that you have enabled Gear VR Developer mode on your phone. 
 	This way you can launch your project and see it without having to use the Gear VR Headset which, if used, would prevent you from being able to touch your phone's screen.
 	
-	![](T_GearVR_Dev_Mode_04.png)
+	[REGION:lightbox]
+	[![](T_GearVR_Dev_Mode_04.png)(w:250 h:225)](T_GearVR_Dev_Mode_04.png)
+	[/REGION]
+	
+	[REGION:caption]
+	Click for full image.
+	[/REGION]
 
 1. Next launch your Gear VR project by tapping on the projects icon to launch it. 
 	When the project has been launched, you should see something looking like this on your phone.
 	
-	![](T_App_On_Phone.png)
+	[REGION:lightbox]
+	[![](T_App_On_Phone.png)(w:500 h:281)](T_App_On_Phone.png)
+	[/REGION]
+	
+	[REGION:caption]
+	Click for full image.
+	[/REGION]
 
 1. With your Gear VR project now running, tap the phone's screen with **Four (4)** fingers at once to open up the Unreal Engine Console. 
 	If done correctly you should see something like this on your phone's screen. 
 
-	![](T_GearR_Open_Console.png)
+	[REGION:lightbox]
+	[![](T_GearR_Open_Console.png)(w:500 h:281)](T_GearR_Open_Console.png)
+	[/REGION]
+	
+	[REGION:caption]
+	Click for full image.
+	[/REGION]
 
 1. Now input a console command you wish to use in the input box and then press the **Ok** button. 
 	For this example we will use the **Stat Memory** command to see how much memory the project is using. 
 
-	![](T_GearR_Input_Command.png)
+	[REGION:lightbox]
+	[![](T_GearR_Input_Command.png)(w:500 h:281)](T_GearR_Input_Command.png)
+	[/REGION]
+	
+	[REGION:caption]
+	Click for full image.
+	[/REGION]
 
 1. Once the **Ok** button has been pressed, your project should now be displaying how much memory is being used and by what.
 
-	![](T_GearR_Displaying_Memory.png)
+	[REGION:lightbox]
+	[![](T_GearR_Displaying_Memory.png)(w:500 h:281)](T_GearR_Displaying_Memory.png)
+	[/REGION]
+	
+	[REGION:caption]
+	Click for full image.
+	[/REGION]
 
 1. When you are done debugging or if you just wish to clear the screen of debug information, you can input **Stat None** to clear the screen off all debug information.
 
@@ -112,11 +166,23 @@ Check out the official Android documentation for much more detailed information
 
 1. Once the **OK** button has been toggled a Windows Command Line window like the following should now be displayed.
 	
-	![](T_Windows_CMD.png)
+	[REGION:lightbox]
+	[![](T_Windows_CMD.png)(w:422 h:269)](T_Windows_CMD.png)
+	[/REGION]
+	
+	[REGION:caption]
+	Click for full image.
+	[/REGION]
 
 1. Now check to make sure that your PC can see your Smartphone by entering **ADB Devices** into the command line window and then press the **Enter** key to display devices attached to your PC.
 	
-	![](T_Windows_ADB_Devices.png)
+	[REGION:lightbox]
+	[![](T_Windows_ADB_Devices.png)(w:422 h:269)](T_Windows_ADB_Devices.png)
+	[/REGION]
+	
+	[REGION:caption]
+	Click for full image.
+	[/REGION]
 
 	[region:tip]
 	If you enter **ADB Devices** but do not see anything show up, make sure that you have gone through the [](Platforms/GearVR) setup section and installed / setup everything that is listed there and then try this step again. 
@@ -127,18 +193,39 @@ Check out the official Android documentation for much more detailed information
 	To clear the log enter the following command.
 
 		adb logcat -c
-	![](T_ADB_Clear_Log.png)
+	[REGION:lightbox]
+	[![](T_ADB_Clear_Log.png)(w:422 h:269)](T_ADB_Clear_Log.png)
+	[/REGION]
+	
+	[REGION:caption]
+	Click for full image.
+	[/REGION]
 
 1. With the log now cleared of old data, launch the application on your Smartphone and enter the following command to get the log information.
 	
 		adb logcat -s UE4 OVR VrApi VrLib VRSVC HmdInfo >logoutput.txt
-	![](T_ADB_Get_Log.png)
+
+	[REGION:lightbox]
+	[![](T_ADB_Get_Log.png)(w:422 h:269)](T_ADB_Get_Log.png)
+	[/REGION]
+	
+	[REGION:caption]
+	Click for full image.
+	[/REGION]
+
  [region:note]
  To see the live output of the log, omit the **>logoutput.txt** part and all log updates will happen in the command line window.
  [/region]
 
 1. To complete the log capture press **CTRL + C** on the keyboard and when completed you should now have a log file named,** logoutput.txt**, that is saved to your **Users** folder in Windows.
-	![](T_ADB_Log_Location.png)
+
+	[REGION:lightbox]
+	[![](T_ADB_Log_Location.png)(w:381 h:338)](T_ADB_Log_Location.png)
+	[/REGION]
+	
+	[REGION:caption]
+	Click for full image.
+	[/REGION]
 
 	[region:note]
 	Data will still continue to be output to the log file while the command line window is open and running the ADB logcat command. To stop this all you have to do is press **CTRL + C** and data will stop being written to the log file.
@@ -148,7 +235,13 @@ Check out the official Android documentation for much more detailed information
 	When you look at the log file that was created and search for anyone of those terms you will find all debug information related to those terms. 
 	You can see an example of this in the image below, some of the terms used, OVR, UE4, Vrapi, and VrApi all show up in the log with information that is related to those terms.
 	
-	![](T_ADB_Log_Info.png)
+	[REGION:lightbox]
+	[![](T_ADB_Log_Info.png)(w:516 h:349)](T_ADB_Log_Info.png)
+	[/REGION]
+	
+	[REGION:caption]
+	Click for full image.
+	[/REGION]
 
 If your Gear VR Project is not starting up for some reason, using the logcat command is a great way to see exactly where your project is failing and why. 
 It also provides you with detailed information about the problem that can then be posted to the UE4 Forums or Answer Hub so that you can get help with the issue.
@@ -166,23 +259,48 @@ Please note that in order for this to work you have to have your development PC
 1. Next open up a Windows Command Shell by first pressing the **Windows Key + R** and then entering **CMD** in the **Run** dialogue box that is opened.
 	Then press either the **Enter** key on your keyboard or just click on the **OK** button with your mouse. 
 	 
-	 ![](T_Open_CMD_Prompt.png)
+	[REGION:lightbox]
+	[![](T_Open_CMD_Prompt.png)(w:427 h:269)](T_Open_CMD_Prompt.png)
+	[/REGION]
+	
+	[REGION:caption]
+	Click for full image.
+	[/REGION]
 
 1. Once the **OK** button has been toggled a Windows Command Line window like the following should now be displayed.
 	
-	![](T_Windows_CMD.png)
+	[REGION:lightbox]
+	[![](T_Windows_CMD.png)(w:427 h:269)](T_Windows_CMD.png)
+	[/REGION]
+	
+	[REGION:caption]
+	Click for full image.
+	[/REGION]
+
 
 1. Now enter the following command to see what your Smartphone's IP address is and make sure to take note of it. The IP address you want to use is the one that comes after **inet**.
 	
 		adb shell ip -f inet addr show wlan0
 	
-	![](T_ADB_Phone_IP.png)
+	[REGION:lightbox]
+	[![](T_ADB_Phone_IP.png)(w:427 h:269)](T_ADB_Phone_IP.png)
+	[/REGION]
+	
+	[REGION:caption]
+	Click for full image.
+	[/REGION]
 
 1. Next enter the following command to change the network connection to **TCP** and set the port to listen to.
 
 		adb tcpip 5555
 	
-	![](T_ADB_Set_Port.png)
+	[REGION:lightbox]
+	[![](T_ADB_Set_Port.png)(w:427 h:269)](T_ADB_Set_Port.png)
+	[/REGION]
+	
+	[REGION:caption]
+	Click for full image.
+	[/REGION]
 
 1. Once the port has been set, it is now time to connect to the Smartphone. To do this all you need to do is enter the following command and press the Enter key.
 	
@@ -194,7 +312,14 @@ Please note that in order for this to work you have to have your development PC
 
 1. To make sure you are connected over Wifi, enter **adb devices** into the command line window and press the Enter key and you will now see that in the **List of devices attached** the IP Address used above is now the only attached device.
 
-	![](T_ADB_Attached_IP.png)
+	[REGION:lightbox]
+	[![](T_ADB_Attached_IP.png)(w:427 h:269)](T_ADB_Attached_IP.png)
+	[/REGION]
+	
+	[REGION:caption]
+	Click for full image.
+	[/REGION]
+
 
 With your Smartphone and PC now connected to one another via the network you can now begin to get debug information from your Smartphone or even deploy projects to your phone without having to use the USB cable. 
 However keep in mind that deploying to your Smartphone from your PC over the network is four to five times slower than using the USB cable.
@@ -204,7 +329,13 @@ However keep in mind that deploying to your Smartphone from your PC over the net
 	
 		adb usb
 
-	![](T_ADB_Detached.png)
+	[REGION:lightbox]
+	[![](T_ADB_Detached.png)(w:427 h:269)](T_ADB_Detached.png)
+	[/REGION]
+	
+	[REGION:caption]
+	Click for full image.
+	[/REGION]
 
 
 
diff --git a/Engine/Documentation/Source/Platforms/GearVR/Prerequisites/GearVRPrerequisites.INT.udn b/Engine/Documentation/Source/Platforms/GearVR/Prerequisites/GearVRPrerequisites.INT.udn
--- a/Engine/Documentation/Source/Platforms/GearVR/Prerequisites/GearVRPrerequisites.INT.udn
+++ b/Engine/Documentation/Source/Platforms/GearVR/Prerequisites/GearVRPrerequisites.INT.udn
@@ -45,7 +45,7 @@ The USB cord is of vital importance as this is how you transfer data between you
 
 * Download Unreal Engine version **4.7.4** or later from the **[UE4 Launcher](https://www.unrealengine.com/)** or **[Github](https://www.unrealengine.com/ue4-on-github)**. 
 
-* Install the **Tegra Android Development Pack** which can be found in **Engine\Extras\Android**.  
+* Install the **CodeWorks for Android 1R4** which can be found in **Engine\Extras\Android**.  
 
 * Check to see if **[Developer Mode](https://docs.unrealengine.com/latest/INT/Platforms/Android/GettingStarted/2/index.html )** has been enabled on your Samsung Galaxy Note 4 Smartphone and if not, enable it. 
 
diff --git a/Engine/Documentation/Source/Platforms/GearVR/QuickStart/2/GearVRQuickStart_2.INT.udn b/Engine/Documentation/Source/Platforms/GearVR/QuickStart/2/GearVRQuickStart_2.INT.udn
--- a/Engine/Documentation/Source/Platforms/GearVR/QuickStart/2/GearVRQuickStart_2.INT.udn
+++ b/Engine/Documentation/Source/Platforms/GearVR/QuickStart/2/GearVRQuickStart_2.INT.udn
@@ -81,7 +81,7 @@ Version: 4.9
 	You can find the needed folders in the location where you installed the Tegra files. 
 	
 	[Region:note]
-	If you have already installed TADP this information should be already filled in for you. The only reason that you should change something here is if the information is missing or point to the wrong location.
+	If you have already installed CodeWorks for Android 1R4 this information should be already filled in for you. The only reason that you should change something here is if the information is missing or point to the wrong location.
 	[/region]
 
 	![](GVR_Android_SDK_Setup.png)
