#UE4 Docs: Making updates to the doc and its images.

diff --git a/Engine/Documentation/Source/Engine/Animation/Tools/VertexAnimationTool/VAT_TL_Meshes/VAT_TL_Meshes.INT.udn b/Engine/Documentation/Source/Engine/Animation/Tools/VertexAnimationTool/VAT_TL_Meshes/VAT_TL_Meshes.INT.udn
--- a/Engine/Documentation/Source/Engine/Animation/Tools/VertexAnimationTool/VAT_TL_Meshes/VAT_TL_Meshes.INT.udn
+++ b/Engine/Documentation/Source/Engine/Animation/Tools/VertexAnimationTool/VAT_TL_Meshes/VAT_TL_Meshes.INT.udn
@@ -15,7 +15,8 @@ In the following section we will cover how you can use the 3Ds Max Animation Tim
 	* **Radius:** 70.0
 	* **Segments:** 5
 	
-	![](VAT_Teapot_Setup_00.png)
+	[![](VAT_Teapot_Setup_00.png)
+
 
 1. Next from the Modifier List, apply a **Bend** Modifier to the Teapot and then set the **Direction** to **90**. 
 	
@@ -98,6 +99,7 @@ Open up this Texture by **Double-Clicking** on it in the Content Browser with th
 * Change the Texture Group to **UI**.
 * Make sure that **SRGB** has been disabled.
 * Set the Compression Settings to **TC_HDR**.
+* Set the Filter to **Nearest**.
 
 	![](VAT_EXR_Texture_Setup.png)
 
@@ -107,6 +109,7 @@ Next open up the Normal Map Texture and make sure that the following options are
 * Change the Texture Group to **UI**.
 * Make sure that **SRGB** has been disabled.
 * Set the Compression Settings to **TC_Vector Displacement Map**.
+* Set the Filter to **Nearest**.
 	
 	![](VAT_Normal_Texture_Setup.png)
 
diff --git a/Engine/Documentation/Source/Engine/Animation/Tools/VertexAnimationTool/VertexAnimationTools.INT.udn b/Engine/Documentation/Source/Engine/Animation/Tools/VertexAnimationTool/VertexAnimationTools.INT.udn
--- a/Engine/Documentation/Source/Engine/Animation/Tools/VertexAnimationTool/VertexAnimationTools.INT.udn
+++ b/Engine/Documentation/Source/Engine/Animation/Tools/VertexAnimationTool/VertexAnimationTools.INT.udn
@@ -25,7 +25,9 @@ This script is only available in Unreal Engine 4 (UE4) Version **4.8** or later.
 
 The Vertex Animation Script will be in the following location and is called **VertexAnimationTools.ms.**
 
-		4.8\Engine\Extras\3dsMaxScripts\VertexAnimationTools.ms
+		Engine\Extras\3dsMaxScripts\VertexAnimationTools.ms
+
+The newest version of the Verext Animation Script has been updated to correct for the way 3Ds Max 2014 and later handle image gamma.
 
 ### Tool Caveats
 While this tool is very good at reducing animation overhead by storing complex animation data into a texture, there are a few drawbacks to using the tool. 
@@ -177,6 +179,8 @@ Vertex positional morph target information is stored in a 16 bit signed floating
 A 32 bit image would be more precise but 16 should suffice for most FX work. 
 With that said, the offset vertex positions will be less precise as they move further away from their resting positions. 
 
+Also note that the scripts Textures should be sampled using the nearest neighbor method.
+
 ### Memory Usage
 The memory usage for each frame per vert is as follows :
 
