Completed draft Landing Page and implemented page navigation.

#UE4docs

diff --git a/Engine/Documentation/Source/Platforms/Linux/BeginnerLinuxDeveloper/SettingUpAnUnrealWorkflow/1/SettingUpAnUnrealWorkflow_1.INT.udn b/Engine/Documentation/Source/Platforms/Linux/BeginnerLinuxDeveloper/SettingUpAnUnrealWorkflow/1/SettingUpAnUnrealWorkflow_1.INT.udn
--- a/Engine/Documentation/Source/Platforms/Linux/BeginnerLinuxDeveloper/SettingUpAnUnrealWorkflow/1/SettingUpAnUnrealWorkflow_1.INT.udn
+++ b/Engine/Documentation/Source/Platforms/Linux/BeginnerLinuxDeveloper/SettingUpAnUnrealWorkflow/1/SettingUpAnUnrealWorkflow_1.INT.udn
@@ -9,3 +9,23 @@ order: 1
 tags:Linux
 SkillLevel:Beginner
 
+[VAR:Steps]
+[OBJECT:Navigation]
+	[PARAM:previous]
+	[/PARAM]
+	[PARAM:current]
+	[/PARAM]
+	[PARAM:home]
+		[Home](Platforms/Linux/BeginnerLinuxDeveloper/SettingUpAnUnrealWorkflow)
+	[/PARAM]
+	[PARAM:next]
+		[Next Step](Platforms/Linux/BeginnerLinuxDeveloper/SettingUpAnUnrealWorkflow/2 "%Platforms/Linux/BeginnerLinuxDeveloper/SettingUpAnUnrealWorkflow/2:title%")
+	[/PARAM]
+[/OBJECT]
+[/VAR]
+
+%Steps%
+
+
+
+%Steps%
\ No newline at end of file
diff --git a/Engine/Documentation/Source/Platforms/Linux/BeginnerLinuxDeveloper/SettingUpAnUnrealWorkflow/2/SettingUpAnUnrealWorkflow_2.INT.udn b/Engine/Documentation/Source/Platforms/Linux/BeginnerLinuxDeveloper/SettingUpAnUnrealWorkflow/2/SettingUpAnUnrealWorkflow_2.INT.udn
--- a/Engine/Documentation/Source/Platforms/Linux/BeginnerLinuxDeveloper/SettingUpAnUnrealWorkflow/2/SettingUpAnUnrealWorkflow_2.INT.udn
+++ b/Engine/Documentation/Source/Platforms/Linux/BeginnerLinuxDeveloper/SettingUpAnUnrealWorkflow/2/SettingUpAnUnrealWorkflow_2.INT.udn
@@ -9,3 +9,24 @@ order: 2
 tags:Linux
 SkillLevel:Beginner
 
+[VAR:Steps]
+[OBJECT:Navigation]
+	[PARAM:previous]
+		[Previous Step](Platforms/Linux/BeginnerLinuxDeveloper/SettingUpAnUnrealWorkflow/1 "%Platforms/Linux/BeginnerLinuxDeveloper/SettingUpAnUnrealWorkflow/1:title%")
+	[/PARAM]
+	[PARAM:current]
+	[/PARAM]
+	[PARAM:home]
+		[Home](Platforms/Linux/BeginnerLinuxDeveloper/SettingUpAnUnrealWorkflow)
+	[/PARAM]
+	[PARAM:next]
+		[Next Step](Platforms/Linux/BeginnerLinuxDeveloper/SettingUpAnUnrealWorkflow/3 "%Platforms/Linux/BeginnerLinuxDeveloper/SettingUpAnUnrealWorkflow/3:title%")
+	[/PARAM]
+[/OBJECT]
+[/VAR]
+
+%Steps%
+
+
+
+%Steps%
\ No newline at end of file
diff --git a/Engine/Documentation/Source/Platforms/Linux/BeginnerLinuxDeveloper/SettingUpAnUnrealWorkflow/3/SettingUpAnUnrealWorkflow_3.INT.udn b/Engine/Documentation/Source/Platforms/Linux/BeginnerLinuxDeveloper/SettingUpAnUnrealWorkflow/3/SettingUpAnUnrealWorkflow_3.INT.udn
--- a/Engine/Documentation/Source/Platforms/Linux/BeginnerLinuxDeveloper/SettingUpAnUnrealWorkflow/3/SettingUpAnUnrealWorkflow_3.INT.udn
+++ b/Engine/Documentation/Source/Platforms/Linux/BeginnerLinuxDeveloper/SettingUpAnUnrealWorkflow/3/SettingUpAnUnrealWorkflow_3.INT.udn
@@ -9,3 +9,24 @@ order: 3
 tags:Linux
 SkillLevel:Beginner
 
+[VAR:Steps]
+[OBJECT:Navigation]
+	[PARAM:previous]
+		[Previous Step](Platforms/Linux/BeginnerLinuxDeveloper/SettingUpAnUnrealWorkflow/2 "%Platforms/Linux/BeginnerLinuxDeveloper/SettingUpAnUnrealWorkflow/2:title%")
+	[/PARAM]
+	[PARAM:current]
+	[/PARAM]
+	[PARAM:home]
+		[Home](Platforms/Linux/BeginnerLinuxDeveloper/SettingUpAnUnrealWorkflow)
+	[/PARAM]
+	[PARAM:next]
+		[Next Step](Platforms/Linux/BeginnerLinuxDeveloper/SettingUpAnUnrealWorkflow/4 "%Platforms/Linux/BeginnerLinuxDeveloper/SettingUpAnUnrealWorkflow/4:title%")
+	[/PARAM]
+[/OBJECT]
+[/VAR]
+
+%Steps%
+
+
+
+%Steps%
\ No newline at end of file
diff --git a/Engine/Documentation/Source/Platforms/Linux/BeginnerLinuxDeveloper/SettingUpAnUnrealWorkflow/4/SettingUpAnUnrealWorkflow_4.INT.udn b/Engine/Documentation/Source/Platforms/Linux/BeginnerLinuxDeveloper/SettingUpAnUnrealWorkflow/4/SettingUpAnUnrealWorkflow_4.INT.udn
--- a/Engine/Documentation/Source/Platforms/Linux/BeginnerLinuxDeveloper/SettingUpAnUnrealWorkflow/4/SettingUpAnUnrealWorkflow_4.INT.udn
+++ b/Engine/Documentation/Source/Platforms/Linux/BeginnerLinuxDeveloper/SettingUpAnUnrealWorkflow/4/SettingUpAnUnrealWorkflow_4.INT.udn
@@ -9,3 +9,23 @@ order: 4
 tags:Linux
 SkillLevel:Beginner
 
+[VAR:Steps]
+[OBJECT:Navigation]
+	[PARAM:previous]
+		[Previous Step](Platforms/Linux/BeginnerLinuxDeveloper/SettingUpAnUnrealWorkflow/3 "%Platforms/Linux/BeginnerLinuxDeveloper/SettingUpAnUnrealWorkflow/3:title%")
+	[/PARAM]
+	[PARAM:current]
+	[/PARAM]
+	[PARAM:home]
+		[Home](Platforms/Linux/BeginnerLinuxDeveloper/SettingUpAnUnrealWorkflow)
+	[/PARAM]
+	[PARAM:next]
+	[/PARAM]
+[/OBJECT]
+[/VAR]
+
+%Steps%
+
+
+
+%Steps%
\ No newline at end of file
diff --git a/Engine/Documentation/Source/Platforms/Linux/BeginnerLinuxDeveloper/SettingUpAnUnrealWorkflow/SettingUpAnUnrealWorkflow.INT.udn b/Engine/Documentation/Source/Platforms/Linux/BeginnerLinuxDeveloper/SettingUpAnUnrealWorkflow/SettingUpAnUnrealWorkflow.INT.udn
--- a/Engine/Documentation/Source/Platforms/Linux/BeginnerLinuxDeveloper/SettingUpAnUnrealWorkflow/SettingUpAnUnrealWorkflow.INT.udn
+++ b/Engine/Documentation/Source/Platforms/Linux/BeginnerLinuxDeveloper/SettingUpAnUnrealWorkflow/SettingUpAnUnrealWorkflow.INT.udn
@@ -9,4 +9,38 @@ parent:Platforms/Linux/BeginnerLinuxDeveloper
 order: 2
 tags:Linux
 SkillLevel:Beginner
+prereq:Platforms/Linux/BeginnerLinuxDeveloper/SupportedGfxDrivers
 
+![](CompiledEditor_Linux.png)
+
+_At the end of this guide, you'll have Unreal Engine 4 (UE4) running on your Linux machine._
+
+## Goal
+The focus of this guide is to demonstrate the workflow needed to get UE4 running on your Linux machine.
+
+## Objectives
+
+After working through this guide, you'll know the following:
+
+* How to setup Git to download the source code for Unreal Engine.
+* How to fork and clone Unreal Engine's source code.
+* How to build Unreal Engine from source.
+* How to run Unreal Engine on your Linux machine.
+
+[OBJECT:TopicList]
+	[PARAM:icon]
+		![](%ROOT%/build_icon.png)(convert:false)
+	[/PARAM]
+	[PARAM:title]
+		Steps
+	[/PARAM]
+	[PARAM:description]
+	[/PARAM]
+	[PARAM:links]
+	[DIR(output:"checkpoint" parent:"Platforms/Linux/BeginnerLinuxDeveloper/SettingUpAnUnrealWorkflow" org:"hierarchy")]
+	[/PARAM]
+[/OBJECT]
+
+[REGION:call_to_action]
+	[Click to Start](Platforms/Linux/BeginnerLinuxDeveloper/SettingUpAnUnrealWorkflow/1 "%Platforms/Linux/BeginnerLinuxDeveloper/SettingUpAnUnrealWorkflow/1:title%")
+[/REGION]
\ No newline at end of file
