Added related page link for Required Hardware & Software Specifications.

#UE4docs

diff --git a/Engine/Documentation/Source/Platforms/Linux/BeginnerLinuxDeveloper/SupportedGfxDrivers/SupportedGfxDrivers.INT.udn b/Engine/Documentation/Source/Platforms/Linux/BeginnerLinuxDeveloper/SupportedGfxDrivers/SupportedGfxDrivers.INT.udn
--- a/Engine/Documentation/Source/Platforms/Linux/BeginnerLinuxDeveloper/SupportedGfxDrivers/SupportedGfxDrivers.INT.udn
+++ b/Engine/Documentation/Source/Platforms/Linux/BeginnerLinuxDeveloper/SupportedGfxDrivers/SupportedGfxDrivers.INT.udn
@@ -7,6 +7,7 @@ parent:Platforms/Linux/BeginnerLinuxDeveloper
 order: 1
 tags:Linux
 type:Reference
+related:GettingStarted/RecommendedSpecifications
 
 [TOC(start:1 end:4)]
 
