#UE4doc Refactored Account Setup excerpt, separating the steps from the section title and overview.

diff --git a/Engine/Documentation/Source/GettingStarted/AccountSetup/AccountSetup.INT.udn b/Engine/Documentation/Source/GettingStarted/AccountSetup/AccountSetup.INT.udn
--- a/Engine/Documentation/Source/GettingStarted/AccountSetup/AccountSetup.INT.udn
+++ b/Engine/Documentation/Source/GettingStarted/AccountSetup/AccountSetup.INT.udn
@@ -7,19 +7,19 @@ Related: Support/Builds
 Version: 4.8
 SkillLevel: Beginner
 
-[EXCERPT:epicAccount]
-
 ## Creating an Epic Games Account
 
 In order to run the Epic Games Launcher, you need to have a valid Epic Games account.
 
-1.	**Navigate** to [UnrealEngine.com](https://www.unrealengine.com/what-is-unreal-engine-4).
+[EXCERPT:epicAccount]
+
+1.	Navigate to [UnrealEngine.com](https://www.unrealengine.com/what-is-unreal-engine-4).
 
-1.	**Click** the **Get Unreal** button (pictured below).
+1.	Click the **Get Unreal** button (pictured below).
 
 	![](UnrealWebsite.png)
 
-1.	**Fill-out** our **'Join the Community'** form and **click** the **Sign Up** button (pictured below).
+1.	Fill-out our **'Join the Community'** form and click the **Sign Up** button (pictured below).
 
 	![](SignupScreen.png)
 
@@ -27,7 +27,7 @@ In order to run the Epic Games Launcher, you need to have a valid Epic Games acc
 	Make sure to read through our [terms of service](http://epicgames.com/tos) prior to signing up.
 	[/REGION]
 
-1.	**Read** our [Unreal Engine End User License Agreement](https://www.unrealengine.com/eula) (EULA) prior to **clicking** the **Accept** button (pictured below).
+1.	Read our [Unreal Engine End User License Agreement](https://www.unrealengine.com/eula) (EULA) prior to clicking the **Accept** button (pictured below).
 
 	![](EULA.png)
 
diff --git a/Engine/Documentation/Source/GettingStarted/Installation/InstallingUnrealEngine.INT.udn b/Engine/Documentation/Source/GettingStarted/Installation/InstallingUnrealEngine.INT.udn
--- a/Engine/Documentation/Source/GettingStarted/Installation/InstallingUnrealEngine.INT.udn
+++ b/Engine/Documentation/Source/GettingStarted/Installation/InstallingUnrealEngine.INT.udn
@@ -29,6 +29,10 @@ If you already have an Epic Games account, skip ahead to learn about [Downloadin
 
 (#Bookmark1)
 
+## Creating an Epic Games Account
+
+In order to run the Epic Games Launcher, you need to have a valid Epic Games account.
+
 [INCLUDE:GettingStarted/AccountSetup#epicAccount]
 
 ## Downloading and Running the Installer
