Wrote first draft of the 'Setting up Git for UE4' How To page.

#UE4docs

diff --git a/Engine/Documentation/Source/Platforms/Linux/BeginnerLinuxDeveloper/SettingUpAnUnrealWorkflow/1/SettingUpAnUnrealWorkflow_1.INT.udn b/Engine/Documentation/Source/Platforms/Linux/BeginnerLinuxDeveloper/SettingUpAnUnrealWorkflow/1/SettingUpAnUnrealWorkflow_1.INT.udn
--- a/Engine/Documentation/Source/Platforms/Linux/BeginnerLinuxDeveloper/SettingUpAnUnrealWorkflow/1/SettingUpAnUnrealWorkflow_1.INT.udn
+++ b/Engine/Documentation/Source/Platforms/Linux/BeginnerLinuxDeveloper/SettingUpAnUnrealWorkflow/1/SettingUpAnUnrealWorkflow_1.INT.udn
@@ -1,5 +1,5 @@
 Availability: Docs
-Title: 1 - Setting up Git
+Title: 1 - Setting up Git for UE4
 Crumbs:%ROOT%, Platforms, Platforms/Linux, Platforms/Linux/BeginnerLinuxDeveloper, Platforms/Linux/BeginnerLinuxDeveloper/SettingUpAnUnrealWorkflow
 Description:Learn how to setup Git on Linux.
 Version: 4.11
@@ -24,8 +24,38 @@ SkillLevel:Beginner
 [/OBJECT]
 [/VAR]
 
+Currently, we don't supply a binary installer of UE4 for Linux users. The good news is that you can compile a binary of Unreal Engine from the same source
+code that we use to develop the engine. If you want to access our source code, you'll have to register for an Epic Games account, sign up for a GitHub account, 
+associate your GitHub username with your Epic Games account, and setup Git on your Linux machine.
+
 %Steps%
 
+## Getting an Epic Games Account
+
+If you want to create games with Unreal Engine or contribute to the future of UE4 on Linux, you'll need to have a valid Epic Games account. 
+
+[INCLUDE:GettingStarted/AccountSetup#epicAccount]
+
+There are many benefits for having an account with us, including (but not limited to) access to the Unreal Engine community, where you'll 
+connect with other UE4 Linux developers, and where you'll be able to keep up with the latest news about Unreal Engine.
+
+## Setup Git
+
+Now that you've registered for an Epic Games account, you're ready to work through the following steps to gain access to the UE4 C++ source code on GitHub.
+
+1. Sign up for a GitHub account at [GitHub.com](https://github.com/EpicGames).
+
+	![](GitHubSignupPage.png)
+
+1. Go to [your account](https://accounts.unrealengine.com/login) and link your GitHub username by entering it into the account field labeled **GITHUB ACCOUNT NAME**.
+
+	![](AssociateGitHubAcct.png)
+
+1. Finally, set up [Git](https://help.github.com/articles/set-up-git/) before moving onto the next step.
+
+## End Result
 
+By now, you've registered for an Epic Games account and linked it with your GitHub username. Also, you've installed Git and you're ready for the next step,
+where you'll fork and clone the UE4 C++ source code from the Unreal Engine GitHub repository.
 
 %Steps%
\ No newline at end of file
