#UE4 Docs: Fixing a small typo in the page title.

diff --git a/Engine/Documentation/Source/Platforms/Android/VulkanMobileRenderer/VulkanMobileRender.INT.udn b/Engine/Documentation/Source/Platforms/Android/VulkanMobileRenderer/VulkanMobileRender.INT.udn
--- a/Engine/Documentation/Source/Platforms/Android/VulkanMobileRenderer/VulkanMobileRender.INT.udn
+++ b/Engine/Documentation/Source/Platforms/Android/VulkanMobileRenderer/VulkanMobileRender.INT.udn
@@ -1,4 +1,4 @@
-Title: Android Vulkan Mobile Render
+Title: Android Vulkan Mobile Renderer
 Description:Setting up your UE4 Android project to work with the Vulkan mobile render
 Crumbs:
 Availability:Public
