Updated based on SME review.

diff --git a/Engine/Documentation/Source/Engine/Animation/NodeReference/AnimationNodes.INT.udn b/Engine/Documentation/Source/Engine/Animation/NodeReference/AnimationNodes.INT.udn
--- a/Engine/Documentation/Source/Engine/Animation/NodeReference/AnimationNodes.INT.udn
+++ b/Engine/Documentation/Source/Engine/Animation/NodeReference/AnimationNodes.INT.udn
@@ -52,3 +52,10 @@ Data pins are actually properties like any other and are exposed in the **Detail
 between being used as a data pin on the node in the graph and being a simple property. When exposed as a data 
 pin, the property is only modifiable in the graph by linking a value to it. When not exposed as a data pin, 
 properties can only be modified through the **Details** panel.
+
+![](AsPin.png)
+
+Above a Blendspace Player node has properties in the **Details** panel that can be set as pins. 
+The X and Y coordinate pins which make up the **Speed** property of our Blendspace as well as Play Rate and Loop are exposed on the node since they are checked inside the **Details** panel. 
+
+
diff --git a/Engine/Documentation/Source/Engine/Animation/NodeReference/SkeletalControls/SkeletalControls.INT.udn b/Engine/Documentation/Source/Engine/Animation/NodeReference/SkeletalControls/SkeletalControls.INT.udn
--- a/Engine/Documentation/Source/Engine/Animation/NodeReference/SkeletalControls/SkeletalControls.INT.udn
+++ b/Engine/Documentation/Source/Engine/Animation/NodeReference/SkeletalControls/SkeletalControls.INT.udn
@@ -12,7 +12,7 @@ version: 4.12
 **SkeletalControls** (also called **SkelControls**) allow direct control of bones within a Skeleton. These can be 
 used within Animation Blueprints to control an individual bone, create IK chains, and more. This direct control of the 
 underlying Skeleton provides the ability to create procedural, dynamically-driven animation. The **Transform** of 
-one bone can be used to drive another or traces can be used to conform the feet of a character to the ground while 
+one bone can be used to drive another or can be used to conform the feet of a character to the ground while 
 playing a generic walk animation. Any sort of modification can be used to tweak or completely override the bone 
 Transforms applied by AnimationSequences.
 [/EXCERPT:Intro]
diff --git a/Engine/Documentation/Source/Engine/Animation/NodeReference/SpaceConversion/SpaceConversion.INT.udn b/Engine/Documentation/Source/Engine/Animation/NodeReference/SpaceConversion/SpaceConversion.INT.udn
--- a/Engine/Documentation/Source/Engine/Animation/NodeReference/SpaceConversion/SpaceConversion.INT.udn
+++ b/Engine/Documentation/Source/Engine/Animation/NodeReference/SpaceConversion/SpaceConversion.INT.udn
@@ -15,7 +15,7 @@ version: 4.9
 [EXCERPT:Intro]
 Poses can either be in local-space or component-space. The **Convert Spaces** nodes available in the **AnimGraph** 
 of Animation Blueprints provide the ability to convert poses between local and component space. Local space assumes 
-the transform of a bone to be relative to that bone. Component space assumes the bone's transform to be relative to 
+the transform of a bone to be relative to its parent bone. Component space assumes the bone's transform to be relative to 
 the SkeletalMeshComponent.
 
 Generally, when working with poses in an Animation Blueprint, they will be in local-space. However, certain blend 
@@ -24,17 +24,11 @@ prior to being passed in to one of these types of nodes. If the input pose is co
 space pose, the pose must be converted to the correct space before a SkeletalControl can perform operations on it. 
 The resulting pose, after the SkeletalControl have performed the operations, then must be converted back to 
 local space for input to additional blends or the Result pin.
-[/EXCERPT:Intro]
-
-[REGION:fullwidth]
-![Usage in graph](graph_usage.png)
-[/REGION]
 
 Because there is a cost associated with each conversion to or from component space, it is best to group any nodes that 
 operate in component-space, such as SkeletalControls, so that they are all performed in a row, limiting the number 
 of conversions needed.
-
-
+[/EXCERPT:Intro]
 
 ## Component to Local
 
