#UE4doc Scrubbed images to reflect GitHub's updated UI. Updated text to meet current Styles & Standards.

diff --git a/Engine/Documentation/Source/GettingStarted/DownloadingUnrealEngine/DownloadZip/DownloadZip.INT.udn b/Engine/Documentation/Source/GettingStarted/DownloadingUnrealEngine/DownloadZip/DownloadZip.INT.udn
--- a/Engine/Documentation/Source/GettingStarted/DownloadingUnrealEngine/DownloadZip/DownloadZip.INT.udn
+++ b/Engine/Documentation/Source/GettingStarted/DownloadingUnrealEngine/DownloadZip/DownloadZip.INT.udn
@@ -23,7 +23,9 @@ If you prefer to download our source code as a Zip file from GitHub, this page i
 	[/PARAMLITERAL]
 	[PARAM:content]
 	1.	Navigate to our GitHub page.
-	1.	**Click** the **Download ZIP** button.
+	1.	Click the **Clone or download** button.
+		![](CloneorDownloadButton.png)
+	1.	Click the **Download ZIP** button.
 		![](DownloadZipButton.png)
 		[REGION:warning]
 		The built-in Windows zip utility marks downloaded zip file contents as unsafe to execute.  
@@ -62,7 +64,9 @@ If you prefer to download our source code as a Zip file from GitHub, this page i
 	[/PARAMLITERAL]
 	[PARAM:content]
 	1.	Navigate to our GitHub page.
-	1.	**Click** the **Download ZIP** button.
+	1.	Click the **Clone or download** button.
+		![](CloneorDownloadButton.png)
+	1.	Click the **Download ZIP** button.
 		![](DownloadZipButton.png)
 	1.	Decompress 'UnrealEngine-(branch_name).zip' onto your hard disk.
 	1.	If you haven't already done so, install the latest version of [Xcode](https://itunes.apple.com/us/app/xcode/id497799835).
@@ -91,7 +95,9 @@ If you prefer to download our source code as a Zip file from GitHub, this page i
 	[/REGION]
 
 	1.	Navigate to our GitHub page.
-	1.	**Click** the **Download ZIP** button.
+	1.	Click the **Clone or download** button.
+		![](CloneorDownloadButton.png)
+	1.	Click the **Download ZIP** button.
 		![](DownloadZipButton.png)
 	1.	Decompress 'UnrealEngine-(branch_name).zip' onto your hard disk.
 	1.	Inside the root folder, run **Setup.sh** from the terminal to download the engine's binary content.
