Implemented first draft of 'Downloading UE4 on Linux' How To page.

#UE4docs

diff --git a/Engine/Documentation/Source/Platforms/Linux/BeginnerLinuxDeveloper/SettingUpAnUnrealWorkflow/2/SettingUpAnUnrealWorkflow_2.INT.udn b/Engine/Documentation/Source/Platforms/Linux/BeginnerLinuxDeveloper/SettingUpAnUnrealWorkflow/2/SettingUpAnUnrealWorkflow_2.INT.udn
--- a/Engine/Documentation/Source/Platforms/Linux/BeginnerLinuxDeveloper/SettingUpAnUnrealWorkflow/2/SettingUpAnUnrealWorkflow_2.INT.udn
+++ b/Engine/Documentation/Source/Platforms/Linux/BeginnerLinuxDeveloper/SettingUpAnUnrealWorkflow/2/SettingUpAnUnrealWorkflow_2.INT.udn
@@ -25,8 +25,71 @@ SkillLevel:Beginner
 [/OBJECT]
 [/VAR]
 
+Now that you've setup Git on your Linux machine, you're ready to download the Unreal Engine source code from our [GitHub site](http://github.com/EpicGames). There 
+are a couple of ways to download the UE4 source code, one way is to use the download feature on our GitHub page and the other way is to fork and clone the Unreal 
+Engine repository with your Git client. We'll cover both ways on this page.
+
 %Steps%
 
+## Downloading from GitHub
+
+If don't want to fork and clone the Unreal Engine repository using Git, you can use the **Download ZIP** button on our [GitHub site](http://github.com/EpicGames).
+
+1.	First, select the branch you want to download by clicking the **Branch:** button on the left side of the page.
+
+	![](SelectRepoBranch.png)
+
+	[INCLUDE:GettingStarted/DownloadingUnrealEngine/SourceBranches]
+
+1.	After selecting the branch you want to download, click the **Clone or download** button located on the right side of the page.
+
+	![](CloneorDownloadButton.png)
+
+1.	Click the **Download ZIP** button to download the Unreal Engine source code.
+
+	![](DownloadSourceButton.png)
+
+1.	Finally, decompress 'UnrealEngine-(branch_name).zip' onto your hard disk.
+
+## Forking and Cloning with Git
+
+If you're new to forking and cloning with Git, refer to GitHub's [Forking Projects Guide](https://guides.github.com/activities/forking/) before executing the following steps.
+
+1.	Click on the **Unreal Engine** link to access our [repository](https://github.com/EpicGames/UnrealEngine).
+	
+	![](UnrealRepo.png)
+	
+1.	Click the **Fork** button to make a fork of our repository.
+
+	![](Fork.png)
+
+1.	Navigate to the fork from your GitHub account.
+
+	![](NavigateToFork.png)
+
+1.	Now, click the **Clone or download** button.
+
+	![](ForkCloneOrDownloadButton.png)
+
+1.	Click the **Copy URL** button to copy the clone URL for the repository.
+
+	![](CopyHTTPS.png)
+
+1.	Open Git Bash, type `git clone`, and paste the URL you copied in the previous step, making sure to add your GitHub username rather than `YOUR-USERNAME`.
+
+	![](GitBashClone.png)
+
+1.	Finally, press **Enter** to create your local clone.
+
+	[REGION:tip]
+	If you need a refresher on using Git from the command line, refer to GitHub's articles on [Git](https://help.github.com/articles/set-up-git/) and [Forking a Repository](https://help.github.com/articles/fork-a-repo/).
+	[/REGION]
+
+## End Result
+
+![](end_result.png)
 
+Whether you downloaded the Unreal Engine source code by using the **Download ZIP** button on our [GitHub site](http://github.com/EpicGames), or by forking and cloning the Unreal Engine repository 
+using Git, you are now ready to build UE4 on Linux with the Unreal Engine source code that is now located on your hard disk.
 
 %Steps%
\ No newline at end of file
