#UE4doc Improved narrative flow of this page's overview.

diff --git a/Engine/Documentation/Source/GettingStarted/DownloadingUnrealEngine/SourceBranches/SourceBranches.INT.udn b/Engine/Documentation/Source/GettingStarted/DownloadingUnrealEngine/SourceBranches/SourceBranches.INT.udn
--- a/Engine/Documentation/Source/GettingStarted/DownloadingUnrealEngine/SourceBranches/SourceBranches.INT.udn
+++ b/Engine/Documentation/Source/GettingStarted/DownloadingUnrealEngine/SourceBranches/SourceBranches.INT.udn
@@ -11,7 +11,7 @@ Version: 4.8
 
 ## Source Code Branches
 
-Once you've accessed the repository, you'll notice that we've published UE4's source code in three rolling branches.
+You'll notice that we've published UE4's source code in three rolling branches.
 [REGION:raw]
 ![](Branches.png)
 [/REGION]
