This tutorial has been marked in Jira as being, 'Publish Ready'.
Flipped availalibity from Docs to Publish.
Added button to the Programming Tutorials homepage.

#UE4docs

diff --git a/Engine/Documentation/Source/Programming/Tutorials/FirstPersonShooter/1/1/FirstPersonShooter_1_1.INT.udn b/Engine/Documentation/Source/Programming/Tutorials/FirstPersonShooter/1/1/FirstPersonShooter_1_1.INT.udn
--- a/Engine/Documentation/Source/Programming/Tutorials/FirstPersonShooter/1/1/FirstPersonShooter_1_1.INT.udn
+++ b/Engine/Documentation/Source/Programming/Tutorials/FirstPersonShooter/1/1/FirstPersonShooter_1_1.INT.udn
@@ -1,4 +1,4 @@
-Availability: Docs
+Availability: Public
 Title: 1.1 - Setting up your Project
 Crumbs: %ROOT%, Programming, Programming/Tutorials, Programming/Tutorials/FirstPersonShooter, Programming/Tutorials/FirstPersonShooter/1
 Description: Learn how to set up your First Person Shooter project.
diff --git a/Engine/Documentation/Source/Programming/Tutorials/FirstPersonShooter/1/2/FirstPersonShooter_1_2.INT.udn b/Engine/Documentation/Source/Programming/Tutorials/FirstPersonShooter/1/2/FirstPersonShooter_1_2.INT.udn
--- a/Engine/Documentation/Source/Programming/Tutorials/FirstPersonShooter/1/2/FirstPersonShooter_1_2.INT.udn
+++ b/Engine/Documentation/Source/Programming/Tutorials/FirstPersonShooter/1/2/FirstPersonShooter_1_2.INT.udn
@@ -1,4 +1,4 @@
-Availability: Docs
+Availability: Public
 Title: 1.2 - Opening the Project in Visual Studio
 Crumbs: %ROOT%, Programming, Programming/Tutorials, Programming/Tutorials/FirstPersonShooter, Programming/Tutorials/FirstPersonShooter/1
 Description: Learn how to open your C++ project in Visual Studio.
diff --git a/Engine/Documentation/Source/Programming/Tutorials/FirstPersonShooter/1/3/FirstPersonShooter_1_3.INT.udn b/Engine/Documentation/Source/Programming/Tutorials/FirstPersonShooter/1/3/FirstPersonShooter_1_3.INT.udn
--- a/Engine/Documentation/Source/Programming/Tutorials/FirstPersonShooter/1/3/FirstPersonShooter_1_3.INT.udn
+++ b/Engine/Documentation/Source/Programming/Tutorials/FirstPersonShooter/1/3/FirstPersonShooter_1_3.INT.udn
@@ -1,4 +1,4 @@
-Availability: Docs
+Availability: Public
 Title: 1.3 - Adding Log Messaging
 Crumbs: %ROOT%, Programming, Programming/Tutorials, Programming/Tutorials/FirstPersonShooter, Programming/Tutorials/FirstPersonShooter/1
 Description: Learn how to add log messaging to your game.
diff --git a/Engine/Documentation/Source/Programming/Tutorials/FirstPersonShooter/1/4/FirstPersonShooter_1_4.INT.udn b/Engine/Documentation/Source/Programming/Tutorials/FirstPersonShooter/1/4/FirstPersonShooter_1_4.INT.udn
--- a/Engine/Documentation/Source/Programming/Tutorials/FirstPersonShooter/1/4/FirstPersonShooter_1_4.INT.udn
+++ b/Engine/Documentation/Source/Programming/Tutorials/FirstPersonShooter/1/4/FirstPersonShooter_1_4.INT.udn
@@ -1,4 +1,4 @@
-Availability: Docs
+Availability: Public
 Title: 1.4 - Compiling the Project
 Crumbs: %ROOT%, Programming, Programming/Tutorials, Programming/Tutorials/FirstPersonShooter, Programming/Tutorials/FirstPersonShooter/1
 Description: Learn how to compile your first C++ class.
diff --git a/Engine/Documentation/Source/Programming/Tutorials/FirstPersonShooter/1/5/FirstPersonShooter_1_5.INT.udn b/Engine/Documentation/Source/Programming/Tutorials/FirstPersonShooter/1/5/FirstPersonShooter_1_5.INT.udn
--- a/Engine/Documentation/Source/Programming/Tutorials/FirstPersonShooter/1/5/FirstPersonShooter_1_5.INT.udn
+++ b/Engine/Documentation/Source/Programming/Tutorials/FirstPersonShooter/1/5/FirstPersonShooter_1_5.INT.udn
@@ -1,4 +1,4 @@
-Availability: Docs
+Availability: Public
 Title: 1.5 - Setting the Default Game Mode
 Crumbs: %ROOT%, Programming, Programming/Tutorials, Programming/Tutorials/FirstPersonShooter, Programming/Tutorials/FirstPersonShooter/1
 Description: Learn how to set the default game mode for your First Person Shooter project.
diff --git a/Engine/Documentation/Source/Programming/Tutorials/FirstPersonShooter/1/FirstPersonShooter_1.INT.udn b/Engine/Documentation/Source/Programming/Tutorials/FirstPersonShooter/1/FirstPersonShooter_1.INT.udn
--- a/Engine/Documentation/Source/Programming/Tutorials/FirstPersonShooter/1/FirstPersonShooter_1.INT.udn
+++ b/Engine/Documentation/Source/Programming/Tutorials/FirstPersonShooter/1/FirstPersonShooter_1.INT.udn
@@ -1,4 +1,4 @@
-Availability: Docs
+Availability: Public
 Title: 1 - Setting up your Project
 Crumbs: %ROOT%, Programming, Programming/Tutorials, Programming/Tutorials/FirstPersonShooter
 Description: How to set up your First Person Shooter project.
diff --git a/Engine/Documentation/Source/Programming/Tutorials/FirstPersonShooter/2/1/FirstPersonShooter_2_1.INT.udn b/Engine/Documentation/Source/Programming/Tutorials/FirstPersonShooter/2/1/FirstPersonShooter_2_1.INT.udn
--- a/Engine/Documentation/Source/Programming/Tutorials/FirstPersonShooter/2/1/FirstPersonShooter_2_1.INT.udn
+++ b/Engine/Documentation/Source/Programming/Tutorials/FirstPersonShooter/2/1/FirstPersonShooter_2_1.INT.udn
@@ -1,4 +1,4 @@
-Availability: Docs
+Availability: Public
 Title: 2.1 - Making a New Character
 Crumbs: %ROOT%, Programming, Programming/Tutorials/FirstPersonShooter, Programming/Tutorials/FirstPersonShooter/2
 Description: Learn how to make a new character for your First Person Shooter project.
diff --git a/Engine/Documentation/Source/Programming/Tutorials/FirstPersonShooter/2/2/FirstPersonShooter_2_2.INT.udn b/Engine/Documentation/Source/Programming/Tutorials/FirstPersonShooter/2/2/FirstPersonShooter_2_2.INT.udn
--- a/Engine/Documentation/Source/Programming/Tutorials/FirstPersonShooter/2/2/FirstPersonShooter_2_2.INT.udn
+++ b/Engine/Documentation/Source/Programming/Tutorials/FirstPersonShooter/2/2/FirstPersonShooter_2_2.INT.udn
@@ -1,4 +1,4 @@
-Availability: Docs
+Availability: Public
 Title: 2.2 - Setting up Axis Mapping
 Crumbs: %ROOT%, Programming, Programming/Tutorials/FirstPersonShooter, Programming/Tutorials/FirstPersonShooter/2
 Description: Learn how to set-up axis mapping for your First Person Shooter character.
diff --git a/Engine/Documentation/Source/Programming/Tutorials/FirstPersonShooter/2/3/FirstPersonShooter_2_3.INT.udn b/Engine/Documentation/Source/Programming/Tutorials/FirstPersonShooter/2/3/FirstPersonShooter_2_3.INT.udn
--- a/Engine/Documentation/Source/Programming/Tutorials/FirstPersonShooter/2/3/FirstPersonShooter_2_3.INT.udn
+++ b/Engine/Documentation/Source/Programming/Tutorials/FirstPersonShooter/2/3/FirstPersonShooter_2_3.INT.udn
@@ -1,4 +1,4 @@
-Availability: Docs
+Availability: Public
 Title: 2.3 - Implementing Character Movement Functions
 Crumbs: %ROOT%, Programming, Programming/Tutorials/FirstPersonShooter, Programming/Tutorials/FirstPersonShooter/2
 Description: Learn how to implement movement functions for your First Person Shooter character.
diff --git a/Engine/Documentation/Source/Programming/Tutorials/FirstPersonShooter/2/4/FirstPersonShooter_2_4.INT.udn b/Engine/Documentation/Source/Programming/Tutorials/FirstPersonShooter/2/4/FirstPersonShooter_2_4.INT.udn
--- a/Engine/Documentation/Source/Programming/Tutorials/FirstPersonShooter/2/4/FirstPersonShooter_2_4.INT.udn
+++ b/Engine/Documentation/Source/Programming/Tutorials/FirstPersonShooter/2/4/FirstPersonShooter_2_4.INT.udn
@@ -1,4 +1,4 @@
-Availability: Docs
+Availability: Public
 Title: 2.4 - Implementing Mouse Camera Control
 Crumbs: %ROOT%, Programming, Programming/Tutorials/FirstPersonShooter, Programming/Tutorials/FirstPersonShooter/2
 Description: Learn how to implement mouse camera control for your First Person Shooter character.
diff --git a/Engine/Documentation/Source/Programming/Tutorials/FirstPersonShooter/2/5/FirstPersonShooter_2_5.INT.udn b/Engine/Documentation/Source/Programming/Tutorials/FirstPersonShooter/2/5/FirstPersonShooter_2_5.INT.udn
--- a/Engine/Documentation/Source/Programming/Tutorials/FirstPersonShooter/2/5/FirstPersonShooter_2_5.INT.udn
+++ b/Engine/Documentation/Source/Programming/Tutorials/FirstPersonShooter/2/5/FirstPersonShooter_2_5.INT.udn
@@ -1,4 +1,4 @@
-Availability: Docs
+Availability: Public
 Title: 2.5 - Implementing Character Jumping
 Crumbs: %ROOT%, Programming, Programming/Tutorials/FirstPersonShooter, Programming/Tutorials/FirstPersonShooter/2
 Description: Learn how to implement jumping for your First Person Shooter character.
diff --git a/Engine/Documentation/Source/Programming/Tutorials/FirstPersonShooter/2/6/FirstPersonShooter_2_6.INT.udn b/Engine/Documentation/Source/Programming/Tutorials/FirstPersonShooter/2/6/FirstPersonShooter_2_6.INT.udn
--- a/Engine/Documentation/Source/Programming/Tutorials/FirstPersonShooter/2/6/FirstPersonShooter_2_6.INT.udn
+++ b/Engine/Documentation/Source/Programming/Tutorials/FirstPersonShooter/2/6/FirstPersonShooter_2_6.INT.udn
@@ -1,4 +1,4 @@
-Availability: Docs
+Availability: Public
 Title: 2.6 - Adding a Mesh to your Character
 Crumbs: %ROOT%, Programming, Programming/Tutorials/FirstPersonShooter, Programming/Tutorials/FirstPersonShooter/2
 Description: Learn how to add a mesh to your First Person Shooter character.
diff --git a/Engine/Documentation/Source/Programming/Tutorials/FirstPersonShooter/2/7/FirstPersonShooter_2_7.INT.udn b/Engine/Documentation/Source/Programming/Tutorials/FirstPersonShooter/2/7/FirstPersonShooter_2_7.INT.udn
--- a/Engine/Documentation/Source/Programming/Tutorials/FirstPersonShooter/2/7/FirstPersonShooter_2_7.INT.udn
+++ b/Engine/Documentation/Source/Programming/Tutorials/FirstPersonShooter/2/7/FirstPersonShooter_2_7.INT.udn
@@ -1,4 +1,4 @@
-Availability: Docs
+Availability: Public
 Title: 2.7 - Changing the Camera View
 Crumbs: %ROOT%, Programming, Programming/Tutorials/FirstPersonShooter, Programming/Tutorials/FirstPersonShooter/2
 Description: Learn how to change the camera view for your First Person Shooter character.
diff --git a/Engine/Documentation/Source/Programming/Tutorials/FirstPersonShooter/2/8/FirstPersonShooter_2_8.INT.udn b/Engine/Documentation/Source/Programming/Tutorials/FirstPersonShooter/2/8/FirstPersonShooter_2_8.INT.udn
--- a/Engine/Documentation/Source/Programming/Tutorials/FirstPersonShooter/2/8/FirstPersonShooter_2_8.INT.udn
+++ b/Engine/Documentation/Source/Programming/Tutorials/FirstPersonShooter/2/8/FirstPersonShooter_2_8.INT.udn
@@ -1,4 +1,4 @@
-Availability: Docs
+Availability: Public
 Title: 2.8 - Add a First Person Mesh to your Character
 Crumbs: %ROOT%, Programming, Programming/Tutorials/FirstPersonShooter, Programming/Tutorials/FirstPersonShooter/2
 Description: Learn how to add a first person mesh to your First Person Shooter character.
diff --git a/Engine/Documentation/Source/Programming/Tutorials/FirstPersonShooter/2/FirstPersonShooter_2.INT.udn b/Engine/Documentation/Source/Programming/Tutorials/FirstPersonShooter/2/FirstPersonShooter_2.INT.udn
--- a/Engine/Documentation/Source/Programming/Tutorials/FirstPersonShooter/2/FirstPersonShooter_2.INT.udn
+++ b/Engine/Documentation/Source/Programming/Tutorials/FirstPersonShooter/2/FirstPersonShooter_2.INT.udn
@@ -1,4 +1,4 @@
-Availability: Docs
+Availability: Public
 Title: 2 - Implementing your Character
 Crumbs: %ROOT%, Programming, Programming/Tutorials, Programming/Tutorials/FirstPersonShooter
 Description: How to implement your First Person Shooter character.
diff --git a/Engine/Documentation/Source/Programming/Tutorials/FirstPersonShooter/3/1/FirstPersonShooter_3_1.INT.udn b/Engine/Documentation/Source/Programming/Tutorials/FirstPersonShooter/3/1/FirstPersonShooter_3_1.INT.udn
--- a/Engine/Documentation/Source/Programming/Tutorials/FirstPersonShooter/3/1/FirstPersonShooter_3_1.INT.udn
+++ b/Engine/Documentation/Source/Programming/Tutorials/FirstPersonShooter/3/1/FirstPersonShooter_3_1.INT.udn
@@ -1,4 +1,4 @@
-Availability: Docs
+Availability: Public
 Title: 3.1 - Adding Projectiles to your Game
 Crumbs: %ROOT%, Programming, Programming/Tutorials/FirstPersonShooter, Programming/Tutorials/FirstPersonShooter/3
 Description: Learn how to add projectiles to your First Person Shooter project.
diff --git a/Engine/Documentation/Source/Programming/Tutorials/FirstPersonShooter/3/2/FirstPersonShooter_3_2.INT.udn b/Engine/Documentation/Source/Programming/Tutorials/FirstPersonShooter/3/2/FirstPersonShooter_3_2.INT.udn
--- a/Engine/Documentation/Source/Programming/Tutorials/FirstPersonShooter/3/2/FirstPersonShooter_3_2.INT.udn
+++ b/Engine/Documentation/Source/Programming/Tutorials/FirstPersonShooter/3/2/FirstPersonShooter_3_2.INT.udn
@@ -1,4 +1,4 @@
-Availability: Docs
+Availability: Public
 Title: 3.2 - Implementing Shooting
 Crumbs: %ROOT%, Programming, Programming/Tutorials/FirstPersonShooter, Programming/Tutorials/FirstPersonShooter/3
 Description: Learn how to shooting for your First Person Shooter character.
diff --git a/Engine/Documentation/Source/Programming/Tutorials/FirstPersonShooter/3/3/FirstPersonShooter_3_3.INT.udn b/Engine/Documentation/Source/Programming/Tutorials/FirstPersonShooter/3/3/FirstPersonShooter_3_3.INT.udn
--- a/Engine/Documentation/Source/Programming/Tutorials/FirstPersonShooter/3/3/FirstPersonShooter_3_3.INT.udn
+++ b/Engine/Documentation/Source/Programming/Tutorials/FirstPersonShooter/3/3/FirstPersonShooter_3_3.INT.udn
@@ -1,4 +1,4 @@
-Availability: Docs
+Availability: Public
 Title: 3.3 - Setting up Projectile Collision and Lifetime
 Crumbs: %ROOT%, Programming, Programming/Tutorials/FirstPersonShooter, Programming/Tutorials/FirstPersonShooter/3
 Description: Learn how to set up collision and lifetime for your First Person Shooter projectiles.
diff --git a/Engine/Documentation/Source/Programming/Tutorials/FirstPersonShooter/3/4/FirstPersonShooter_3_4.INT.udn b/Engine/Documentation/Source/Programming/Tutorials/FirstPersonShooter/3/4/FirstPersonShooter_3_4.INT.udn
--- a/Engine/Documentation/Source/Programming/Tutorials/FirstPersonShooter/3/4/FirstPersonShooter_3_4.INT.udn
+++ b/Engine/Documentation/Source/Programming/Tutorials/FirstPersonShooter/3/4/FirstPersonShooter_3_4.INT.udn
@@ -1,4 +1,4 @@
-Availability: Docs
+Availability: Public
 Title: 3.4 - Getting Projectiles to Interact with the World
 Crumbs: %ROOT%, Programming, Programming/Tutorials/FirstPersonShooter, Programming/Tutorials/FirstPersonShooter/3
 Description: Learn how to get projectiles to interact with the world for your First Person Shooter project.
diff --git a/Engine/Documentation/Source/Programming/Tutorials/FirstPersonShooter/3/5/FirstPersonShooter_3_5.INT.udn b/Engine/Documentation/Source/Programming/Tutorials/FirstPersonShooter/3/5/FirstPersonShooter_3_5.INT.udn
--- a/Engine/Documentation/Source/Programming/Tutorials/FirstPersonShooter/3/5/FirstPersonShooter_3_5.INT.udn
+++ b/Engine/Documentation/Source/Programming/Tutorials/FirstPersonShooter/3/5/FirstPersonShooter_3_5.INT.udn
@@ -1,4 +1,4 @@
-Availability: Docs
+Availability: Public
 Title: 3.5 - Adding Crosshairs to your Viewport
 Crumbs: %ROOT%, Programming, Programming/Tutorials/FirstPersonShooter, Programming/Tutorials/FirstPersonShooter/3
 Description: Learn how to add crosshairs to the viewport for your First Person Shooter project.
diff --git a/Engine/Documentation/Source/Programming/Tutorials/FirstPersonShooter/3/FirstPersonShooter_3.INT.udn b/Engine/Documentation/Source/Programming/Tutorials/FirstPersonShooter/3/FirstPersonShooter_3.INT.udn
--- a/Engine/Documentation/Source/Programming/Tutorials/FirstPersonShooter/3/FirstPersonShooter_3.INT.udn
+++ b/Engine/Documentation/Source/Programming/Tutorials/FirstPersonShooter/3/FirstPersonShooter_3.INT.udn
@@ -1,4 +1,4 @@
-Availability: Docs
+Availability: Public
 Title: 3 - Implementing Projectiles
 Crumbs: %ROOT%, Programming, Programming/Tutorials, Programming/Tutorials/FirstPersonShooter
 Description: How to implement projectiles for your First Person Shooter game.
diff --git a/Engine/Documentation/Source/Programming/Tutorials/FirstPersonShooter/4/1/FirstPersonShooter_4_1.INT.udn b/Engine/Documentation/Source/Programming/Tutorials/FirstPersonShooter/4/1/FirstPersonShooter_4_1.INT.udn
--- a/Engine/Documentation/Source/Programming/Tutorials/FirstPersonShooter/4/1/FirstPersonShooter_4_1.INT.udn
+++ b/Engine/Documentation/Source/Programming/Tutorials/FirstPersonShooter/4/1/FirstPersonShooter_4_1.INT.udn
@@ -1,4 +1,4 @@
-Availability: Docs
+Availability: Public
 Title: 4.1 - Animating your Character
 Crumbs: %ROOT%, Programming, Programming/Tutorials, Programming/Tutorials/FirstPersonShooter, Programming/Tutorials/FirstPersonShooter/4
 Description: Learn how to animate your First Person Shooter character.
diff --git a/Engine/Documentation/Source/Programming/Tutorials/FirstPersonShooter/4/2/FirstPersonShooter_4_2.INT.udn b/Engine/Documentation/Source/Programming/Tutorials/FirstPersonShooter/4/2/FirstPersonShooter_4_2.INT.udn
--- a/Engine/Documentation/Source/Programming/Tutorials/FirstPersonShooter/4/2/FirstPersonShooter_4_2.INT.udn
+++ b/Engine/Documentation/Source/Programming/Tutorials/FirstPersonShooter/4/2/FirstPersonShooter_4_2.INT.udn
@@ -1,4 +1,4 @@
-Availability: Docs
+Availability: Public
 Title: 4.2 - Setting up your Event Graph
 Crumbs: %ROOT%, Programming, Programming/Tutorials, Programming/Tutorials/FirstPersonShooter, Programming/Tutorials/FirstPersonShooter/4
 Description: Learn how to set-up an Event Graph for your First Person Shooter character.
diff --git a/Engine/Documentation/Source/Programming/Tutorials/FirstPersonShooter/4/3/FirstPersonShooter_4_3.INT.udn b/Engine/Documentation/Source/Programming/Tutorials/FirstPersonShooter/4/3/FirstPersonShooter_4_3.INT.udn
--- a/Engine/Documentation/Source/Programming/Tutorials/FirstPersonShooter/4/3/FirstPersonShooter_4_3.INT.udn
+++ b/Engine/Documentation/Source/Programming/Tutorials/FirstPersonShooter/4/3/FirstPersonShooter_4_3.INT.udn
@@ -1,4 +1,4 @@
-Availability: Docs
+Availability: Public
 Title: 4.3 - Adding an Animation State Machine
 Crumbs: %ROOT%, Programming, Programming/Tutorials, Programming/Tutorials/FirstPersonShooter, Programming/Tutorials/FirstPersonShooter/4
 Description: Learn how to add an animation state machine for your First Person character.
diff --git a/Engine/Documentation/Source/Programming/Tutorials/FirstPersonShooter/4/4/1/FirstPersonShooter_4_4_1.INT.udn b/Engine/Documentation/Source/Programming/Tutorials/FirstPersonShooter/4/4/1/FirstPersonShooter_4_4_1.INT.udn
--- a/Engine/Documentation/Source/Programming/Tutorials/FirstPersonShooter/4/4/1/FirstPersonShooter_4_4_1.INT.udn
+++ b/Engine/Documentation/Source/Programming/Tutorials/FirstPersonShooter/4/4/1/FirstPersonShooter_4_4_1.INT.udn
@@ -1,4 +1,4 @@
-Availability: Docs
+Availability: Public
 Title: 4.4.1 - Add Idle to/from Run Transitions
 Crumbs: %ROOT%, Programming, Programming/Tutorials, Programming/Tutorials/FirstPersonShooter, Programming/Tutorials/FirstPersonShooter/4, Programming/Tutorials/FirstPersonShooter/4/4
 Description: Learn how to add idle to/from run state transitions your First Person Shooter character's animation state machine.
diff --git a/Engine/Documentation/Source/Programming/Tutorials/FirstPersonShooter/4/4/2/FirstPersonShooter_4_4_2.INT.udn b/Engine/Documentation/Source/Programming/Tutorials/FirstPersonShooter/4/4/2/FirstPersonShooter_4_4_2.INT.udn
--- a/Engine/Documentation/Source/Programming/Tutorials/FirstPersonShooter/4/4/2/FirstPersonShooter_4_4_2.INT.udn
+++ b/Engine/Documentation/Source/Programming/Tutorials/FirstPersonShooter/4/4/2/FirstPersonShooter_4_4_2.INT.udn
@@ -1,4 +1,4 @@
-Availability: Docs
+Availability: Public
 Title: 4.4.2 - Add Idle to Jump Start Transition
 Crumbs: %ROOT%, Programming, Programming/Tutorials, Programming/Tutorials/FirstPersonShooter, Programming/Tutorials/FirstPersonShooter/4, Programming/Tutorials/FirstPersonShooter/4/4
 Description: Learn how to add idle to jump start state transitions your First Person Shooter character's animation state machine.
diff --git a/Engine/Documentation/Source/Programming/Tutorials/FirstPersonShooter/4/4/3/FirstPersonShooter_4_4_3.INT.udn b/Engine/Documentation/Source/Programming/Tutorials/FirstPersonShooter/4/4/3/FirstPersonShooter_4_4_3.INT.udn
--- a/Engine/Documentation/Source/Programming/Tutorials/FirstPersonShooter/4/4/3/FirstPersonShooter_4_4_3.INT.udn
+++ b/Engine/Documentation/Source/Programming/Tutorials/FirstPersonShooter/4/4/3/FirstPersonShooter_4_4_3.INT.udn
@@ -1,4 +1,4 @@
-Availability: Docs
+Availability: Public
 Title: 4.4.3 - Add Run to Jump Start Transition
 Crumbs: %ROOT%, Programming, Programming/Tutorials, Programming/Tutorials/FirstPersonShooter, Programming/Tutorials/FirstPersonShooter/4, Programming/Tutorials/FirstPersonShooter/4/4
 Description: Learn how to add run to jump start state transitions your First Person Shooter character's animation state machine.
diff --git a/Engine/Documentation/Source/Programming/Tutorials/FirstPersonShooter/4/4/4/FirstPersonShooter_4_4_4.INT.udn b/Engine/Documentation/Source/Programming/Tutorials/FirstPersonShooter/4/4/4/FirstPersonShooter_4_4_4.INT.udn
--- a/Engine/Documentation/Source/Programming/Tutorials/FirstPersonShooter/4/4/4/FirstPersonShooter_4_4_4.INT.udn
+++ b/Engine/Documentation/Source/Programming/Tutorials/FirstPersonShooter/4/4/4/FirstPersonShooter_4_4_4.INT.udn
@@ -1,4 +1,4 @@
-Availability: Docs
+Availability: Public
 Title: 4.4.4 - Add Jump Start to Jump Loop Transition
 Crumbs: %ROOT%, Programming, Programming/Tutorials, Programming/Tutorials/FirstPersonShooter, Programming/Tutorials/FirstPersonShooter/4, Programming/Tutorials/FirstPersonShooter/4/4
 Description: Learn how to add jump start to jump loop state transitions your First Person Shooter character's animation state machine.
diff --git a/Engine/Documentation/Source/Programming/Tutorials/FirstPersonShooter/4/4/5/FirstPersonShooter_4_4_5.INT.udn b/Engine/Documentation/Source/Programming/Tutorials/FirstPersonShooter/4/4/5/FirstPersonShooter_4_4_5.INT.udn
--- a/Engine/Documentation/Source/Programming/Tutorials/FirstPersonShooter/4/4/5/FirstPersonShooter_4_4_5.INT.udn
+++ b/Engine/Documentation/Source/Programming/Tutorials/FirstPersonShooter/4/4/5/FirstPersonShooter_4_4_5.INT.udn
@@ -1,4 +1,4 @@
-Availability: Docs
+Availability: Public
 Title: 4.4.5 - Add Jump Loop to Jump End Transition
 Crumbs: %ROOT%, Programming, Programming/Tutorials, Programming/Tutorials/FirstPersonShooter, Programming/Tutorials/FirstPersonShooter/4, Programming/Tutorials/FirstPersonShooter/4/4
 Description: Learn how to add jump loop to jump end state transitions your First Person Shooter character's animation state machine.
diff --git a/Engine/Documentation/Source/Programming/Tutorials/FirstPersonShooter/4/4/6/FirstPersonShooter_4_4_6.INT.udn b/Engine/Documentation/Source/Programming/Tutorials/FirstPersonShooter/4/4/6/FirstPersonShooter_4_4_6.INT.udn
--- a/Engine/Documentation/Source/Programming/Tutorials/FirstPersonShooter/4/4/6/FirstPersonShooter_4_4_6.INT.udn
+++ b/Engine/Documentation/Source/Programming/Tutorials/FirstPersonShooter/4/4/6/FirstPersonShooter_4_4_6.INT.udn
@@ -1,4 +1,4 @@
-Availability: Docs
+Availability: Public
 Title: 4.4.6 - Add Jump End to Idle Transition
 Crumbs: %ROOT%, Programming, Programming/Tutorials, Programming/Tutorials/FirstPersonShooter, Programming/Tutorials/FirstPersonShooter/4, Programming/Tutorials/FirstPersonShooter/4/4
 Description: Learn how to add jump end to idle state transitions your First Person Shooter character's animation state machine.
diff --git a/Engine/Documentation/Source/Programming/Tutorials/FirstPersonShooter/4/4/FirstPersonShooter_4_4.INT.udn b/Engine/Documentation/Source/Programming/Tutorials/FirstPersonShooter/4/4/FirstPersonShooter_4_4.INT.udn
--- a/Engine/Documentation/Source/Programming/Tutorials/FirstPersonShooter/4/4/FirstPersonShooter_4_4.INT.udn
+++ b/Engine/Documentation/Source/Programming/Tutorials/FirstPersonShooter/4/4/FirstPersonShooter_4_4.INT.udn
@@ -1,4 +1,4 @@
-Availability: Docs
+Availability: Public
 Title: 4.4 - Adding Animation Transition States
 Crumbs: %ROOT%, Programming, Programming/Tutorials, Programming/Tutorials/FirstPersonShooter, Programming/Tutorials/FirstPersonShooter/4
 Description: How to add animation transition states to your First Person Shooter project.
diff --git a/Engine/Documentation/Source/Programming/Tutorials/FirstPersonShooter/4/5/FirstPersonShooter_4_5.INT.udn b/Engine/Documentation/Source/Programming/Tutorials/FirstPersonShooter/4/5/FirstPersonShooter_4_5.INT.udn
--- a/Engine/Documentation/Source/Programming/Tutorials/FirstPersonShooter/4/5/FirstPersonShooter_4_5.INT.udn
+++ b/Engine/Documentation/Source/Programming/Tutorials/FirstPersonShooter/4/5/FirstPersonShooter_4_5.INT.udn
@@ -1,4 +1,4 @@
-Availability: Docs
+Availability: Public
 Title: 4.5 - Associating Animation & Character Blueprints
 Crumbs: %ROOT%, Programming, Programming/Tutorials, Programming/Tutorials/FirstPersonShooter, Programming/Tutorials/FirstPersonShooter/4
 Description: Learn how to associate the Animation Blueprint with the Character Blueprint for your First Person Shooter project.
diff --git a/Engine/Documentation/Source/Programming/Tutorials/FirstPersonShooter/4/FirstPersonShooter_4.INT.udn b/Engine/Documentation/Source/Programming/Tutorials/FirstPersonShooter/4/FirstPersonShooter_4.INT.udn
--- a/Engine/Documentation/Source/Programming/Tutorials/FirstPersonShooter/4/FirstPersonShooter_4.INT.udn
+++ b/Engine/Documentation/Source/Programming/Tutorials/FirstPersonShooter/4/FirstPersonShooter_4.INT.udn
@@ -1,4 +1,4 @@
-Availability: Docs
+Availability: Public
 Title: 4 - Adding Character Animation
 Crumbs: %ROOT%, Programming, Programming/Tutorials, Programming/Tutorials/FirstPersonShooter
 Description: How to add animations to your First Person Shooter character.
diff --git a/Engine/Documentation/Source/Programming/Tutorials/FirstPersonShooter/FirstPersonShooter.INT.udn b/Engine/Documentation/Source/Programming/Tutorials/FirstPersonShooter/FirstPersonShooter.INT.udn
--- a/Engine/Documentation/Source/Programming/Tutorials/FirstPersonShooter/FirstPersonShooter.INT.udn
+++ b/Engine/Documentation/Source/Programming/Tutorials/FirstPersonShooter/FirstPersonShooter.INT.udn
@@ -1,5 +1,5 @@
-Availability: Docs
-Title: C++ First Person Shooter Tutorial
+Availability: Public
+Title: First Person Shooter Tutorial
 Crumbs: %ROOT%, GettingStarted, Programming/Tutorials
 Description: Learn how to implement mechanics for a First Person Shooter game.
 prereq:Programming/Development/VisualStudioSetup
diff --git a/Engine/Documentation/Source/Programming/Tutorials/ProgrammingTutorialsHome.INT.udn b/Engine/Documentation/Source/Programming/Tutorials/ProgrammingTutorialsHome.INT.udn
--- a/Engine/Documentation/Source/Programming/Tutorials/ProgrammingTutorialsHome.INT.udn
+++ b/Engine/Documentation/Source/Programming/Tutorials/ProgrammingTutorialsHome.INT.udn
@@ -145,7 +145,6 @@ programming
 	[/PARAMLITERAL]
 [/OBJECT]
 
-<!--TO DO: Uncomment the following block when ready for publication.
 [OBJECT:Role]
 	[PARAM:icon]
 		![](%ROOT%/start_icon.png)(convert:false)
@@ -163,7 +162,7 @@ programming
 		cppfps
 	[/PARAMLITERAL]
 [/OBJECT]
--->
+
 [/REGION]
 
 ## Intermediate
