Added Hardware and Software Specifications related page.
This way, users will have a reference point if they have trouble running UE4 on their Linux machine.

#UE4doc

diff --git a/Engine/Documentation/Source/Platforms/Linux/Linux.INT.udn b/Engine/Documentation/Source/Platforms/Linux/Linux.INT.udn
--- a/Engine/Documentation/Source/Platforms/Linux/Linux.INT.udn
+++ b/Engine/Documentation/Source/Platforms/Linux/Linux.INT.udn
@@ -25,3 +25,5 @@ onto more advanced topics.
 %Platforms/Linux/AdvancedLinuxDeveloper:topic%
 -->
 [/REGION]
+
+If you're having trouble running UE4 on your machine, make sure you meet our [](GettingStarted/RecommendedSpecifications).
\ No newline at end of file
