Wrote Overview text for the 'Building UE4 on Linux' How To page.

#UE4docs

diff --git a/Engine/Documentation/Source/Platforms/Linux/BeginnerLinuxDeveloper/SettingUpAnUnrealWorkflow/3/SettingUpAnUnrealWorkflow_3.INT.udn b/Engine/Documentation/Source/Platforms/Linux/BeginnerLinuxDeveloper/SettingUpAnUnrealWorkflow/3/SettingUpAnUnrealWorkflow_3.INT.udn
--- a/Engine/Documentation/Source/Platforms/Linux/BeginnerLinuxDeveloper/SettingUpAnUnrealWorkflow/3/SettingUpAnUnrealWorkflow_3.INT.udn
+++ b/Engine/Documentation/Source/Platforms/Linux/BeginnerLinuxDeveloper/SettingUpAnUnrealWorkflow/3/SettingUpAnUnrealWorkflow_3.INT.udn
@@ -25,6 +25,9 @@ SkillLevel:Beginner
 [/OBJECT]
 [/VAR]
 
+In preparation for this critical step, you downloaded a copy of Unreal Engine's source code on your hard disk. 
+During this step, you'll get to run several utilities from the terminal to build a binary of UE4 on your Linux machine.
+
 %Steps%
 
 
