#UEDOC-1419

Separated steps for running the editor from steps needed to build the editor from source.

#UE4docs

diff --git a/Engine/Documentation/Source/Programming/Development/BuildingUnrealEngine/BuildingUnrealEngine.INT.udn b/Engine/Documentation/Source/Programming/Development/BuildingUnrealEngine/BuildingUnrealEngine.INT.udn
--- a/Engine/Documentation/Source/Programming/Development/BuildingUnrealEngine/BuildingUnrealEngine.INT.udn
+++ b/Engine/Documentation/Source/Programming/Development/BuildingUnrealEngine/BuildingUnrealEngine.INT.udn
@@ -59,11 +59,53 @@ SkillLevel: Advanced
 	[/PARAM]
 [/OBJECT]
 
+## Running the Editor
+
+[OBJECT:ToggleButtonContent]
+	[PARAMLITERAL:category]
+	OS
+	[/PARAMLITERAL]
+	[PARAMLITERAL:id]
+	windows
+	[/PARAMLITERAL]
+	[PARAMLITERAL:active]
+	active_button_content
+	[/PARAMLITERAL]
+	[PARAM:content]
+		[INCLUDE:#windowsrun]
+	[/PARAM]
+[/OBJECT]
+[OBJECT:ToggleButtonContent]
+	[PARAMLITERAL:category]
+	OS
+	[/PARAMLITERAL]
+	[PARAMLITERAL:id]
+	mac
+	[/PARAMLITERAL]
+	[PARAMLITERAL:active]
+	[/PARAMLITERAL]
+	[PARAM:content]
+		[INCLUDE:#macrun]
+	[/PARAM]
+[/OBJECT]
+[OBJECT:ToggleButtonContent]
+	[PARAMLITERAL:category]
+	OS
+	[/PARAMLITERAL]
+	[PARAMLITERAL:id]
+	linux
+	[/PARAMLITERAL]
+	[PARAMLITERAL:active]
+	[/PARAMLITERAL]
+	[PARAM:content]
+		[INCLUDE:#linuxrun]
+	[/PARAM]
+[/OBJECT]
+
 <!--
 [EXCERPT:windowscompile]
-[REGION:note]
-Make sure that Microsoft Visual Studio 2015 is installed prior to building Unreal Engine from source.
-[/REGION]
+
+Make sure that Microsoft Visual Studio 2015 is installed prior to building Unreal Engine from source. Also, depending on your system's specifications, it may take between 10 and 40 minutes to compile the engine.
 
 1.	Inside the root directory, run **GenerateProjectFiles.bat** to set-up your project files.
 
@@ -85,31 +127,11 @@ Make sure that Microsoft Visual Studio 2015 is installed prior to building Unrea
 
 	![](BuildUE4.png)
 
-	[REGION:note]
-	Depending on your system's specifications, it may take between 10 and 40 minutes to compile the engine from source.
-	[/REGION]
-
-1.	After compiling the engine, set your startup project to **UE4**.
-
-	![](SetUE4_StartPrj.png)
-
-1.	**Right-click** the **UE4** project, **hover over 'Debug'** and **click 'Start New Instance'** to launch the editor.
-
-	![](RunCompiledWindowsEditor.png)
-
-	[REGION:tip]
-	Alternatively, you can **press the F5 key** on your keyboard to start a new instance of the editor.
-	[/REGION]
-
-1.	Congratulations! You've compiled and launched the engine from source.
-
-	![](CompiledEditor_Windows.png)
-
 [/EXCERPT:windowscompile]
+
 [EXCERPT:maccompile]
-[REGION:note]
-Make sure that the latest version of XCode is installed prior to building Unreal Engine from Source.
-[/REGION]
+
+Make sure that the latest version of XCode is installed prior to building Unreal Engine from Source. Also, depending on your system's specifications, it may take between 10 and 40 minutes to compile the engine.
 		
 1.	Inside the root directory, run **GenerateProjectFiles.command** to set-up your project files.
 		
@@ -122,29 +144,18 @@ Make sure that the latest version of XCode is installed prior to building Unreal
 1.	To build the project, select **Product > Build**.
 
 	![](BuildingMacEditor.png)
-
-	[REGION:note]
-	Depending on your system's specifications, it may take between 15 and 40 minutes to compile the engine from source.
-	[/REGION]
-		
-1.	After compiling the engine, select **Product > Run** to launch the editor.
-		
-	![](RunningMacEditor.png)
-
-1.	Congratulations! You've compiled and launched the engine from source.
-
-	![](CompiledEditor_Mac.png)
+	
 [/EXCERPT:maccompile]
+
 [EXCERPT:linuxcompile]
-[REGION:warning]
-Our developers and support teams currently use the latest version of Ubuntu; as a result, we may not be able to provide support for other Linux distributions (including other versions of Ubuntu).
-[/REGION]
+
+Our development and support teams currently use the latest version of Ubuntu; as a result, we may not be able to provide support for other Linux distributions (including other versions of Ubuntu). Also, depending on your system's specifications, it may take anywhere from ten minutes to over an hour to compile the engine.
 		
-1.  Inside the root directory, run **Setup.sh** from the terminal to setup the files to prepare for generating the project files.
+1.  Inside the root directory, run **Setup.sh** from the terminal to setup the files needed to generate the project files.
 
 	![](RunSetupShellScript_Linux.png)
 
-1.	Inside the root directory, run **GenerateProjectFiles.sh** from the terminal to generate your project files. 
+1.	Now, run **GenerateProjectFiles.sh** from the terminal to generate your project files. 
 
 	![](RunGenPrjFilesShellScript_Linux.png)
 
@@ -152,24 +163,56 @@ Our developers and support teams currently use the latest version of Ubuntu; as
 
 	![](RunMakeFile_Linux.png)
 
+[/EXCERPT:linuxcompile]
+-->
+
+<!--
+[EXCERPT:windowsrun]
+
+1.	Set your startup project to **UE4**.
+
+	![](SetUE4_StartPrj.png)
+
+1.	**Right-click** the **UE4** project, **hover over 'Debug'** and **click 'Start New Instance'** to launch the editor.
+
+	![](RunCompiledWindowsEditor.png)
+
 	[REGION:tip]
-	You can also build the engine with QtCreator, CodeLite, KDevelop, and CMake.
+	Alternatively, you can **press the F5 key** on your keyboard to start a new instance of the editor.
 	[/REGION]
 
-1.	After compiling the engine, navigate to the editor's binary path by entering ** cd Engine/Binaries/Linux/ ** into the terminal.
+1.	Congratulations! You've compiled and launched the engine from source.
+
+	![](CompiledEditor_Windows.png)
+
+[/EXCERPT:windowsrun]
+
+[EXCERPT:macrun]
+		
+1.	Select **Product > Run** to launch the editor.
+		
+	![](RunningMacEditor.png)
+
+1.	Congratulations! You've compiled and launched the engine from source.
+
+	![](CompiledEditor_Mac.png)
+[/EXCERPT:macrun]
+
+[EXCERPT:linuxrun]
+		
+If you're unable to create an OpenGL context as you're working through these steps, you may need to update your graphics drivers (refer to Ubuntu's documentation<!-- and our list of [](Platforms/Linux/BeginnerLinuxDeveloper/SupportedGfxDrivers) -->). 
+For an optimal experience, we recommend using NVidia binary drivers (versions equal to or greater than 343.xx).
+
+1.	Navigate to the editor's binary path by entering `cd Engine/Binaries/Linux/` into the terminal.
 
 1.	Run **UE4Editor** to launch the editor.
 
 	![](RunUE4Editor_Linux.png)
 
-	[REGION:warning]
-	If you're unable to create an OpenGL context, you may need to update your graphics drivers (refer to Ubuntu's documentation). We recommend using NVidia binary drivers (versions equal to or greater than 343.xx).
-  	[/REGION]
-
 1.	Congratulations! You've compiled and launched the engine from source.
 
 	![](CompiledEditor_Linux.png)
-[/EXCERPT:linuxcompile]
+[/EXCERPT:linuxrun]
 -->
 
 ## Getting Started with Unreal Engine
@@ -187,173 +230,3 @@ UE4's in-editor help features are a great way to get your questions answered.
 [/REGION]
 
 [/EXCERPT:Main]
-
-<!--Currently deprecated for the 'Building from Source' version
-
-## Compiling Unreal Engine
-
-Unreal Engine 4 (UE4) uses a custom building method via the UnrealBuildTool. This tool handles all the complex aspects of compiling the engine. 
-
-### Build Configuration
-
-[OBJECT:ToggleButtonContent]
-	[PARAMLITERAL:category]
-	OS
-	[/PARAMLITERAL]
-	[PARAMLITERAL:id]
-	windows
-	[/PARAMLITERAL]
-	[PARAMLITERAL:active]
-	active_button_content
-	[/PARAMLITERAL]
-	[PARAM:content]
-		Not all configurations will appear by default. Edit the [GenerateProjectFiles.bat](Programming/UnrealBuildSystem/ProjectFileGenerator) file in the main UE4 directory to generate the project files for additional configurations as needed.
-
-		UE4 and the UnrealBuildTool use different build configurations to determine how the engine is compiled. Which configuration you use will be determined by the purposes of the build you want to create.
-
-		The available configurations:
-
-		| Build Configuration | Description |
-		|-------------|-------------|
-		| **Debug** | This configuration builds both engine and game code in debug configuration.|
-		| **DebugGame** | This configuration builds the engine as optimized, but leaves the game code debuggable. This configuration is ideal for debugging only game modules. |
-		| **Development** | This configuration is equivalent to Release.  Both engine and game code will be built in this configuration. |
-		| **Shipping** | This is the configuration for optimal performance and shipping your game.  This configuration strips out console commands, stats, and profiling tools. |
-		| **Test** | This configuration is the **Shipping** configuration, but with some console commands, stats, and profiling tools enabled. |
-
-		The build configuration can be set in the Visual Studio toolbar.
-
-		![UE4 Solution - Configuration Menu](configuration_menu.png)
-	[/PARAM]
-[/OBJECT]
-[OBJECT:ToggleButtonContent]
-	[PARAMLITERAL:category]
-	OS
-	[/PARAMLITERAL]
-	[PARAMLITERAL:id]
-	mac
-	[/PARAMLITERAL]
-	[PARAMLITERAL:active]
-	[/PARAMLITERAL]
-	[PARAM:content]
-		The available configurations:
-
-		| Build Configuration | Description |
-		| ------------------- | ----------- |
-		| **Product > Build For > Running** | This configuration builds both engine and game code in debug configuration.|
-		| **Product > Build For > Profiling** | This configuration is equivalent to Release.  Both engine and game code will be built in this configuration. |
-		
-	[/PARAM]
-[/OBJECT]
-
-### Target Platforms
-
-UE4 currently supports the following platforms:
-
-[OBJECT:ToggleButtonContent]
-	[PARAMLITERAL:category]
-	OS
-	[/PARAMLITERAL]
-	[PARAMLITERAL:id]
-	windows
-	[/PARAMLITERAL]
-	[PARAMLITERAL:active]
-	active_button_content
-	[/PARAMLITERAL]
-	[PARAM:content]
-		* **Windows 32-bit**
-		* **Windows 64-bit**
-
-
-		When building the engine itself, you will use the **Win64** platform.
-
-		![UE4 Solution - Platform Menu](platform_menu.png)
-	[/PARAM]
-[/OBJECT]
-[OBJECT:ToggleButtonContent]
-	[PARAMLITERAL:category]
-	OS
-	[/PARAMLITERAL]
-	[PARAMLITERAL:id]
-	mac
-	[/PARAMLITERAL]
-	[PARAMLITERAL:active]
-	[/PARAMLITERAL]
-	[PARAM:content]
-		* **Mac 64-bit**
-		* **iOS Device**
-
-
-		When building the engine itself, you will use the **Mac 64-bit** platform.
-	[/PARAM]
-[/OBJECT]
-
-(#BinaryNamingConvention)
-### Binary Naming Conventions
-
-The name of the binaries generated will vary depending on the configuration used when building. The general naming convention is:
-
-|Engine Mode | Naming Convention |
-|----------- | ------------------|
-| Editor	  | UE4Editor-Platform-Configuration.extension |
-| Uncooked Game | UE4-Platform-Configuration.extension |
-
-The exception to the rule is that for PC platforms, the platform and configuration are omitted when building the Development configuration. For example, the names of the editor executable files generated for the Win64 platform include:
-
-| Configuration | EXE Name |
-| ------------- | -------- |
-| Debug | UE4Editor-Win64-Debug.exe |
-| DebugGame | UE4Editor-Win64-DebugGame.exe |
-| Development | UE4Editor.exe |
-
-
-### Building the Engine
-
-[OBJECT:ToggleButtonContent]
-	[PARAMLITERAL:category]
-	OS
-	[/PARAMLITERAL]
-	[PARAMLITERAL:id]
-	windows
-	[/PARAMLITERAL]
-	[PARAMLITERAL:active]
-	active_button_content
-	[/PARAMLITERAL]
-	[PARAM:content]
-		1.	Set the **Build Configuration** to the configuration you want to build. In this example, it is set to **Development**. Refer to the [Build Configuration](#BuildConfiguration) section for descriptions of each available configuration.
-
-			![UE4 Solution - Set Configuration to Development](config_development_menu.png)
-
-		1.	**Right-click** the **UE4** project: 
-
-			1. Choose **Clean** to remove any compiled or temporary files. This will ensure the next build is a full rebuild.
-
-			1.	Choose **Rebuild** to recompile the engine.
-
-		Alternatively, if the **UE4** project is still set as your startup project:
-
-		1. Set your **Build Configuration**.
-
-		1. From the **Build** menu:
-
-			1. **Clean UE4**.
-	
-			1. **Rebuild UE4**.
-	[/PARAM]
-[/OBJECT]
-[OBJECT:ToggleButtonContent]
-	[PARAMLITERAL:category]
-	OS
-	[/PARAMLITERAL]
-	[PARAMLITERAL:id]
-	mac
-	[/PARAMLITERAL]
-	[PARAMLITERAL:active]
-	[/PARAMLITERAL]
-	[PARAM:content]
-		1. Make sure the active scheme is set to **UE4Editor - Mac > My Mac 64-bit**.
-
-		1. From the **product** menu, choose **Build For > Running**.
-	[/PARAM]
-[/OBJECT]
--->
\ No newline at end of file
