Implemented first draft of the 'Building UE4 on Linux' step in the Linux Workflow How To series.

#UE4docs

diff --git a/Engine/Documentation/Source/Platforms/Linux/BeginnerLinuxDeveloper/SettingUpAnUnrealWorkflow/3/SettingUpAnUnrealWorkflow_3.INT.udn b/Engine/Documentation/Source/Platforms/Linux/BeginnerLinuxDeveloper/SettingUpAnUnrealWorkflow/3/SettingUpAnUnrealWorkflow_3.INT.udn
--- a/Engine/Documentation/Source/Platforms/Linux/BeginnerLinuxDeveloper/SettingUpAnUnrealWorkflow/3/SettingUpAnUnrealWorkflow_3.INT.udn
+++ b/Engine/Documentation/Source/Platforms/Linux/BeginnerLinuxDeveloper/SettingUpAnUnrealWorkflow/3/SettingUpAnUnrealWorkflow_3.INT.udn
@@ -25,11 +25,20 @@ SkillLevel:Beginner
 [/OBJECT]
 [/VAR]
 
-In preparation for this critical step, you downloaded a copy of Unreal Engine's source code on your hard disk. 
-During this step, you'll get to run several utilities from the terminal to build a binary of UE4 on your Linux machine.
+In preparation for this critical step, you downloaded a copy of Unreal Engine's source code onto your hard disk. 
+During this part of the tutorial, you'll get to run several utilities from the terminal to build a binary of UE4 
+on your Linux machine. 
 
 %Steps%
 
+[INCLUDE:Programming/Development/BuildingUnrealEngine#linuxcompile]
 
+## End Result
+
+![](built_engine.png)
+
+_Some lower-end Linux systems can take over an hour to build the engine._
+
+At this point, you should have a fully built binary of UE4 on your hard disk. Alas, the time has come to run the engine on your Linux system!
 
 %Steps%
\ No newline at end of file
