Minor revision to the Goal statement.

#UE4docs

diff --git a/Engine/Documentation/Source/Platforms/Linux/BeginnerLinuxDeveloper/SettingUpAnUnrealWorkflow/SettingUpAnUnrealWorkflow.INT.udn b/Engine/Documentation/Source/Platforms/Linux/BeginnerLinuxDeveloper/SettingUpAnUnrealWorkflow/SettingUpAnUnrealWorkflow.INT.udn
--- a/Engine/Documentation/Source/Platforms/Linux/BeginnerLinuxDeveloper/SettingUpAnUnrealWorkflow/SettingUpAnUnrealWorkflow.INT.udn
+++ b/Engine/Documentation/Source/Platforms/Linux/BeginnerLinuxDeveloper/SettingUpAnUnrealWorkflow/SettingUpAnUnrealWorkflow.INT.udn
@@ -9,18 +9,20 @@ parent:Platforms/Linux/BeginnerLinuxDeveloper
 order: 2
 tags:Linux
 SkillLevel:Beginner
+prereq:GettingStarted/RecommendedSpecifications
 prereq:Platforms/Linux/BeginnerLinuxDeveloper/SupportedGfxDrivers
 
 ![](CompiledEditor_Linux.png)
 
-_At the end of this guide, you'll have Unreal Engine 4 (UE4) running on your Linux machine._
+_At the end of this tutorial, you'll have Unreal Engine 4 (UE4) running on your Linux machine._
 
 ## Goal
-The focus of this guide is to demonstrate the workflow needed to get UE4 running on your Linux machine.
+
+The focus of this guide is to show you how to setup the primary workflow needed to get UE4 running on your Linux machine.
 
 ## Objectives
 
-After working through this guide, you'll know the following:
+After working through this tutorial, you'll know the following:
 
 * How to setup Git to download the source code for Unreal Engine.
 * How to fork and clone Unreal Engine's source code.
