Implemented the first draft of the final step in the Linux Workflow Setup How To series.

#UE4docs

diff --git a/Engine/Documentation/Source/Platforms/Linux/BeginnerLinuxDeveloper/SettingUpAnUnrealWorkflow/4/SettingUpAnUnrealWorkflow_4.INT.udn b/Engine/Documentation/Source/Platforms/Linux/BeginnerLinuxDeveloper/SettingUpAnUnrealWorkflow/4/SettingUpAnUnrealWorkflow_4.INT.udn
--- a/Engine/Documentation/Source/Platforms/Linux/BeginnerLinuxDeveloper/SettingUpAnUnrealWorkflow/4/SettingUpAnUnrealWorkflow_4.INT.udn
+++ b/Engine/Documentation/Source/Platforms/Linux/BeginnerLinuxDeveloper/SettingUpAnUnrealWorkflow/4/SettingUpAnUnrealWorkflow_4.INT.udn
@@ -17,15 +17,34 @@ SkillLevel:Beginner
 	[PARAM:current]
 	[/PARAM]
 	[PARAM:home]
-		[Home](Platforms/Linux/BeginnerLinuxDeveloper/SettingUpAnUnrealWorkflow)
+		[Home](Platforms/Linux/BeginnerLinuxDeveloper)
 	[/PARAM]
 	[PARAM:next]
 	[/PARAM]
 [/OBJECT]
 [/VAR]
 
+As you worked your way to this final step, you linked your Epic Games account with your GitHub username, you downloaded the Unreal Engine source code from our 
+GitHub site, and you compiled a binary of UE4 on your Linux system. The only remaining task is for you to run the UE4 Editor so that you can start making 
+games for the Linux platform.
+
 %Steps%
 
+[INCLUDE:Programming/Development/BuildingUnrealEngine#linuxrun]
+
+## End Result
+
+![](howTo_endResult.png)
+
+If you're creating your first game on Linux, you just learned how to establish your primary UE4 workflow on Linux. As you worked your way through this tutorial, 
+you learned:
+
+&#x2713; How to setup Git to download the source code for Unreal Engine.  
+&#x2713; How to fork and clone Unreal Engine's source code.  
+&#x2713; How to build Unreal Engine from source.  
+&#x2713; How to run Unreal Engine on your Linux machine.  
 
+Throughout the course of this tutorial, you spent most of your time using the terminal. In the next tutorial, you'll learn how to setup an integrated development 
+environment (IDE), enabling you to harness the context rich features of an IDE to help you build and debug UE4 games.
 
 %Steps%
\ No newline at end of file
