Minor update.

#UE4docs.

diff --git a/Engine/Documentation/Source/Platforms/Linux/BeginnerLinuxDeveloper/SettingUpAnUnrealWorkflow/1/SettingUpAnUnrealWorkflow_1.INT.udn b/Engine/Documentation/Source/Platforms/Linux/BeginnerLinuxDeveloper/SettingUpAnUnrealWorkflow/1/SettingUpAnUnrealWorkflow_1.INT.udn
--- a/Engine/Documentation/Source/Platforms/Linux/BeginnerLinuxDeveloper/SettingUpAnUnrealWorkflow/1/SettingUpAnUnrealWorkflow_1.INT.udn
+++ b/Engine/Documentation/Source/Platforms/Linux/BeginnerLinuxDeveloper/SettingUpAnUnrealWorkflow/1/SettingUpAnUnrealWorkflow_1.INT.udn
@@ -41,7 +41,7 @@ connect with other UE4 Linux developers, and where you'll be able to keep up wit
 
 ## Setup Git
 
-Now that you've registered for an Epic Games account, you're ready to work through the following steps to gain access to the UE4 C++ source code on GitHub.
+Now that you've registered for an Epic Games account, you're ready to work through the following steps to gain access to the UE4 source code on GitHub.
 
 1.	Sign up for a GitHub account at [GitHub.com](https://github.com/EpicGames).
 
@@ -56,6 +56,6 @@ Now that you've registered for an Epic Games account, you're ready to work throu
 ## End Result
 
 By now, you've registered for an Epic Games account and linked it with your GitHub username. Also, you've installed Git and you're ready for the next step,
-where you'll fork and clone the UE4 C++ source code from the Unreal Engine GitHub repository.
+where you'll fork and clone the UE4 source code from the Unreal Engine GitHub repository.
 
 %Steps%
\ No newline at end of file
