Added notational preambles to the linux run and linux compile write-ups in the 'Building Unreal Engine' documentation page.
Fixed the formats in the How To pages to conform with the latest Style & Standards Guide.

#UE4docs

diff --git a/Engine/Documentation/Source/Programming/Development/BuildingUnrealEngine/BuildingUnrealEngine.INT.udn b/Engine/Documentation/Source/Programming/Development/BuildingUnrealEngine/BuildingUnrealEngine.INT.udn
--- a/Engine/Documentation/Source/Programming/Development/BuildingUnrealEngine/BuildingUnrealEngine.INT.udn
+++ b/Engine/Documentation/Source/Programming/Development/BuildingUnrealEngine/BuildingUnrealEngine.INT.udn
@@ -149,8 +149,10 @@ Make sure that the latest version of XCode is installed prior to building Unreal
 
 [EXCERPT:linuxcompile]
 
+[REGION:note]
 Our development and support teams currently use the latest version of Ubuntu; as a result, we may not be able to provide support for other Linux distributions (including other versions of Ubuntu). Also, depending on your system's specifications, it may take anywhere from ten minutes to over an hour to compile the engine.
-		
+[/REGION]
+
 1.  Inside the root directory, run **Setup.sh** from the terminal to setup the files needed to generate the project files.
 
 	![](RunSetupShellScript_Linux.png)
@@ -199,9 +201,11 @@ Our development and support teams currently use the latest version of Ubuntu; as
 [/EXCERPT:macrun]
 
 [EXCERPT:linuxrun]
-		
+
+[REGION:note]		
 If you're unable to create an OpenGL context as you're working through these steps, you may need to update your graphics drivers (refer to Ubuntu's documentation<!-- and our list of [](Platforms/Linux/BeginnerLinuxDeveloper/SupportedGfxDrivers) -->). 
 For an optimal experience, we recommend using NVidia binary drivers (versions equal to or greater than 343.xx).
+[/REGION]
 
 1.	Navigate to the editor's binary path by entering `cd Engine/Binaries/Linux/` into the terminal.
 
